MeshModelingToolset: Split plugin into non-experimental and experimental

diff --git a/Engine/Plugins/Editor/ModelingToolsEditorMode/ModelingToolsEditorMode.uplugin b/Engine/Plugins/Editor/ModelingToolsEditorMode/ModelingToolsEditorMode.uplugin
--- a/Engine/Plugins/Editor/ModelingToolsEditorMode/ModelingToolsEditorMode.uplugin
+++ b/Engine/Plugins/Editor/ModelingToolsEditorMode/ModelingToolsEditorMode.uplugin
@@ -26,6 +26,10 @@
 			"Name": "MeshModelingToolset",
 			"Enabled": true
 		},
+		{
+			"Name": "MeshModelingToolsetExp",
+			"Enabled": true
+		},
 		{
 			"Name": "MeshLODToolset",
 			"Enabled": true
diff --git a/Engine/Plugins/Editor/ModelingToolsEditorMode/Source/ModelingToolsEditorMode/ModelingToolsEditorMode.Build.cs b/Engine/Plugins/Editor/ModelingToolsEditorMode/Source/ModelingToolsEditorMode/ModelingToolsEditorMode.Build.cs
--- a/Engine/Plugins/Editor/ModelingToolsEditorMode/Source/ModelingToolsEditorMode/ModelingToolsEditorMode.Build.cs
+++ b/Engine/Plugins/Editor/ModelingToolsEditorMode/Source/ModelingToolsEditorMode/ModelingToolsEditorMode.Build.cs
@@ -57,6 +57,7 @@ public class ModelingToolsEditorMode : ModuleRules
 				"ModelingComponents",
 				"ModelingComponentsEditorOnly",
 				"MeshModelingTools",
+				"MeshModelingToolsExp",
 				"MeshModelingToolsEditorOnly",
 				"MeshLODToolset",
 
diff --git a/Engine/Plugins/Experimental/ChaosEditor/ChaosEditor.uplugin b/Engine/Plugins/Experimental/ChaosEditor/ChaosEditor.uplugin
--- a/Engine/Plugins/Experimental/ChaosEditor/ChaosEditor.uplugin
+++ b/Engine/Plugins/Experimental/ChaosEditor/ChaosEditor.uplugin
@@ -31,7 +31,7 @@
 			"Enabled": true
 		},
 		{
-			"Name": "MeshModelingToolset",
+			"Name": "MeshModelingToolsetExp",
 			"Enabled": true
 		},
 		{
diff --git a/Engine/Plugins/Experimental/Enterprise/DataprepGeometryOperations/DataprepGeometryOperations.uplugin b/Engine/Plugins/Experimental/Enterprise/DataprepGeometryOperations/DataprepGeometryOperations.uplugin
--- a/Engine/Plugins/Experimental/Enterprise/DataprepGeometryOperations/DataprepGeometryOperations.uplugin
+++ b/Engine/Plugins/Experimental/Enterprise/DataprepGeometryOperations/DataprepGeometryOperations.uplugin
@@ -32,7 +32,7 @@
 			"Enabled": true
 		},
 		{
-			"Name": "MeshModelingToolset",
+			"Name": "MeshModelingToolsetExp",
 			"Enabled": true
 		},
 		{
diff --git a/Engine/Plugins/Experimental/Enterprise/DataprepGeometryOperations/Source/DataprepGeometryOperations.Build.cs b/Engine/Plugins/Experimental/Enterprise/DataprepGeometryOperations/Source/DataprepGeometryOperations.Build.cs
--- a/Engine/Plugins/Experimental/Enterprise/DataprepGeometryOperations/Source/DataprepGeometryOperations.Build.cs
+++ b/Engine/Plugins/Experimental/Enterprise/DataprepGeometryOperations/Source/DataprepGeometryOperations.Build.cs
@@ -20,7 +20,7 @@ namespace UnrealBuildTool.Rules
 					"MeshConversion",
 					"MeshDescription",
 					"MeshMergeUtilities",
-					"MeshModelingTools",
+					"MeshModelingToolsExp",
 					"MeshUtilitiesCommon",
 					"ModelingOperators",
 					"ModelingOperatorsEditorOnly",
diff --git a/Engine/Plugins/Experimental/MeshLODToolset/Source/MeshLODToolset/MeshLODToolset.Build.cs b/Engine/Plugins/Experimental/MeshLODToolset/Source/MeshLODToolset/MeshLODToolset.Build.cs
--- a/Engine/Plugins/Experimental/MeshLODToolset/Source/MeshLODToolset/MeshLODToolset.Build.cs
+++ b/Engine/Plugins/Experimental/MeshLODToolset/Source/MeshLODToolset/MeshLODToolset.Build.cs
@@ -36,7 +36,7 @@ public class MeshLODToolset : ModuleRules
                 "StaticMeshDescription",
 				"ModelingComponents",
 				"ModelingOperators",
-				"MeshModelingTools",
+				"MeshModelingToolsExp",
 				"GeometryFlowCore",
 				"GeometryFlowMeshProcessing",
 				"GeometryFlowMeshProcessingEditor"
diff --git a/Engine/Plugins/Experimental/MeshModelingToolsetExp/Source/MeshModelingToolsEditorOnly/MeshModelingToolsEditorOnly.Build.cs b/Engine/Plugins/Experimental/MeshModelingToolsetExp/Source/MeshModelingToolsEditorOnly/MeshModelingToolsEditorOnly.Build.cs
new file mode 100644
--- /dev/null
+++ b/Engine/Plugins/Experimental/MeshModelingToolsetExp/Source/MeshModelingToolsEditorOnly/MeshModelingToolsEditorOnly.Build.cs
@@ -0,0 +1,87 @@
+// Copyright Epic Games, Inc. All Rights Reserved.
+
+using UnrealBuildTool;
+
+public class MeshModelingToolsEditorOnly : ModuleRules
+{
+	public MeshModelingToolsEditorOnly(ReadOnlyTargetRules Target) : base(Target)
+	{
+		PCHUsage = ModuleRules.PCHUsageMode.UseExplicitOrSharedPCHs;
+		
+		PublicIncludePaths.AddRange(
+			new string[] {
+				// ... add public include paths required here ...
+			}
+			);
+				
+		
+		PrivateIncludePaths.AddRange(
+			new string[] {
+				// ... add other private include paths required here ...
+			}
+			);
+			
+		
+		PublicDependencyModuleNames.AddRange(
+			new string[]
+			{
+				"Core",
+                "InteractiveToolsFramework",
+                "MeshDescription",
+				"StaticMeshDescription",
+				"SkeletalMeshDescription",
+				"GeometryCore",
+				"GeometryFramework",
+				"DynamicMesh",
+				"MeshConversion",
+				"MeshModelingToolsExp",
+				"ModelingComponents",
+				"ModelingComponentsEditorOnly",
+				"ModelingOperators",
+				"ModelingOperatorsEditorOnly",
+				// ... add other public dependencies that you statically link with here ...
+
+				"HairStrandsCore"		// required for Hair toolset
+			}
+            );
+			
+		
+		PrivateDependencyModuleNames.AddRange(
+			new string[]
+			{
+				"CoreUObject",
+				"Engine",
+                "RenderCore",
+                "RHI",
+				"InputCore",
+
+				"MeshUtilities",		// tangents calculation
+				"EditorFramework",
+				"UnrealEd",
+				"MeshBuilder",
+				"MeshUtilitiesCommon",  
+				"MeshReductionInterface", // for UE standard simplification 
+				//"Slate",
+				//"SlateCore",
+				// ... add private dependencies that you statically link with here ...	
+			}
+			);
+
+		if (Target.Platform == UnrealTargetPlatform.Win64)
+		{
+			PrivateDependencyModuleNames.AddRange(
+				new string[]
+				{
+					"ProxyLODMeshReduction", // currently Win64-only
+				}
+				);
+		}
+
+		DynamicallyLoadedModuleNames.AddRange(
+			new string[]
+			{
+				// ... add any modules that your module loads dynamically here ...
+			}
+			);
+	}
+}
diff --git a/Engine/Plugins/Experimental/MeshModelingToolsetExp/Source/MeshModelingToolsEditorOnly/Public/MeshTangentsTool.h b/Engine/Plugins/Experimental/MeshModelingToolsetExp/Source/MeshModelingToolsEditorOnly/Public/MeshTangentsTool.h
new file mode 100644
--- /dev/null
+++ b/Engine/Plugins/Experimental/MeshModelingToolsetExp/Source/MeshModelingToolsEditorOnly/Public/MeshTangentsTool.h
@@ -0,0 +1,142 @@
+// Copyright Epic Games, Inc. All Rights Reserved.
+
+#pragma once
+
+#include "CoreMinimal.h"
+#include "SingleSelectionTool.h"
+#include "InteractiveToolBuilder.h"
+#include "PreviewMesh.h"
+#include "Drawing/PreviewGeometryActor.h"
+#include "MeshOpPreviewHelpers.h"
+#include "DynamicMesh/DynamicMesh3.h"
+#include "DynamicMesh/MeshTangents.h"
+#include "ParameterizationOps/CalculateTangentsOp.h"
+#include "BaseTools/SingleSelectionMeshEditingTool.h"
+#include "MeshTangentsTool.generated.h"
+
+
+// predeclarations
+struct FMeshDescription;
+class UDynamicMeshComponent;
+class UMaterialInterface;
+class UMaterialInstanceDynamic;
+
+
+/**
+ *
+ */
+UCLASS()
+class MESHMODELINGTOOLSEDITORONLY_API UMeshTangentsToolBuilder : public USingleSelectionMeshEditingToolBuilder
+{
+	GENERATED_BODY()
+
+public:
+	virtual USingleSelectionMeshEditingTool* CreateNewTool(const FToolBuilderState& SceneState) const override;
+
+protected:
+	virtual const FToolTargetTypeRequirements& GetTargetRequirements() const override;
+};
+
+
+
+UCLASS()
+class MESHMODELINGTOOLSEDITORONLY_API UMeshTangentsToolProperties : public UInteractiveToolPropertySet
+{
+	GENERATED_BODY()
+public:
+
+	UPROPERTY(EditAnywhere, Category = Options)
+	EMeshTangentsType TangentType = EMeshTangentsType::FastMikkTSpace;
+
+
+	UPROPERTY(EditAnywhere, Category = Visualization)
+	bool bShowTangents = true;
+
+	UPROPERTY(EditAnywhere, Category = Visualization)
+	bool bShowNormals = false;
+
+	UPROPERTY(EditAnywhere, Category = Visualization, AdvancedDisplay)
+	bool bHideDegenerates = true;
+
+	UPROPERTY(EditAnywhere, Category = Visualization, AdvancedDisplay, meta = (UIMin = "0.01", UIMax = "25.0", ClampMin = "0.01", ClampMax = "10000000.0"))
+	float LineLength = 2.0;
+
+	UPROPERTY(EditAnywhere, Category = Visualization, AdvancedDisplay, meta = (UIMin = "0", UIMax = "25.0", ClampMin = "0", ClampMax = "1000.0"))
+	float LineThickness = 3.0;
+
+
+	UPROPERTY(EditAnywhere, Category = Visualization, AdvancedDisplay)
+	bool bCompareWithMikkt = false;
+
+	UPROPERTY(EditAnywhere, Category = Visualization, AdvancedDisplay, meta = (UIMin = "0.5", UIMax = "90.0"))
+	float AngleThreshDeg = 5.0;
+
+};
+
+
+
+
+
+/**
+ * Simple Mesh Simplifying Tool
+ */
+UCLASS()
+class MESHMODELINGTOOLSEDITORONLY_API UMeshTangentsTool : public USingleSelectionMeshEditingTool, public UE::Geometry::IGenericDataOperatorFactory<UE::Geometry::FMeshTangentsd>
+{
+	GENERATED_BODY()
+public:
+	UMeshTangentsTool();
+
+	virtual void Setup() override;
+	virtual void Shutdown(EToolShutdownType ShutdownType) override;
+
+	virtual void OnTick(float DeltaTime) override;
+	virtual void Render(IToolsContextRenderAPI* RenderAPI) override;
+
+	virtual bool HasCancel() const override { return true; }
+	virtual bool HasAccept() const override { return true; }
+	virtual bool CanAccept() const override;
+
+	// IGenericDataOperatorFactory API
+	virtual TUniquePtr<UE::Geometry::TGenericDataOperator<UE::Geometry::FMeshTangentsd>> MakeNewOperator() override;
+
+protected:
+	UPROPERTY()
+	TObjectPtr<UMeshTangentsToolProperties> Settings = nullptr;
+
+protected:
+	UPROPERTY()
+	TObjectPtr<UMaterialInterface> DefaultMaterial = nullptr;
+
+	UPROPERTY()
+	TObjectPtr<UPreviewMesh> PreviewMesh = nullptr;
+
+	UPROPERTY()
+	TObjectPtr<UPreviewGeometry> PreviewGeometry = nullptr;
+
+	TUniquePtr<TGenericDataBackgroundCompute<UE::Geometry::FMeshTangentsd>> Compute = nullptr;
+
+protected:
+	TSharedPtr<FMeshDescription, ESPMode::ThreadSafe> InputMeshDescription;
+	TSharedPtr<UE::Geometry::FMeshTangentsf, ESPMode::ThreadSafe> InitialTangents;
+	TSharedPtr<UE::Geometry::FDynamicMesh3, ESPMode::ThreadSafe> InputMesh;
+
+	bool bThicknessDirty = false;
+	bool bLengthDirty = false;
+	bool bVisibilityChanged = false;
+
+	void OnTangentsUpdated(const TUniquePtr<UE::Geometry::FMeshTangentsd>& NewResult);
+	void UpdateVisualization(bool bThicknessChanged, bool bLengthChanged);
+
+	struct FMikktDeviation
+	{
+		float MaxAngleDeg;
+		int32 TriangleID;
+		int32 TriVertIndex;
+		FVector3f VertexPos;
+		FVector3f MikktTangent, MikktBitangent;
+		FVector3f OtherTangent, OtherBitangent;
+	};
+	TArray<FMikktDeviation> Deviations;
+
+};
diff --git a/Engine/Plugins/Experimental/MeshModelingToolsetExp/Source/MeshModelingToolsEditorOnly/Public/UVLayoutTool.h b/Engine/Plugins/Experimental/MeshModelingToolsetExp/Source/MeshModelingToolsEditorOnly/Public/UVLayoutTool.h
new file mode 100644
--- /dev/null
+++ b/Engine/Plugins/Experimental/MeshModelingToolsetExp/Source/MeshModelingToolsEditorOnly/Public/UVLayoutTool.h
@@ -0,0 +1,176 @@
+// Copyright Epic Games, Inc. All Rights Reserved.
+
+#pragma once
+
+#include "CoreMinimal.h"
+#include "UObject/NoExportTypes.h"
+#include "MultiSelectionTool.h"
+#include "InteractiveToolBuilder.h"
+#include "MeshOpPreviewHelpers.h"
+#include "ToolDataVisualizer.h"
+#include "DynamicMesh/DynamicMesh3.h"
+#include "BaseTools/SingleClickTool.h"
+#include "Properties/MeshMaterialProperties.h"
+#include "Properties/MeshUVChannelProperties.h"
+#include "Drawing/UVLayoutPreview.h"
+
+#include "UVLayoutTool.generated.h"
+
+
+// predeclarations
+struct FMeshDescription;
+class UDynamicMeshComponent;
+
+/**
+ *
+ */
+UCLASS()
+class MESHMODELINGTOOLSEDITORONLY_API UUVLayoutToolBuilder : public UInteractiveToolBuilder
+{
+	GENERATED_BODY()
+
+public:
+	virtual bool CanBuildTool(const FToolBuilderState& SceneState) const override;
+	virtual UInteractiveTool* BuildTool(const FToolBuilderState& SceneState) const override;
+
+protected:
+	virtual const FToolTargetTypeRequirements& GetTargetRequirements() const override;
+};
+
+
+
+UENUM()
+enum class EUVLayoutType
+{
+	Transform,
+	Stack,
+	Repack
+};
+
+
+
+/**
+ * Standard properties
+ */
+UCLASS()
+class MESHMODELINGTOOLSEDITORONLY_API UUVLayoutToolProperties : public UInteractiveToolPropertySet
+{
+	GENERATED_BODY()
+
+public:
+	UUVLayoutToolProperties();
+
+	/** Type of transformation to apply to input UV islands */
+	UPROPERTY(EditAnywhere, Category = UVLayout)
+	EUVLayoutType LayoutType = EUVLayoutType::Repack;
+
+	/** Expected resolution of output textures; controls spacing left between charts */
+	UPROPERTY(EditAnywhere, Category = UVLayout, meta = (UIMin = "64", UIMax = "2048", ClampMin = "2", ClampMax = "4096"))
+	int TextureResolution = 1024;
+
+	/** Apply this uniform scaling to the UVs after any layout recalculation */
+	UPROPERTY(EditAnywhere, Category = UVLayout, meta = (UIMin = "0.1", UIMax = "5.0", ClampMin = "0.0001", ClampMax = "10000") )
+	float UVScaleFactor = 1;
+
+	/** Apply this 2D translation to the UVs after any layout recalculation, and after scaling */
+	UPROPERTY(EditAnywhere, Category = UVLayout)
+	FVector2D UVTranslate = FVector2D(0,0);
+
+	/** Allow the packer to flip the orientation of UV islands if it save space. May cause problems for downstream operations, not recommended. */
+	UPROPERTY(EditAnywhere, Category = UVLayout, AdvancedDisplay)
+	bool bAllowFlips = false;
+
+};
+
+
+
+
+
+
+/**
+ * Factory with enough info to spawn the background-thread Operator to do a chunk of work for the tool
+ *  stores a pointer to the tool and enough info to know which specific operator it should spawn
+ */
+UCLASS()
+class MESHMODELINGTOOLSEDITORONLY_API UUVLayoutOperatorFactory : public UObject, public UE::Geometry::IDynamicMeshOperatorFactory
+{
+	GENERATED_BODY()
+
+public:
+	// IDynamicMeshOperatorFactory API
+	virtual TUniquePtr<UE::Geometry::FDynamicMeshOperator> MakeNewOperator() override;
+
+	UPROPERTY()
+	TObjectPtr<UUVLayoutTool> Tool;
+
+	int ComponentIndex;
+
+};
+
+/**
+ * Simple Mesh Normal Updating Tool
+ */
+UCLASS()
+class MESHMODELINGTOOLSEDITORONLY_API UUVLayoutTool : public UMultiSelectionTool
+{
+	GENERATED_BODY()
+
+public:
+
+	friend UUVLayoutOperatorFactory;
+
+	UUVLayoutTool();
+
+	virtual void Setup() override;
+	virtual void Shutdown(EToolShutdownType ShutdownType) override;
+
+	virtual void SetWorld(UWorld* World, UInteractiveGizmoManager* GizmoManagerIn);
+
+	virtual void OnTick(float DeltaTime) override;
+	virtual void Render(IToolsContextRenderAPI* RenderAPI) override;
+
+	virtual bool HasCancel() const override { return true; }
+	virtual bool HasAccept() const override { return true; }
+	virtual bool CanAccept() const override;
+
+	virtual void OnPropertyModified(UObject* PropertySet, FProperty* Property) override;
+
+
+public:
+	int32 GetSelectedUVChannel() const;
+
+
+protected:
+
+	UPROPERTY()
+	TObjectPtr<UMeshUVChannelProperties> UVChannelProperties = nullptr;
+
+	UPROPERTY()
+	TObjectPtr<UUVLayoutToolProperties> BasicProperties = nullptr;
+
+	UPROPERTY()
+	TObjectPtr<UExistingMeshMaterialProperties> MaterialSettings = nullptr;
+
+	UPROPERTY()
+	TArray<TObjectPtr<UMeshOpPreviewWithBackgroundCompute>> Previews;
+
+protected:
+	TArray<TSharedPtr<UE::Geometry::FDynamicMesh3, ESPMode::ThreadSafe>> OriginalDynamicMeshes;
+
+	UWorld* TargetWorld = nullptr;
+
+	FViewCameraState CameraState;
+
+	void UpdateNumPreviews();
+
+	void UpdateVisualization();
+
+	void OnPreviewMeshUpdated(UMeshOpPreviewWithBackgroundCompute* Compute);
+
+	void GenerateAsset(const TArray<FDynamicMeshOpResult>& Results);
+
+
+protected:
+	UPROPERTY()
+	TObjectPtr<UUVLayoutPreview> UVLayoutView = nullptr;
+};
diff --git a/Engine/Plugins/Experimental/MeshModelingToolsetExp/Source/MeshModelingToolsExp/MeshModelingToolsExp.Build.cs b/Engine/Plugins/Experimental/MeshModelingToolsetExp/Source/MeshModelingToolsExp/MeshModelingToolsExp.Build.cs
new file mode 100644
--- /dev/null
+++ b/Engine/Plugins/Experimental/MeshModelingToolsetExp/Source/MeshModelingToolsExp/MeshModelingToolsExp.Build.cs
@@ -0,0 +1,92 @@
+// Copyright Epic Games, Inc. All Rights Reserved.
+
+using UnrealBuildTool;
+
+public class MeshModelingToolsExp : ModuleRules
+{
+	public MeshModelingToolsExp(ReadOnlyTargetRules Target) : base(Target)
+	{
+		PCHUsage = ModuleRules.PCHUsageMode.UseExplicitOrSharedPCHs;
+		
+		PublicIncludePaths.AddRange(
+			new string[] {
+				// ... add public include paths required here ...
+			}
+			);
+				
+		PublicIncludePathModuleNames.AddRange(
+			new string[] {
+				"AnimationCore",			// For the BoneWeights.h include
+			}
+		);
+	
+		PublicDependencyModuleNames.AddRange(
+			new string[]
+			{
+				// These dependencies were commented out when we split out MeshModelingToolsEditorOnly
+				// Many of them are editor only and should not be depended on here.  If you re-add any
+				// dependent modules, please confirm by "launching" or packaging that you are not introducing
+				// editor dependencies
+				//
+				//"MeshDescription",
+				//"ProxyLODMeshReduction", // currently required to be public due to IVoxelBasedCSG API
+
+				"Core",
+				"Eigen",
+                "InteractiveToolsFramework",
+				"GeometryCore",
+				"GeometryFramework",
+				"GeometryAlgorithms",
+				"DynamicMesh",
+				"MeshConversion",
+				"MeshDescription",
+                "StaticMeshDescription",
+                "SkeletalMeshDescription",
+				"ModelingComponents",
+				"ModelingOperators",
+
+				// ... add other public dependencies that you statically link with here ...
+			}
+            );
+			
+		
+		PrivateDependencyModuleNames.AddRange(
+			new string[]
+			{
+				// These dependencies were commented out when we split out MeshModelingToolsEditorOnly
+				// Many of them are editor only and should not be depended on here.  If you re-add any
+				// dependent modules, please confirm by "launching" or packaging that you are not introducing
+				// editor dependencies
+				//
+				//"RenderCore",
+				//"RHI",
+				//"MeshUtilities",    // temp for saving mesh asset
+				//"UnrealEd",
+				//"MeshBuilder",
+                //"MeshDescriptionOperations",
+				//"MeshUtilitiesCommon",
+				//"MeshReductionInterface", // for UE standard simplification
+                //"ProxyLODMeshReduction", // for mesh merging voxel-based csg
+				//"Slate",
+				//"SlateCore",
+
+				"CoreUObject",
+				"Engine",
+				"RenderCore",
+				"ModelingOperators",
+				"InputCore",
+				"PhysicsCore",
+
+				// ... add private dependencies that you statically link with here ...
+			}
+			);
+		
+		
+		DynamicallyLoadedModuleNames.AddRange(
+			new string[]
+			{
+				// ... add any modules that your module loads dynamically here ...
+			}
+			);
+	}
+}
diff --git a/Engine/Plugins/Experimental/MeshModelingToolsetExp/Source/MeshModelingToolsExp/Private/DrawPolyPathTool.cpp b/Engine/Plugins/Experimental/MeshModelingToolsetExp/Source/MeshModelingToolsExp/Private/DrawPolyPathTool.cpp
new file mode 100644
--- /dev/null
+++ b/Engine/Plugins/Experimental/MeshModelingToolsetExp/Source/MeshModelingToolsExp/Private/DrawPolyPathTool.cpp
@@ -0,0 +1,860 @@
+// Copyright Epic Games, Inc. All Rights Reserved.
+
+#include "DrawPolyPathTool.h"
+#include "InteractiveToolManager.h"
+#include "ToolBuilderUtil.h"
+
+#include "BaseBehaviors/SingleClickBehavior.h"
+#include "BaseBehaviors/MouseHoverBehavior.h"
+
+#include "ToolSceneQueriesUtil.h"
+#include "Intersection/IntersectionUtil.h"
+#include "Util/ColorConstants.h"
+#include "ToolSetupUtil.h"
+#include "DynamicMesh/MeshIndexUtil.h"
+#include "Generators/RectangleMeshGenerator.h"
+#include "Distance/DistLine3Line3.h"
+#include "ModelingObjectsCreationAPI.h"
+#include "DynamicMesh/MeshTransforms.h"
+#include "Selection/ToolSelectionUtil.h"
+#include "Operations/ExtrudeMesh.h"
+#include "DynamicMesh/MeshNormals.h"
+
+#include "ExplicitUseGeometryMathTypes.h"		// using UE::Geometry::(math types)
+using namespace UE::Geometry;
+
+#define LOCTEXT_NAMESPACE "UDrawPolyPathTool"
+
+namespace
+{
+	// Simple mesh generator that generates a quad for each edge of an input polygon
+	class FPolygonEdgeMeshGenerator : public FMeshShapeGenerator
+	{
+
+	private:
+
+		// Polygon to triangulate. Assumed to be closed (i.e. last edge is (LastVertex, FirstVertex)). If Polygon has 
+		// self-intersections or degenerate edges, result is undefined.
+		TArray<FFrame3d> Polygon;
+
+		// For each polygon vertex, a scale factor for the patch width at that vertex. Helps keep the width constant
+		// going around acute corners. 
+		TArray<double> OffsetScaleFactors;
+
+		// Width of quads to generate.
+		double Width = 1.0;
+
+		// Normal vector of all vertices will be set to this value. Default is +Z axis.
+		FVector3d Normal = FVector3d::UnitZ();
+
+	public:
+
+		FPolygonEdgeMeshGenerator(const TArray<FFrame3d>& InPolygon, 
+								  const TArray<double>& InOffsetScaleFactors,
+								  double InWidth = 1.0,
+								  FVector3d InNormal = FVector3d::UnitZ()) :
+			Polygon(InPolygon),
+			OffsetScaleFactors(InOffsetScaleFactors),
+			Width(InWidth),
+			Normal(InNormal)
+		{
+			check(Polygon.Num() == OffsetScaleFactors.Num());
+		}
+
+
+		// Generate the triangulation
+		virtual FMeshShapeGenerator& Generate() final
+		{
+			const int NumInputVertices = Polygon.Num();
+			check(NumInputVertices >= 3);
+			if (NumInputVertices < 3)
+			{
+				return *this;
+			}
+			const int NumVertices = 2 * NumInputVertices;
+			const int NumTriangles = 2 * NumInputVertices;
+			SetBufferSizes(NumVertices, NumTriangles, NumVertices, NumVertices);
+
+			// Trace the input path, placing vertices on either side of each input vertex 
+
+			const FVector3d LeftVertex{ 0, -Width, 0 };
+			const FVector3d RightVertex{ 0, Width, 0 };
+			for (int CurrentInputVertex = 0; CurrentInputVertex < NumInputVertices; ++CurrentInputVertex)
+			{
+				const FFrame3d& CurrentFrame = Polygon[CurrentInputVertex];
+				const int NewVertexAIndex = 2 * CurrentInputVertex;
+				const int NewVertexBIndex = NewVertexAIndex + 1;
+				Vertices[NewVertexAIndex] = CurrentFrame.FromFramePoint(OffsetScaleFactors[CurrentInputVertex] * LeftVertex);
+				Vertices[NewVertexBIndex] = CurrentFrame.FromFramePoint(OffsetScaleFactors[CurrentInputVertex] * RightVertex);
+			}
+
+			// Triangulate the vertices we just placed
+
+			for (int CurrentVertex = 0; CurrentVertex < NumInputVertices; ++CurrentVertex)
+			{
+				const int NewVertexA = 2 * CurrentVertex;
+				const int NewVertexB = NewVertexA + 1;
+				const int NewVertexC = (NewVertexA + 2) % NumVertices;
+				const int NewVertexD = (NewVertexA + 3) % NumVertices;
+
+				FIndex3i NewTriA{ NewVertexA, NewVertexB, NewVertexC };
+				const int NewTriAIndex = 2 * CurrentVertex;
+
+				SetTriangle(NewTriAIndex, NewTriA);
+				SetTriangleUVs(NewTriAIndex, NewTriA);
+				SetTriangleNormals(NewTriAIndex, NewTriA);
+				SetTrianglePolygon(NewTriAIndex, 0);
+
+				const int NewTriBIndex = NewTriAIndex + 1;
+				FIndex3i NewTriB{ NewVertexC, NewVertexB, NewVertexD };
+
+				SetTriangle(NewTriBIndex, NewTriB);
+				SetTriangleUVs(NewTriBIndex, NewTriB);
+				SetTriangleNormals(NewTriBIndex, NewTriB);
+				SetTrianglePolygon(NewTriBIndex, 0);
+			}
+
+			// Set UVs, etc. for the vertices we put down previously
+
+			FAxisAlignedBox2d BoundingBox;
+			for (auto& CurrentFrame : Polygon)
+			{
+				BoundingBox.Contain(FVector2d{ CurrentFrame.Origin.X, CurrentFrame.Origin.Y });
+			}
+			BoundingBox.Max += {Width, Width};
+			BoundingBox.Min -= {Width, Width};
+
+			double BoxWidth = BoundingBox.Width(), BoxHeight = BoundingBox.Height();
+			double UVScale = FMath::Max(BoxWidth, BoxHeight);
+
+			for (int NewVertexIndex = 0; NewVertexIndex < NumVertices; ++NewVertexIndex)
+			{
+				FVector3d Pos = Vertices[NewVertexIndex];
+
+				UVs[NewVertexIndex] = FVector2f(
+					(float)((Pos.X - BoundingBox.Min.X) / UVScale),
+					(float)((Pos.Y - BoundingBox.Min.Y) / UVScale));
+				UVParentVertex[NewVertexIndex] = NewVertexIndex;
+				Normals[NewVertexIndex] = FVector3f(Normal);
+				NormalParentVertex[NewVertexIndex] = NewVertexIndex;
+			}
+
+			return *this;
+		}
+	};		// class FPolygonEdgeMeshGenerator
+
+}	// unnamed namespace
+
+/*
+ * ToolBuilder
+ */
+bool UDrawPolyPathToolBuilder::CanBuildTool(const FToolBuilderState& SceneState) const
+{
+	return true;
+}
+
+UInteractiveTool* UDrawPolyPathToolBuilder::BuildTool(const FToolBuilderState& SceneState) const
+{
+	UDrawPolyPathTool* NewTool = NewObject<UDrawPolyPathTool>(SceneState.ToolManager);
+	NewTool->SetWorld(SceneState.World);
+	return NewTool;
+}
+
+/*
+* Tool methods
+*/
+void UDrawPolyPathTool::SetWorld(UWorld* World)
+{
+	this->TargetWorld = World;
+}
+
+void UDrawPolyPathTool::Setup()
+{
+	UInteractiveTool::Setup();
+
+	// register click behavior
+	USingleClickInputBehavior* ClickBehavior = NewObject<USingleClickInputBehavior>(this);
+	ClickBehavior->Initialize(this);
+	AddInputBehavior(ClickBehavior);
+
+	UMouseHoverBehavior* HoverBehavior = NewObject<UMouseHoverBehavior>(this);
+	HoverBehavior->Initialize(this);
+	AddInputBehavior(HoverBehavior);
+
+	DrawPlaneWorld = FFrame3d();
+
+	PlaneMechanic = NewObject<UConstructionPlaneMechanic>(this);
+	PlaneMechanic->Setup(this);
+	PlaneMechanic->CanUpdatePlaneFunc = [this]() { return CanUpdateDrawPlane(); };
+	PlaneMechanic->Initialize(TargetWorld, DrawPlaneWorld);
+	PlaneMechanic->UpdateClickPriority(ClickBehavior->GetPriority().MakeHigher());
+	PlaneMechanic->OnPlaneChanged.AddLambda([this]() {
+		DrawPlaneWorld = PlaneMechanic->Plane;
+		UpdateSurfacePathPlane();
+	});
+
+	OutputTypeProperties = NewObject<UCreateMeshObjectTypeProperties>(this);
+	OutputTypeProperties->RestoreProperties(this);
+	OutputTypeProperties->InitializeDefault();
+	OutputTypeProperties->WatchProperty(OutputTypeProperties->OutputType, [this](FString) { OutputTypeProperties->UpdatePropertyVisibility(); });
+	AddToolPropertySource(OutputTypeProperties);
+
+	// add properties
+	TransformProps = NewObject<UDrawPolyPathProperties>(this);
+	TransformProps->RestoreProperties(this);
+	AddToolPropertySource(TransformProps);
+
+	ExtrudeProperties = NewObject<UDrawPolyPathExtrudeProperties>();
+	ExtrudeProperties->RestoreProperties(this);
+	AddToolPropertySource(ExtrudeProperties);
+	SetToolPropertySourceEnabled(ExtrudeProperties, false);
+
+	// initialize material properties for new objects
+	MaterialProperties = NewObject<UNewMeshMaterialProperties>(this);
+	MaterialProperties->RestoreProperties(this);
+	MaterialProperties->bShowExtendedOptions = false;
+	AddToolPropertySource(MaterialProperties);
+
+	// begin path draw
+	InitializeNewSurfacePath();
+
+	SetToolDisplayName(LOCTEXT("ToolName", "Extrude PolyPath"));
+	ShowStartupMessage();
+}
+
+
+void UDrawPolyPathTool::Shutdown(EToolShutdownType ShutdownType)
+{
+	PlaneMechanic->Shutdown();
+	PlaneMechanic = nullptr;
+
+	OutputTypeProperties->SaveProperties(this);
+	TransformProps->SaveProperties(this);
+	ExtrudeProperties->SaveProperties(this);
+	MaterialProperties->SaveProperties(this);
+
+	ClearPreview();
+}
+
+
+
+
+void UDrawPolyPathTool::RegisterActions(FInteractiveToolActionSet& ActionSet)
+{
+	//ActionSet.RegisterAction(this, (int32)EStandardToolActions::BaseClientDefinedActionID + 1,
+	//	TEXT("PopLastVertex"),
+	//	LOCTEXT("PopLastVertex", "Pop Last Vertex"),
+	//	LOCTEXT("PopLastVertexTooltip", "Pop last vertex added to polygon"),
+	//	EModifierKey::None, EKeys::BackSpace,
+	//	[this]() { PopLastVertexAction(); });
+
+
+	//ActionSet.RegisterAction(this, (int32)EStandardToolActions::BaseClientDefinedActionID + 2,
+	//	TEXT("ToggleGizmo"),
+	//	LOCTEXT("ToggleGizmo", "Toggle Gizmo"),
+	//	LOCTEXT("ToggleGizmoTooltip", "Toggle visibility of the transformation Gizmo"),
+	//	EModifierKey::None, EKeys::A,
+	//	[this]() { PolygonProperties->bShowGizmo = !PolygonProperties->bShowGizmo; });
+}
+
+
+bool UDrawPolyPathTool::HitTest(const FRay& Ray, FHitResult& OutHit)
+{
+	if (SurfacePathMechanic != nullptr)
+	{
+		FFrame3d HitPoint;
+		if (SurfacePathMechanic->IsHitByRay(FRay3d(Ray), HitPoint))
+		{
+			OutHit.Distance = FRay3d(Ray).Project(HitPoint.Origin);
+			OutHit.ImpactPoint = (FVector)HitPoint.Origin;
+			OutHit.ImpactNormal = (FVector)HitPoint.Z();
+			return true;
+		}
+		return false;
+	}
+	else if (CurveDistMechanic != nullptr)
+	{
+		OutHit.ImpactPoint = Ray.PointAt(100);
+		OutHit.Distance = 100;
+		return true;
+	}
+	else if (ExtrudeHeightMechanic != nullptr)
+	{
+		OutHit.ImpactPoint = Ray.PointAt(100);
+		OutHit.Distance = 100;
+		return true;
+	}
+
+	return false;
+}
+
+
+
+FInputRayHit UDrawPolyPathTool::IsHitByClick(const FInputDeviceRay& ClickPos)
+{
+	FHitResult OutHit;
+	if (HitTest(ClickPos.WorldRay, OutHit))
+	{
+		return FInputRayHit(OutHit.Distance);
+	}
+
+	// background capture, if nothing else is hit
+	return FInputRayHit(TNumericLimits<float>::Max());
+}
+
+
+void UDrawPolyPathTool::OnClicked(const FInputDeviceRay& ClickPos)
+{
+	if (SurfacePathMechanic != nullptr)
+	{
+		bool bIsFirstPoint = SurfacePathMechanic->HitPath.Num() == 0;
+		if (SurfacePathMechanic->TryAddPointFromRay(ClickPos.WorldRay))
+		{
+			if (SurfacePathMechanic->IsDone())
+			{
+				bPathIsClosed = SurfacePathMechanic->LoopWasClosed();
+				GetToolManager()->EmitObjectChange(this, MakeUnique<FDrawPolyPathStateChange>(CurrentCurveTimestamp), LOCTEXT("DrawPolyPathBeginOffset", "Set Offset"));
+				OnCompleteSurfacePath();
+			}
+			else
+			{
+				GetToolManager()->EmitObjectChange(this, MakeUnique<FDrawPolyPathStateChange>(CurrentCurveTimestamp), LOCTEXT("DrawPolyPathBeginPath", "Begin Path"));
+			}
+		}
+	}
+	else if (CurveDistMechanic != nullptr)
+	{
+		GetToolManager()->EmitObjectChange(this, MakeUnique<FDrawPolyPathStateChange>(CurrentCurveTimestamp), LOCTEXT("DrawPolyPathBeginHeight", "Set Height"));
+		OnCompleteOffsetDistance();
+	}
+	else if (ExtrudeHeightMechanic != nullptr)
+	{
+		OnCompleteExtrudeHeight();
+	}
+}
+
+
+
+
+FInputRayHit UDrawPolyPathTool::BeginHoverSequenceHitTest(const FInputDeviceRay& PressPos)
+{
+	FHitResult OutHit;
+	if (HitTest(PressPos.WorldRay, OutHit))
+	{
+		return FInputRayHit(OutHit.Distance);
+	}
+
+	// background capture, if nothing else is hit
+	return FInputRayHit(TNumericLimits<float>::Max());
+}
+
+
+bool UDrawPolyPathTool::OnUpdateHover(const FInputDeviceRay& DevicePos)
+{
+	if (SurfacePathMechanic != nullptr)
+	{
+		SurfacePathMechanic->UpdatePreviewPoint(DevicePos.WorldRay);
+	} 
+	else if (CurveDistMechanic != nullptr)
+	{
+		CurveDistMechanic->UpdateCurrentDistance(DevicePos.WorldRay);
+		TransformProps->Width = CurveDistMechanic->CurrentDistance;
+		CurOffsetDistance = CurveDistMechanic->CurrentDistance;
+		UpdatePathPreview();
+	}
+	else if (ExtrudeHeightMechanic != nullptr)
+	{
+		ExtrudeHeightMechanic->UpdateCurrentDistance(DevicePos.WorldRay);
+		CurHeight = ExtrudeHeightMechanic->CurrentHeight;
+		TransformProps->Height = ExtrudeHeightMechanic->CurrentHeight;
+		UpdateExtrudePreview();
+	}
+	return true;
+}
+
+
+
+
+
+
+void UDrawPolyPathTool::OnTick(float DeltaTime)
+{
+	if (PlaneMechanic != nullptr)
+	{
+		PlaneMechanic->Tick(DeltaTime);
+	}
+}
+
+
+void UDrawPolyPathTool::Render(IToolsContextRenderAPI* RenderAPI)
+{
+	GetToolManager()->GetContextQueriesAPI()->GetCurrentViewState(CameraState);
+
+	if (PlaneMechanic != nullptr)
+	{
+		PlaneMechanic->Render(RenderAPI);
+	}
+
+	if (ExtrudeHeightMechanic != nullptr)
+	{
+		ExtrudeHeightMechanic->Render(RenderAPI);
+	}
+	if (CurveDistMechanic != nullptr)
+	{
+		CurveDistMechanic->Render(RenderAPI);
+	}
+	if (SurfacePathMechanic != nullptr)
+	{
+		SurfacePathMechanic->Render(RenderAPI);
+	}
+}
+
+
+
+
+void UDrawPolyPathTool::InitializeNewSurfacePath()
+{
+	SurfacePathMechanic = NewObject<UCollectSurfacePathMechanic>(this);
+	SurfacePathMechanic->Setup(this);
+	double SnapTol = ToolSceneQueriesUtil::GetDefaultVisualAngleSnapThreshD();
+	SurfacePathMechanic->SpatialSnapPointsFunc = [this, SnapTol](FVector3d Position1, FVector3d Position2)
+	{
+		return true && ToolSceneQueriesUtil::PointSnapQuery(this->CameraState, Position1, Position2, SnapTol);
+	};
+	SurfacePathMechanic->SetDoubleClickOrCloseLoopMode();
+	UpdateSurfacePathPlane();
+
+	ShowStartupMessage();
+}
+
+
+bool UDrawPolyPathTool::CanUpdateDrawPlane() const
+{
+ 	return (SurfacePathMechanic != nullptr && SurfacePathMechanic->HitPath.Num() == 0);
+}
+
+void UDrawPolyPathTool::UpdateSurfacePathPlane()
+{
+	if (SurfacePathMechanic != nullptr)
+	{
+		SurfacePathMechanic->InitializePlaneSurface(DrawPlaneWorld);
+	}
+}
+
+
+void UDrawPolyPathTool::OnCompleteSurfacePath()
+{
+	check(SurfacePathMechanic != nullptr);
+
+	CurPathPoints = SurfacePathMechanic->HitPath;
+	int NumPoints = CurPathPoints.Num();
+	// align frames
+	FVector3d PlaneNormal = DrawPlaneWorld.Z();
+	CurPathPoints[0].ConstrainedAlignAxis(0, UE::Geometry::Normalized(CurPathPoints[1].Origin - CurPathPoints[0].Origin), PlaneNormal);
+	CurPathPoints[NumPoints-1].ConstrainedAlignAxis(0, UE::Geometry::Normalized(CurPathPoints[NumPoints-1].Origin - CurPathPoints[NumPoints-2].Origin), PlaneNormal);
+	double DistOffsetDelta = 0.01;
+	OffsetScaleFactors.SetNum(NumPoints);
+	OffsetScaleFactors[0] = OffsetScaleFactors[NumPoints-1] = 1.0;
+	ArcLengths.SetNum(NumPoints);
+	ArcLengths[0] = 0;
+
+	// Set local frames for path points. If the path is closed, we will adjust the first and last frames for continuity,
+	// otherwise we will leave them as set above.
+	int LastPointIndex = bPathIsClosed ? NumPoints : NumPoints - 1;
+	int FirstPointIndex = bPathIsClosed ? 0 : 1;
+	for (int j = FirstPointIndex; j < LastPointIndex; ++j)
+	{
+		int NextJ = (j + 1) % NumPoints;
+		int PrevJ = (j - 1 + NumPoints) % NumPoints;
+		FVector3d Prev(CurPathPoints[PrevJ].Origin), Next(CurPathPoints[NextJ].Origin), Cur(CurPathPoints[j].Origin);
+		ArcLengths[j] = ArcLengths[PrevJ] + Distance(Cur, Prev);
+		FLine3d Line1(FLine3d::FromPoints(Prev, Cur)), Line2(FLine3d::FromPoints(Cur, Next));
+		Line1.Origin += DistOffsetDelta * PlaneNormal.Cross(Line1.Direction);
+		Line2.Origin += DistOffsetDelta * PlaneNormal.Cross(Line2.Direction);
+
+		if (Line1.Direction.Dot(Line2.Direction) > 0.999 )
+		{
+			CurPathPoints[j].ConstrainedAlignAxis(0, UE::Geometry::Normalized(Next-Prev), PlaneNormal);
+			OffsetScaleFactors[j] = 1.0;
+		}
+		else
+		{
+			FDistLine3Line3d LineDist(Line1, Line2);
+			LineDist.GetSquared();
+			FVector3d OffsetPoint = 0.5 * (LineDist.Line1ClosestPoint + LineDist.Line2ClosestPoint);
+			OffsetScaleFactors[j] = Distance(OffsetPoint, Cur) / DistOffsetDelta;
+			FVector3d TangentDir = UE::Geometry::Normalized(OffsetPoint - Cur).Cross(PlaneNormal);
+			CurPathPoints[j].ConstrainedAlignAxis(0, TangentDir, PlaneNormal);
+		}
+	}
+	ArcLengths[NumPoints-1] = ArcLengths[NumPoints-2] + CurPathPoints[NumPoints-1].Origin.Distance(CurPathPoints[NumPoints - 2].Origin);
+
+	CurPolyLine.Reset();
+	for (const FFrame3d& Point : SurfacePathMechanic->HitPath)
+	{
+		CurPolyLine.Add(Point.Origin);
+	}
+
+	SurfacePathMechanic = nullptr;
+	if (TransformProps->WidthMode == EDrawPolyPathWidthMode::Constant)
+	{
+		BeginConstantOffsetDistance();
+	}
+	else
+	{
+		BeginInteractiveOffsetDistance();
+	}
+}
+
+
+void UDrawPolyPathTool::BeginInteractiveOffsetDistance()
+{
+	// begin setting offset distance
+	CurveDistMechanic = NewObject<USpatialCurveDistanceMechanic>(this);
+	CurveDistMechanic->Setup(this);
+	CurveDistMechanic->InitializePolyCurve(CurPolyLine, UE::Geometry::FTransform3d::Identity());
+
+	InitializePreviewMesh();
+
+	ShowOffsetMessage();
+}
+
+
+void UDrawPolyPathTool::BeginConstantOffsetDistance()
+{
+	InitializePreviewMesh();
+	CurOffsetDistance = TransformProps->Width;
+	UpdatePathPreview();
+	OnCompleteOffsetDistance();
+}
+
+
+
+void UDrawPolyPathTool::OnCompleteOffsetDistance()
+{
+	CurveDistMechanic = nullptr;
+
+	if (TransformProps->OutputType == EDrawPolyPathOutputMode::Ribbon)
+	{
+		OnCompleteExtrudeHeight();
+	}
+	else if (TransformProps->HeightMode == EDrawPolyPathHeightMode::Constant)
+	{
+		CurHeight = TransformProps->Height;
+		OnCompleteExtrudeHeight();
+	}
+	else
+	{
+		BeginInteractiveExtrudeHeight();
+	}
+}
+
+
+void UDrawPolyPathTool::OnCompleteExtrudeHeight()
+{
+	CurHeight = TransformProps->Height;
+	ExtrudeHeightMechanic = nullptr;
+
+	ClearPreview();
+
+	EmitNewObject(TransformProps->OutputType);
+
+	InitializeNewSurfacePath();
+	CurrentCurveTimestamp++;
+}
+
+
+
+void UDrawPolyPathTool::UpdatePathPreview()
+{
+	check(EditPreview != nullptr);
+
+	FDynamicMesh3 PathMesh;
+	GeneratePathMesh(PathMesh);
+	EditPreview->ReplaceMesh( MoveTempIfPossible(PathMesh) );
+}
+
+
+
+void UDrawPolyPathTool::GeneratePathMesh(FDynamicMesh3& Mesh)
+{
+	Mesh.Clear();
+	int32 NumPoints = CurPathPoints.Num();
+	if (NumPoints > 1)
+	{
+		if (bPathIsClosed)
+		{
+			FPolygonEdgeMeshGenerator MeshGen(CurPathPoints, OffsetScaleFactors, CurOffsetDistance, FVector3d::UnitZ());
+			MeshGen.Generate();
+			Mesh.Copy(&MeshGen);
+		}
+		else
+		{
+			CurPathLength = 0;
+			for (int32 k = 1; k < NumPoints; ++k)
+			{
+				CurPathLength += CurPathPoints[k].Origin.Distance(CurPathPoints[k - 1].Origin);
+			}
+
+			FRectangleMeshGenerator MeshGen;
+			MeshGen.Width = CurPathLength;
+			MeshGen.Height = 2 * CurOffsetDistance;
+			MeshGen.Normal = FVector3f::UnitZ();
+			MeshGen.Origin = FVector3d(CurPathLength / 2, 0, 0);
+			MeshGen.HeightVertexCount = 2;
+			MeshGen.WidthVertexCount = NumPoints;
+			MeshGen.Generate();
+			Mesh.Copy(&MeshGen);
+			Mesh.EnableVertexUVs(FVector2f::Zero());		// we will store arc length for each vtx in VertexUV
+
+			FAxisAlignedBox3d PathBounds = Mesh.GetBounds();
+
+			double ShiftX = 0;
+			double DeltaX = CurPathLength / (double)(NumPoints - 1);
+			for (int32 k = 0; k < NumPoints; ++k)
+			{
+				FFrame3d PathFrame = CurPathPoints[k];
+				FVector3d V0 = Mesh.GetVertex(k);
+				V0.X -= ShiftX;
+				V0.Y *= OffsetScaleFactors[k];
+				V0 = PathFrame.FromFramePoint(V0);
+				Mesh.SetVertex(k, V0);
+				Mesh.SetVertexUV(k, FVector2f((float)ArcLengths[k], (float)k));
+				FVector3d V1 = Mesh.GetVertex(NumPoints + k);
+				V1.X -= ShiftX;
+				V1.Y *= OffsetScaleFactors[k];
+				V1 = PathFrame.FromFramePoint(V1);
+				Mesh.SetVertex(NumPoints + k, V1);
+				Mesh.SetVertexUV(NumPoints + k, FVector2f((float)ArcLengths[k], (float)k));
+				ShiftX += DeltaX;
+			}
+		}
+
+		FMeshNormals::QuickRecomputeOverlayNormals(Mesh);
+	}
+}
+
+
+
+void UDrawPolyPathTool::BeginInteractiveExtrudeHeight()
+{
+	// begin extrude
+	ExtrudeHeightMechanic = NewObject<UPlaneDistanceFromHitMechanic>(this);
+	ExtrudeHeightMechanic->Setup(this);
+
+	ExtrudeHeightMechanic->WorldHitQueryFunc = [this](const FRay& WorldRay, FHitResult& HitResult)
+	{
+		return ToolSceneQueriesUtil::FindNearestVisibleObjectHit(TargetWorld, HitResult, WorldRay);
+	};
+	ExtrudeHeightMechanic->WorldPointSnapFunc = [this](const FVector3d& WorldPos, FVector3d& SnapPos)
+	{
+		return TransformProps->bSnapToWorldGrid && ToolSceneQueriesUtil::FindWorldGridSnapPoint(this, WorldPos, SnapPos);
+	};
+	ExtrudeHeightMechanic->CurrentHeight = 1.0f;  // initialize to something non-zero...prob should be based on polygon bounds maybe?
+
+	InitializePreviewMesh();
+
+	FDynamicMesh3 PathMesh;
+	GeneratePathMesh(PathMesh);
+	EditPreview->InitializeExtrudeType(MoveTemp(PathMesh), DrawPlaneWorld.Z(), nullptr, false);
+
+	FDynamicMesh3 TmpMesh;
+	EditPreview->MakeExtrudeTypeHitTargetMesh(TmpMesh, false);
+	FFrame3d UseFrame = DrawPlaneWorld; 
+	UseFrame.Origin = CurPathPoints.Last().Origin;
+	ExtrudeHeightMechanic->Initialize(MoveTemp(TmpMesh), UseFrame, true);
+
+	ShowExtrudeMessage();
+}
+
+
+void UDrawPolyPathTool::UpdateExtrudePreview()
+{
+	if (TransformProps->OutputType == EDrawPolyPathOutputMode::Ramp)
+	{
+		EditPreview->UpdateExtrudeType( [&](FDynamicMesh3& Mesh) { GenerateRampMesh(Mesh); }, true);
+	}
+	else
+	{
+		EditPreview->UpdateExtrudeType([&](FDynamicMesh3& Mesh) { GenerateExtrudeMesh(Mesh); }, true);
+	}
+}
+
+
+void UDrawPolyPathTool::InitializePreviewMesh()
+{
+	if (EditPreview == nullptr)
+	{
+		EditPreview = NewObject<UPolyEditPreviewMesh>(this);
+		EditPreview->CreateInWorld(TargetWorld, FTransform::Identity);
+		if ( MaterialProperties->Material == nullptr )
+		{
+			EditPreview->SetMaterial(
+				ToolSetupUtil::GetSelectionMaterial(FLinearColor(0.8f, 0.75f, 0.0f), GetToolManager()));
+		}
+		else
+		{
+			EditPreview->SetMaterial(MaterialProperties->Material.Get());
+		}
+	}
+}
+
+void UDrawPolyPathTool::ClearPreview()
+{
+	if (EditPreview != nullptr)
+	{
+		EditPreview->Disconnect();
+		EditPreview = nullptr;
+	}
+}
+
+
+
+void UDrawPolyPathTool::GenerateExtrudeMesh(FDynamicMesh3& PathMesh)
+{
+	FAxisAlignedBox3d Bounds = PathMesh.GetBounds();
+
+	FExtrudeMesh Extruder(&PathMesh);
+	FVector3d ExtrudeDir = DrawPlaneWorld.Z();
+	Extruder.ExtrudedPositionFunc = [&](const FVector3d& P, const FVector3f& N, int32 VID) {
+		return P + CurHeight * ExtrudeDir;
+	};
+	Extruder.UVScaleFactor = 1.0 / Bounds.MaxDim();
+	Extruder.IsPositiveOffset = (CurHeight >= 0);
+	Extruder.Apply();
+
+	FMeshNormals::QuickRecomputeOverlayNormals(PathMesh);
+}
+
+void UDrawPolyPathTool::GenerateRampMesh(FDynamicMesh3& PathMesh)
+{
+	FAxisAlignedBox3d Bounds = PathMesh.GetBounds();
+
+	FExtrudeMesh Extruder(&PathMesh);
+	FVector3d ExtrudeDir = DrawPlaneWorld.Z();
+
+	int NumPoints = CurPathPoints.Num();
+	double RampStartRatio = TransformProps->RampStartRatio;
+	double StartHeight = FMathd::Max(0.1, RampStartRatio * FMathd::Abs(CurHeight)) * FMathd::Sign(CurHeight);
+	double EndHeight = CurHeight;
+	Extruder.ExtrudedPositionFunc = [&](const FVector3d& P, const FVector3f& N, int32 VID) {
+		FVector2f UV = Extruder.Mesh->GetVertexUV(VID);
+		double UseHeight = FMathd::Lerp(StartHeight, EndHeight, (UV.X / this->CurPathLength));
+		return P + UseHeight * ExtrudeDir;
+	};
+
+	Extruder.UVScaleFactor = 1.0 / Bounds.MaxDim();
+	Extruder.IsPositiveOffset = (CurHeight >= 0);
+	Extruder.Apply();
+
+	FMeshNormals::QuickRecomputeOverlayNormals(PathMesh);
+}
+
+
+
+void UDrawPolyPathTool::EmitNewObject(EDrawPolyPathOutputMode OutputMode)
+{
+	FDynamicMesh3 PathMesh;
+	GeneratePathMesh(PathMesh);
+
+	if (OutputMode == EDrawPolyPathOutputMode::Extrusion)
+	{
+		GenerateExtrudeMesh(PathMesh);
+	}
+	else if (OutputMode == EDrawPolyPathOutputMode::Ramp)
+	{
+		GenerateRampMesh(PathMesh);
+	}
+	PathMesh.DiscardVertexUVs();  // throw away arc lengths
+
+	FFrame3d MeshTransform = DrawPlaneWorld;
+	FVector3d Center = PathMesh.GetBounds().Center();
+	MeshTransform.Origin = MeshTransform.ToPlane(Center, 2);
+	MeshTransforms::WorldToFrameCoords(PathMesh, MeshTransform);
+
+	GetToolManager()->BeginUndoTransaction(LOCTEXT("CreatePolyPath", "Create PolyPath"));
+
+	FCreateMeshObjectParams NewMeshObjectParams;
+	NewMeshObjectParams.TargetWorld = TargetWorld;
+	NewMeshObjectParams.Transform = MeshTransform.ToFTransform();
+	NewMeshObjectParams.BaseName = TEXT("Path");
+	NewMeshObjectParams.Materials.Add(MaterialProperties->Material.Get());
+	NewMeshObjectParams.SetMesh(&PathMesh);
+	OutputTypeProperties->ConfigureCreateMeshObjectParams(NewMeshObjectParams);
+	FCreateMeshObjectResult Result = UE::Modeling::CreateMeshObject(GetToolManager(), MoveTemp(NewMeshObjectParams));
+	if (Result.IsOK() && Result.NewActor != nullptr)
+	{
+		ToolSelectionUtil::SetNewActorSelection(GetToolManager(), Result.NewActor);
+	}
+
+	GetToolManager()->EndUndoTransaction();
+}
+
+
+
+void UDrawPolyPathTool::ShowStartupMessage()
+{
+	GetToolManager()->DisplayMessage(
+		LOCTEXT("OnStartDraw", "Use this Tool to draw a path on the Drawing Plane, and then Extrude. Left-click to place points, Double-click (or close) to complete path. Ctrl-click on the scene to reposition the Plane (Shift+Ctrl-click to ignore Normal)."),
+		EToolMessageLevel::UserNotification);
+}
+
+void UDrawPolyPathTool::ShowOffsetMessage()
+{
+	GetToolManager()->DisplayMessage(
+		LOCTEXT("OnStartOffset", "Set the Width of the Extrusion by clicking on the Drawing Plane."),
+		EToolMessageLevel::UserNotification);
+}
+
+void UDrawPolyPathTool::ShowExtrudeMessage()
+{
+	GetToolManager()->DisplayMessage(
+		LOCTEXT("OnStartExtrude", "Set the Height of the Extrusion by positioning the mouse over the extrusion volume, or over the scene to snap to relative heights."),
+		EToolMessageLevel::UserNotification);
+}
+
+
+
+void UDrawPolyPathTool::UndoCurrentOperation()
+{
+	if (SurfacePathMechanic != nullptr)
+	{
+		SurfacePathMechanic->PopLastPoint();
+		if (SurfacePathMechanic->HitPath.Num() == 0)
+		{
+			CurrentCurveTimestamp++;
+		}
+	}
+	else if (CurveDistMechanic != nullptr)
+	{
+		CurveDistMechanic = nullptr;
+		ClearPreview();
+		InitializeNewSurfacePath();
+		SurfacePathMechanic->HitPath = CurPathPoints;
+	}
+	else if (ExtrudeHeightMechanic != nullptr)
+	{
+		ExtrudeHeightMechanic = nullptr;
+		BeginInteractiveOffsetDistance();
+	}
+}
+
+
+void FDrawPolyPathStateChange::Revert(UObject* Object)
+{
+	Cast<UDrawPolyPathTool>(Object)->UndoCurrentOperation();
+	bHaveDoneUndo = true;
+}
+bool FDrawPolyPathStateChange::HasExpired(UObject* Object) const
+{
+	return bHaveDoneUndo || (Cast<UDrawPolyPathTool>(Object)->CheckInCurve(CurveTimestamp) == false);
+}
+FString FDrawPolyPathStateChange::ToString() const
+{
+	return TEXT("FDrawPolyPathStateChange");
+}
+
+
+
+
+#undef LOCTEXT_NAMESPACE
diff --git a/Engine/Plugins/Experimental/MeshModelingToolsetExp/Source/MeshModelingToolsExp/Private/EditNormalsTool.cpp b/Engine/Plugins/Experimental/MeshModelingToolsetExp/Source/MeshModelingToolsExp/Private/EditNormalsTool.cpp
new file mode 100644
--- /dev/null
+++ b/Engine/Plugins/Experimental/MeshModelingToolsetExp/Source/MeshModelingToolsExp/Private/EditNormalsTool.cpp
@@ -0,0 +1,319 @@
+// Copyright Epic Games, Inc. All Rights Reserved.
+
+#include "EditNormalsTool.h"
+#include "InteractiveToolManager.h"
+#include "ToolBuilderUtil.h"
+
+#include "ToolSetupUtil.h"
+
+#include "DynamicMesh/DynamicMesh3.h"
+#include "BaseBehaviors/MultiClickSequenceInputBehavior.h"
+#include "Selection/SelectClickedAction.h"
+
+#include "MeshDescriptionToDynamicMesh.h"
+#include "DynamicMeshToMeshDescription.h"
+
+#include "InteractiveGizmoManager.h"
+
+#include "AssetUtils/MeshDescriptionUtil.h"
+#include "Engine/Classes/Engine/StaticMesh.h"
+#include "Engine/Classes/Components/StaticMeshComponent.h"
+
+#include "TargetInterfaces/MeshDescriptionCommitter.h"
+#include "TargetInterfaces/MeshDescriptionProvider.h"
+#include "TargetInterfaces/PrimitiveComponentBackedTarget.h"
+#include "ToolTargetManager.h"
+
+#include "ExplicitUseGeometryMathTypes.h"		// using UE::Geometry::(math types)
+using namespace UE::Geometry;
+
+#define LOCTEXT_NAMESPACE "UEditNormalsTool"
+
+/*
+ * ToolBuilder
+ */
+
+
+const FToolTargetTypeRequirements& UEditNormalsToolBuilder::GetTargetRequirements() const
+{
+	static FToolTargetTypeRequirements TypeRequirements({
+		UMeshDescriptionCommitter::StaticClass(),
+		UMeshDescriptionProvider::StaticClass(),
+		UPrimitiveComponentBackedTarget::StaticClass()
+		});
+	return TypeRequirements;
+}
+
+bool UEditNormalsToolBuilder::CanBuildTool(const FToolBuilderState& SceneState) const
+{
+	return SceneState.TargetManager->CountSelectedAndTargetable(SceneState, GetTargetRequirements()) > 0;
+}
+
+UInteractiveTool* UEditNormalsToolBuilder::BuildTool(const FToolBuilderState& SceneState) const
+{
+	UEditNormalsTool* NewTool = NewObject<UEditNormalsTool>(SceneState.ToolManager);
+
+	TArray<TObjectPtr<UToolTarget>> Targets = SceneState.TargetManager->BuildAllSelectedTargetable(SceneState, GetTargetRequirements());
+	NewTool->SetTargets(MoveTemp(Targets));
+	NewTool->SetWorld(SceneState.World);
+
+	return NewTool;
+}
+
+
+
+
+/*
+ * Tool
+ */
+
+UEditNormalsToolProperties::UEditNormalsToolProperties()
+{
+	bFixInconsistentNormals = false;
+	bInvertNormals = false;
+	bRecomputeNormals = true;
+	NormalCalculationMethod = ENormalCalculationMethod::AreaAngleWeighting;
+	SplitNormalMethod = ESplitNormalMethod::UseExistingTopology;
+	SharpEdgeAngleThreshold = 60;
+	bAllowSharpVertices = false;
+}
+
+UEditNormalsAdvancedProperties::UEditNormalsAdvancedProperties()
+{
+}
+
+
+UEditNormalsTool::UEditNormalsTool()
+{
+}
+
+void UEditNormalsTool::SetWorld(UWorld* World)
+{
+	this->TargetWorld = World;
+}
+
+void UEditNormalsTool::Setup()
+{
+	UInteractiveTool::Setup();
+
+	// hide input StaticMeshComponent
+	for (int32 ComponentIdx = 0, NumTargets = Targets.Num(); ComponentIdx < NumTargets; ComponentIdx++)
+	{
+		IPrimitiveComponentBackedTarget* TargetComponent = TargetComponentInterface(ComponentIdx);
+		TargetComponent->SetOwnerVisibility(false);
+	}
+
+	BasicProperties = NewObject<UEditNormalsToolProperties>(this, TEXT("Mesh Normals Settings"));
+	BasicProperties->RestoreProperties(this);
+	AddToolPropertySource(BasicProperties);
+
+	AdvancedProperties = NewObject<UEditNormalsAdvancedProperties>(this, TEXT("Advanced Settings"));
+	AdvancedProperties->RestoreProperties(this);
+	AddToolPropertySource(AdvancedProperties);
+
+	// initialize the PreviewMesh+BackgroundCompute object
+	UpdateNumPreviews();
+
+	
+	
+	for (UMeshOpPreviewWithBackgroundCompute* Preview : Previews)
+	{
+		Preview->InvalidateResult();
+	}
+
+	SetToolDisplayName(LOCTEXT("ToolName", "Edit Normals"));
+	GetToolManager()->DisplayMessage(
+		LOCTEXT("OnStartTool", "Configure or Recalculate Normals on a Mesh (disables autogenerated Normals)"),
+		EToolMessageLevel::UserNotification);
+}
+
+
+void UEditNormalsTool::UpdateNumPreviews()
+{
+	int32 CurrentNumPreview = Previews.Num();
+	int32 TargetNumPreview = Targets.Num();
+	if (TargetNumPreview < CurrentNumPreview)
+	{
+		for (int32 PreviewIdx = CurrentNumPreview - 1; PreviewIdx >= TargetNumPreview; PreviewIdx--)
+		{
+			Previews[PreviewIdx]->Cancel();
+		}
+		Previews.SetNum(TargetNumPreview);
+		OriginalDynamicMeshes.SetNum(TargetNumPreview);
+	}
+	else
+	{
+		OriginalDynamicMeshes.SetNum(TargetNumPreview);
+		for (int32 PreviewIdx = CurrentNumPreview; PreviewIdx < TargetNumPreview; PreviewIdx++)
+		{
+			UEditNormalsOperatorFactory *OpFactory = NewObject<UEditNormalsOperatorFactory>();
+			OpFactory->Tool = this;
+			OpFactory->ComponentIndex = PreviewIdx;
+			OriginalDynamicMeshes[PreviewIdx] = MakeShared<FDynamicMesh3, ESPMode::ThreadSafe>();
+			FMeshDescriptionToDynamicMesh Converter;
+			Converter.Convert(TargetMeshProviderInterface(PreviewIdx)->GetMeshDescription(), *OriginalDynamicMeshes[PreviewIdx]);
+
+			UMeshOpPreviewWithBackgroundCompute* Preview = Previews.Add_GetRef(NewObject<UMeshOpPreviewWithBackgroundCompute>(OpFactory, "Preview"));
+			Preview->Setup(this->TargetWorld, OpFactory);
+			Preview->PreviewMesh->SetTangentsMode(EDynamicMeshComponentTangentsMode::AutoCalculated);
+
+			FComponentMaterialSet MaterialSet;
+			TargetMaterialInterface(PreviewIdx)->GetMaterialSet(MaterialSet);
+			Preview->ConfigureMaterials(MaterialSet.Materials,
+				ToolSetupUtil::GetDefaultWorkingMaterial(GetToolManager())
+			);
+
+			Preview->SetVisibility(true);
+		}
+	}
+}
+
+
+void UEditNormalsTool::Shutdown(EToolShutdownType ShutdownType)
+{
+	BasicProperties->SaveProperties(this);
+	AdvancedProperties->SaveProperties(this);
+
+	// Restore (unhide) the source meshes
+	for (int32 ComponentIdx = 0, NumTargets = Targets.Num(); ComponentIdx < NumTargets; ComponentIdx++)
+	{
+		IPrimitiveComponentBackedTarget* TargetComponent = TargetComponentInterface(ComponentIdx);
+		TargetComponent->SetOwnerVisibility(true);
+	}
+
+	TArray<FDynamicMeshOpResult> Results;
+	for (UMeshOpPreviewWithBackgroundCompute* Preview : Previews)
+	{
+		Results.Add(Preview->Shutdown());
+	}
+	if (ShutdownType == EToolShutdownType::Accept)
+	{
+		GenerateAsset(Results);
+	}
+}
+
+TUniquePtr<FDynamicMeshOperator> UEditNormalsOperatorFactory::MakeNewOperator()
+{
+	TUniquePtr<FEditNormalsOp> NormalsOp = MakeUnique<FEditNormalsOp>();
+	NormalsOp->bFixInconsistentNormals = Tool->BasicProperties->bFixInconsistentNormals;
+	NormalsOp->bInvertNormals = Tool->BasicProperties->bInvertNormals;
+	NormalsOp->bRecomputeNormals = Tool->BasicProperties->bRecomputeNormals;
+	NormalsOp->SplitNormalMethod = Tool->BasicProperties->SplitNormalMethod;
+	NormalsOp->bAllowSharpVertices = Tool->BasicProperties->bAllowSharpVertices;
+	NormalsOp->NormalCalculationMethod = Tool->BasicProperties->NormalCalculationMethod;
+	NormalsOp->NormalSplitThreshold = Tool->BasicProperties->SharpEdgeAngleThreshold;
+
+	FTransform LocalToWorld = Tool->TargetComponentInterface(ComponentIndex)->GetWorldTransform();
+	NormalsOp->OriginalMesh = Tool->OriginalDynamicMeshes[ComponentIndex];
+	
+	NormalsOp->SetTransform(LocalToWorld);
+
+	return NormalsOp;
+}
+
+
+
+void UEditNormalsTool::Render(IToolsContextRenderAPI* RenderAPI)
+{
+}
+
+void UEditNormalsTool::OnTick(float DeltaTime)
+{
+	for (UMeshOpPreviewWithBackgroundCompute* Preview : Previews)
+	{
+		Preview->Tick(DeltaTime);
+	}
+}
+
+
+
+#if WITH_EDITOR
+void UEditNormalsTool::PostEditChangeProperty(FPropertyChangedEvent& PropertyChangedEvent)
+{
+	UpdateNumPreviews();
+	for (UMeshOpPreviewWithBackgroundCompute* Preview : Previews)
+	{
+		Preview->InvalidateResult();
+	}
+}
+#endif
+
+void UEditNormalsTool::OnPropertyModified(UObject* PropertySet, FProperty* Property)
+{
+	UpdateNumPreviews();
+	for (UMeshOpPreviewWithBackgroundCompute* Preview : Previews)
+	{
+		Preview->InvalidateResult();
+	}
+}
+
+
+
+bool UEditNormalsTool::HasAccept() const
+{
+	return true;
+}
+
+bool UEditNormalsTool::CanAccept() const
+{
+	for (UMeshOpPreviewWithBackgroundCompute* Preview : Previews)
+	{
+		if (!Preview->HaveValidResult())
+		{
+			return false;
+		}
+	}
+	return Super::CanAccept();
+}
+
+
+void UEditNormalsTool::GenerateAsset(const TArray<FDynamicMeshOpResult>& Results)
+{
+	GetToolManager()->BeginUndoTransaction(LOCTEXT("EditNormalsToolTransactionName", "Edit Normals Tool"));
+
+	check(Results.Num() == Targets.Num());
+	
+	for (int32 ComponentIdx = 0; ComponentIdx < Targets.Num(); ComponentIdx++)
+	{
+		// disable auto-generated normals StaticMesh build setting
+		UStaticMeshComponent* StaticMeshComponent = Cast<UStaticMeshComponent>(TargetComponentInterface(ComponentIdx)->GetOwnerComponent());
+		if (StaticMeshComponent != nullptr)
+		{
+			UStaticMesh* StaticMesh = StaticMeshComponent->GetStaticMesh();
+			if (ensure(StaticMesh != nullptr))
+			{
+				StaticMesh->Modify();
+				UE::MeshDescription::FStaticMeshBuildSettingChange SettingsChange;
+				SettingsChange.AutoGeneratedNormals = UE::MeshDescription::EBuildSettingBoolChange::Disable;
+				UE::MeshDescription::ConfigureBuildSettings(StaticMesh, 0, SettingsChange);
+			}
+		}
+
+		check(Results[ComponentIdx].Mesh.Get() != nullptr);
+		TargetMeshCommitterInterface(ComponentIdx)->CommitMeshDescription([&Results, &ComponentIdx, this](const IMeshDescriptionCommitter::FCommitterParams& CommitParams)
+		{
+			FDynamicMeshToMeshDescription Converter;
+			
+			if (BasicProperties->WillTopologyChange() || !FDynamicMeshToMeshDescription::HaveMatchingElementCounts(Results[ComponentIdx].Mesh.Get(), CommitParams.MeshDescriptionOut, false, true))
+			{
+				// full conversion if normal topology changed or faces were inverted
+				Converter.Convert(Results[ComponentIdx].Mesh.Get(), *CommitParams.MeshDescriptionOut);
+			}
+			else
+			{
+				// otherwise just copy attributes
+				const bool bUpdateNormals = true;
+				const bool bCopyOverlayTangents = false;
+				const bool bCopyOverlayUVs = false;
+				Converter.UpdateAttributes(Results[ComponentIdx].Mesh.Get(), *CommitParams.MeshDescriptionOut, bUpdateNormals, bCopyOverlayTangents, bCopyOverlayUVs);
+			}
+		});
+	}
+
+	GetToolManager()->EndUndoTransaction();
+}
+
+
+
+
+#undef LOCTEXT_NAMESPACE
diff --git a/Engine/Plugins/Experimental/MeshModelingToolsetExp/Source/MeshModelingToolsExp/Private/FindPolygonsAlgorithm.cpp b/Engine/Plugins/Experimental/MeshModelingToolsetExp/Source/MeshModelingToolsExp/Private/FindPolygonsAlgorithm.cpp
new file mode 100644
--- /dev/null
+++ b/Engine/Plugins/Experimental/MeshModelingToolsetExp/Source/MeshModelingToolsExp/Private/FindPolygonsAlgorithm.cpp
@@ -0,0 +1,402 @@
+// Copyright Epic Games, Inc. All Rights Reserved.
+
+#include "FindPolygonsAlgorithm.h"
+#include "DynamicMesh/MeshNormals.h"
+#include "Selections/MeshConnectedComponents.h"
+#include "Util/IndexUtil.h"
+#include "Parameterization/MeshDijkstra.h"
+#include "Parameterization/IncrementalMeshDijkstra.h"
+#include "Parameterization/MeshRegionGraph.h"
+
+#include "Curve/DynamicGraph3.h"
+
+
+
+#include "ExplicitUseGeometryMathTypes.h"		// using UE::Geometry::(math types)
+using namespace UE::Geometry;
+
+FFindPolygonsAlgorithm::FFindPolygonsAlgorithm(FDynamicMesh3* MeshIn)
+{
+	Mesh = MeshIn;
+}
+
+
+
+bool FFindPolygonsAlgorithm::FindPolygonsFromUVIslands()
+{
+	const FDynamicMeshUVOverlay* UV = Mesh->Attributes()->GetUVLayer(0);
+
+	FMeshConnectedComponents Components(Mesh);
+	Components.FindConnectedTriangles([&UV](int32 TriIdx0, int32 TriIdx1)
+	{
+		return UV->AreTrianglesConnected(TriIdx0, TriIdx1);
+	});
+
+	int32 NumComponents = Components.Num();
+	for (int32 ci = 0; ci < NumComponents; ++ci)
+	{
+		FoundPolygons.Add(Components.GetComponent(ci).Indices);
+	}
+
+	PostProcessPolygons(false);
+	SetGroupsFromPolygons();
+
+	return (FoundPolygons.Num() > 0);
+}
+
+
+
+bool FFindPolygonsAlgorithm::FindPolygonsFromConnectedTris()
+{
+
+	FMeshConnectedComponents Components(Mesh);
+	Components.FindConnectedTriangles([this](int32 TriIdx0, int32 TriIdx1)
+	{
+		FIndex3i NbrTris = Mesh->GetTriNeighbourTris(TriIdx0);
+		int NbrIndex = IndexUtil::FindTriIndex(TriIdx1, NbrTris);
+		return (NbrIndex != IndexConstants::InvalidID);
+	});
+
+	int32 NumComponents = Components.Num();
+	for (int32 ci = 0; ci < NumComponents; ++ci)
+	{
+		FoundPolygons.Add(Components.GetComponent(ci).Indices);
+	}
+
+	PostProcessPolygons(false);
+	SetGroupsFromPolygons();
+
+	return (FoundPolygons.Num() > 0);
+}
+
+
+
+bool FFindPolygonsAlgorithm::FindPolygonsFromFaceNormals(double DotTolerance)
+{
+	DotTolerance = 1.0 - DotTolerance;
+
+	// compute face normals
+	FMeshNormals Normals(Mesh);
+	Normals.ComputeTriangleNormals();
+
+	TArray<bool> DoneTriangle;
+	DoneTriangle.SetNum(Mesh->MaxTriangleID());
+
+	TArray<int> Stack;
+
+	// grow outward from vertices until we have no more left
+	for (int TriID : Mesh->TriangleIndicesItr())
+	{
+		if (DoneTriangle[TriID] == true)
+		{
+			continue;
+		}
+
+		TArray<int> Polygon;
+		Polygon.Add(TriID);
+		DoneTriangle[TriID] = true;
+
+		Stack.SetNum(0);
+		Stack.Add(TriID);
+		while (Stack.Num() > 0)
+		{
+			int CurTri = Stack.Pop(false);
+			FIndex3i NbrTris = Mesh->GetTriNeighbourTris(CurTri);
+			for (int j = 0; j < 3; ++j)
+			{
+				if (NbrTris[j] >= 0
+					&& DoneTriangle[NbrTris[j]] == false)
+				{
+					double Dot = Normals[CurTri].Dot(Normals[NbrTris[j]]);
+					if (Dot > DotTolerance)
+					{
+						Polygon.Add(NbrTris[j]);
+						Stack.Add(NbrTris[j]);
+						DoneTriangle[NbrTris[j]] = true;
+					}
+				}
+			}
+		}
+
+		FoundPolygons.Add(Polygon);
+	}
+
+	PostProcessPolygons(true);
+	SetGroupsFromPolygons();
+
+	return (FoundPolygons.Num() > 0);
+}
+
+
+
+
+/**
+ * Dual graph of mesh faces, ie graph of edges across faces between face centers.
+ * Normals and Areas are tracked for each point.
+ */
+class FMeshFaceDualGraph : public FDynamicGraph3d
+{
+public:
+	TDynamicVectorN<double, 3> Normals;
+	TDynamicVector<double> Areas;
+
+	int AppendVertex(const FVector3d& Centroid, const FVector3d& Normal, double Area)
+	{
+		int vid = FDynamicGraph3d::AppendVertex(Centroid);
+		Normals.InsertAt({ {Normal.X, Normal.Y, Normal.Z} }, vid);
+		Areas.InsertAt(Area, vid);
+		return vid;
+	}
+
+	FVector3d GetNormal(int32 vid) const
+	{
+		return Normals.AsVector3(vid);
+	}
+
+	/** Build a Face Dual Graph for a triangle mesh */
+	static void MakeFaceDualGraphForMesh(FDynamicMesh3* Mesh, FMeshFaceDualGraph& FaceGraph)
+	{
+		for (int32 tid : Mesh->TriangleIndicesItr())
+		{
+			FVector3d Normal, Centroid;
+			double Area;
+			Mesh->GetTriInfo(tid, Normal, Area, Centroid);
+
+			int32 newvid = FaceGraph.AppendVertex(Centroid, Normal, Area);
+			check(newvid == tid);
+		}
+
+		for (int32 tid : Mesh->TriangleIndicesItr())
+		{
+			FIndex3i NbrT = Mesh->GetTriNeighbourTris(tid);
+			for (int32 j = 0; j < 3; ++j)
+			{
+				if (Mesh->IsTriangle(NbrT[j]))
+				{
+					FaceGraph.AppendEdge(tid, NbrT[j]);
+				}
+			}
+		}
+	}
+
+};
+
+
+
+
+
+
+bool FFindPolygonsAlgorithm::FindPolygonsFromFurthestPointSampling(int32 NumPoints, EWeightingType WeightingType, FVector3d WeightingCoeffs)
+{
+	NumPoints = FMath::Min(NumPoints, Mesh->VertexCount());
+
+	FMeshFaceDualGraph FaceGraph;
+	FMeshFaceDualGraph::MakeFaceDualGraphForMesh(Mesh, FaceGraph);
+
+	TIncrementalMeshDijkstra<FMeshFaceDualGraph> FurthestPoints(&FaceGraph);
+
+	TArray<int32> SeedIndices;
+
+	// need to add at least one seed point for each mesh connected component, so that all triangles are assigned a group
+	// TODO: two seed points for components that have no boundary?
+	FMeshConnectedComponents Components(Mesh);
+	Components.FindConnectedTriangles();
+	int32 NumConnected = Components.Num();
+	for (int32 k = 0; k < NumConnected; ++k)
+	{
+		SeedIndices.Add(Components[k].Indices[0]);
+	}
+
+	// initial incremental update from per-component points
+	TArray<TIncrementalMeshDijkstra<FMeshFaceDualGraph>::FSeedPoint> ComponentSeeds;
+	for (int32 vid : SeedIndices)
+	{
+		ComponentSeeds.Add(TIncrementalMeshDijkstra<FMeshFaceDualGraph>::FSeedPoint{ vid, vid, 0.0 });
+	}
+	FurthestPoints.AddSeedPoints(ComponentSeeds);
+
+	// TODO: can (approximately) bound the size of a region based on mesh area. Then can pass that in
+	// with seed point as upper distance bound. This will change the result as it will initially grow on
+	// the 'front' (and possibly no guarantee that mesh is covered?)
+	//    (initial furthest-points would be any with an invalid value...then we can ensure coverage)
+	// FindMaxGraphDistancePointID() seems like it might be somewhat expensive...
+
+	// incrementally add furthest-points
+	while ( SeedIndices.Num() < NumPoints)
+	{
+		int32 NextPointID = FurthestPoints.FindMaxGraphDistancePointID();
+		if (NextPointID >= 0)
+		{
+			TIncrementalMeshDijkstra<FMeshFaceDualGraph>::FSeedPoint NextSeedPoint{ NextPointID, NextPointID, 0.0 };
+			FurthestPoints.AddSeedPoints({ NextSeedPoint });
+			SeedIndices.Add(NextPointID);
+		}
+		else
+		{
+			break;
+		}
+	}
+
+
+	// Now that we have furthest point set, recompute a Dijkstra propagation with optional weighting
+	// (unweighted version should be the same as the FurthestPoints dijkstra though...could re-use?)
+
+	TMeshDijkstra<FMeshFaceDualGraph> SuperPixels(&FaceGraph);
+	
+	if (WeightingType == EWeightingType::NormalDeviation)
+	{
+		SuperPixels.bEnableDistanceWeighting = true;
+		SuperPixels.GetWeightedDistanceFunc = [this, WeightingCoeffs, &FaceGraph](int32 FromVID, int32 ToVID, int32 SeedVID, double Distance)
+		{
+			FVector3d NA = FaceGraph.GetNormal(ToVID);
+			FVector3d NB = FaceGraph.GetNormal(SeedVID);
+			double Dot = NA.Dot(NB);
+			if (WeightingCoeffs.X > 0.001)
+			{
+				Dot = FMathd::Pow(Dot, WeightingCoeffs.X);
+			}
+			double W = FMathd::Clamp(1.0 - Dot * Dot, 0.0, 1.0);
+			W = W * W * W;
+			double Weight = FMathd::Clamp(W, 0.001, 1.0);
+			return Weight * Distance;
+		};
+	}
+
+	TArray<TMeshDijkstra<FMeshFaceDualGraph>::FSeedPoint> SuperPixelSeeds;
+	for (int32 k = 0; k < SeedIndices.Num(); ++k)
+	{
+		SuperPixelSeeds.Add(TMeshDijkstra<FMeshFaceDualGraph>::FSeedPoint{ k, SeedIndices[k], 0.0 });
+	}
+	SuperPixels.ComputeToMaxDistance(SuperPixelSeeds, TNumericLimits<double>::Max());
+
+	TArray<TArray<int32>> TriSets;
+	TriSets.SetNum(SeedIndices.Num());
+	TArray<int32> FailedSet;
+
+	for (int32 tid : Mesh->TriangleIndicesItr())
+	{
+		int32 SeedID = SuperPixels.GetSeedExternalIDForPointSetID(tid);
+		if (SeedID >= 0)
+		{
+			TriSets[SeedID].Add(tid);
+		}
+		else
+		{
+			FailedSet.Add(tid);
+		}
+	}
+
+
+	for (int32 k = 0; k < NumPoints; ++k)
+	{
+		if (TriSets[k].Num() > 0)
+		{
+			FoundPolygons.Add(MoveTemp(TriSets[k]));
+		}
+	}
+	
+	if (FailedSet.Num() > 0)
+	{
+		FMeshConnectedComponents FailedComponents(Mesh);
+		FailedComponents.FindConnectedTriangles(FailedSet);
+		for (FMeshConnectedComponents::FComponent& Component : FailedComponents)
+		{
+			FoundPolygons.Add(Component.Indices);
+		}
+	}
+
+	PostProcessPolygons(true);
+	SetGroupsFromPolygons();
+
+	return (FoundPolygons.Num() > 0);
+}
+
+
+
+
+void FFindPolygonsAlgorithm::PostProcessPolygons(bool bApplyMerging)
+{
+	if (bApplyMerging && MinGroupSize > 1)
+	{
+		OptimizePolygons();
+	}
+}
+
+
+void FFindPolygonsAlgorithm::OptimizePolygons()
+{
+	FMeshRegionGraph RegionGraph;
+	RegionGraph.BuildFromTriangleSets(*Mesh, FoundPolygons, [&](int32 SetIdx) { return SetIdx; });
+	bool bMerged = RegionGraph.MergeSmallRegions(MinGroupSize-1, [&](int32 A, int32 B) { return RegionGraph.GetRegionTriCount(A) > RegionGraph.GetRegionTriCount(B); });
+	bool bSwapped = RegionGraph.OptimizeBorders();
+	if (bMerged || bSwapped)
+	{
+		FoundPolygons.Reset();
+
+		int32 N = RegionGraph.MaxRegionIndex();
+		for (int32 k = 0; k < N; ++k)
+		{
+			if (RegionGraph.IsRegion(k))
+			{
+				const TArray<int32>& Tris = RegionGraph.GetRegionTris(k);
+				FoundPolygons.Add(Tris);
+			}
+		}
+	}
+}
+
+
+
+
+void FFindPolygonsAlgorithm::SetGroupsFromPolygons()
+{
+	Mesh->EnableTriangleGroups(0);
+
+	// set groups from polygons
+	int NumPolygons = FoundPolygons.Num();
+	PolygonTags.SetNum(NumPolygons);
+	PolygonNormals.SetNum(NumPolygons);
+	// can be parallel for
+	for (int PolyIdx = 0; PolyIdx < NumPolygons; PolyIdx++)
+	{
+		const TArray<int>& Polygon = FoundPolygons[PolyIdx];
+		FVector3d AccumNormal(0, 0, 0);
+		int NumTriangles = Polygon.Num();
+		for (int k = 0; k < NumTriangles; ++k)
+		{
+			Mesh->SetTriangleGroup(Polygon[k], (PolyIdx + 1));
+			AccumNormal += Mesh->GetTriArea(k) * Mesh->GetTriNormal(Polygon[k]);
+		}
+		PolygonTags[PolyIdx] = (PolyIdx + 1);
+
+		// find a normal if the average failed
+		UE::Geometry::Normalize(AccumNormal);
+		int SubIdx = 0;
+		while (AccumNormal.Length() < 0.9 && SubIdx < NumTriangles)
+		{
+			AccumNormal = Mesh->GetTriNormal(Polygon[SubIdx++]);
+		}
+		if (AccumNormal.Length() < 0.9)
+		{
+			AccumNormal = FVector3d::UnitY();
+		}
+
+		PolygonNormals[PolyIdx] = AccumNormal;
+	}
+}
+
+
+
+
+
+bool FFindPolygonsAlgorithm::FindPolygonEdges()
+{
+	for (int eid : Mesh->EdgeIndicesItr())
+	{
+		if (Mesh->IsGroupBoundaryEdge(eid))
+		{
+			PolygonEdges.Add(eid);
+		}
+	}
+	return (PolygonEdges.Num() > 0);
+}
