MeshModelingToolset: Split plugin into non-experimental and experimental

diff --git a/Engine/Plugins/Experimental/MeshModelingToolsetExp/Source/MeshModelingToolsExp/Private/MeshVertexSculptTool.cpp b/Engine/Plugins/Experimental/MeshModelingToolsetExp/Source/MeshModelingToolsExp/Private/MeshVertexSculptTool.cpp
new file mode 100644
--- /dev/null
+++ b/Engine/Plugins/Experimental/MeshModelingToolsetExp/Source/MeshModelingToolsExp/Private/MeshVertexSculptTool.cpp
@@ -0,0 +1,1105 @@
+// Copyright Epic Games, Inc. All Rights Reserved.
+
+#include "MeshVertexSculptTool.h"
+#include "InteractiveToolManager.h"
+#include "InteractiveGizmoManager.h"
+#include "Async/ParallelFor.h"
+#include "Async/Async.h"
+#include "Selections/MeshConnectedComponents.h"
+
+#include "MeshWeights.h"
+#include "DynamicMesh/MeshNormals.h"
+#include "DynamicMesh/MeshIndexUtil.h"
+#include "Util/BufferUtil.h"
+#include "AssetUtils/Texture2DUtil.h"
+
+#include "Changes/MeshVertexChange.h"
+
+#include "Sculpting/KelvinletBrushOp.h"
+#include "Sculpting/MeshSmoothingBrushOps.h"
+#include "Sculpting/MeshInflateBrushOps.h"
+#include "Sculpting/MeshMoveBrushOps.h"
+#include "Sculpting/MeshPlaneBrushOps.h"
+#include "Sculpting/MeshPinchBrushOps.h"
+#include "Sculpting/MeshSculptBrushOps.h"
+#include "Sculpting/StampFalloffs.h"
+#include "Sculpting/MeshSculptUtil.h"
+
+#include "TargetInterfaces/PrimitiveComponentBackedTarget.h"
+
+#include "ExplicitUseGeometryMathTypes.h"		// using UE::Geometry::(math types)
+using namespace UE::Geometry;
+
+#define LOCTEXT_NAMESPACE "UMeshVertexSculptTool"
+
+namespace
+{
+	// probably should be something defined for the whole tool framework...
+#if WITH_EDITOR
+	static EAsyncExecution VertexSculptToolAsyncExecTarget = EAsyncExecution::LargeThreadPool;
+#else
+	static EAsyncExecution VertexSculptToolAsyncExecTarget = EAsyncExecution::ThreadPool;
+#endif
+}
+
+
+/*
+ * ToolBuilder
+ */
+UMeshSurfacePointTool* UMeshVertexSculptToolBuilder::CreateNewTool(const FToolBuilderState& SceneState) const
+{
+	UMeshVertexSculptTool* SculptTool = NewObject<UMeshVertexSculptTool>(SceneState.ToolManager);
+	SculptTool->SetWorld(SceneState.World);
+	return SculptTool;
+}
+
+/*
+ * Tool
+ */
+
+void UMeshVertexSculptTool::Setup()
+{
+	UMeshSculptToolBase::Setup();
+
+	SetToolDisplayName(LOCTEXT("ToolName", "Sculpt"));
+
+	// create dynamic mesh component to use for live preview
+	DynamicMeshComponent = NewObject<UDynamicMeshComponent>(Cast<IPrimitiveComponentBackedTarget>(Target)->GetOwnerActor());
+	InitializeSculptMeshComponent(DynamicMeshComponent);
+
+	// assign materials
+	FComponentMaterialSet MaterialSet;
+	Cast<IMaterialProvider>(Target)->GetMaterialSet(MaterialSet);
+	for (int k = 0; k < MaterialSet.Materials.Num(); ++k)
+	{
+		DynamicMeshComponent->SetMaterial(k, MaterialSet.Materials[k]);
+	}
+
+	DynamicMeshComponent->SetInvalidateProxyOnChangeEnabled(false);
+	OnDynamicMeshComponentChangedHandle = DynamicMeshComponent->OnMeshVerticesChanged.AddUObject(this, &UMeshVertexSculptTool::OnDynamicMeshComponentChanged);
+
+	FDynamicMesh3* SculptMesh = GetSculptMesh();
+	FAxisAlignedBox3d Bounds = SculptMesh->GetBounds(true);
+
+	// initialize dynamic octree
+	TFuture<void> InitializeOctree = Async(VertexSculptToolAsyncExecTarget, [SculptMesh, Bounds, this]()
+	{
+		if (SculptMesh->TriangleCount() > 100000)
+		{
+			Octree.RootDimension = Bounds.MaxDim() / 10.0;
+			Octree.SetMaxTreeDepth(4);
+		}
+		else
+		{
+			Octree.RootDimension = Bounds.MaxDim() / 2.0;
+			Octree.SetMaxTreeDepth(8);
+		}
+		Octree.Initialize(SculptMesh);
+		//Octree.CheckValidity(EValidityCheckFailMode::Check, true, true);
+		//FDynamicMeshOctree3::FStatistics Stats;
+		//Octree.ComputeStatistics(Stats);
+		//UE_LOG(LogTemp, Warning, TEXT("Octree Stats: %s"), *Stats.ToString());
+	});
+
+	// find mesh connected-component index for each triangle
+	TFuture<void> InitializeComponents = Async(VertexSculptToolAsyncExecTarget, [SculptMesh, this]()
+	{
+		TriangleComponentIDs.SetNum(SculptMesh->MaxTriangleID());
+		FMeshConnectedComponents Components(SculptMesh);
+		Components.FindConnectedTriangles();
+		int32 ComponentIdx = 1;
+		for (const FMeshConnectedComponents::FComponent& Component : Components)
+		{
+			for (int32 TriIdx : Component.Indices)
+			{
+				TriangleComponentIDs[TriIdx] = ComponentIdx;
+			}
+			ComponentIdx++;
+		}
+	});
+
+	// currently only supporting default polygroup set
+	TFuture<void> InitializeGroups = Async(VertexSculptToolAsyncExecTarget, [SculptMesh, this]()
+	{
+		ActiveGroupSet = MakeUnique<UE::Geometry::FPolygroupSet>(SculptMesh);
+	});
+
+	// initialize target mesh
+	TFuture<void> InitializeBaseMesh = Async(VertexSculptToolAsyncExecTarget, [this]()
+	{
+		UpdateBaseMesh(nullptr);
+		bTargetDirty = false;
+	});
+
+	// initialize render decomposition
+	TFuture<void> InitializeRenderDecomp = Async(VertexSculptToolAsyncExecTarget, [SculptMesh, &MaterialSet, this]()
+	{
+		TUniquePtr<FMeshRenderDecomposition> Decomp = MakeUnique<FMeshRenderDecomposition>();
+		FMeshRenderDecomposition::BuildChunkedDecomposition(SculptMesh, &MaterialSet, *Decomp);
+		Decomp->BuildAssociations(SculptMesh);
+		//UE_LOG(LogTemp, Warning, TEXT("Decomposition has %d groups"), Decomp->Num());
+		DynamicMeshComponent->SetExternalDecomposition(MoveTemp(Decomp));
+	});
+
+	// Wait for above precomputations to finish before continuing
+	InitializeOctree.Wait();
+	InitializeComponents.Wait();
+	InitializeGroups.Wait();
+	InitializeBaseMesh.Wait();
+	InitializeRenderDecomp.Wait();
+
+	// initialize brush radius range interval, brush properties
+	UMeshSculptToolBase::InitializeBrushSizeRange(Bounds);
+
+	// initialize other properties
+	SculptProperties = NewObject<UVertexBrushSculptProperties>(this);
+	
+	// init state flags flags
+	ActiveVertexChange = nullptr;
+
+	InitializeIndicator();
+
+	// initialize our properties
+	AddToolPropertySource(UMeshSculptToolBase::BrushProperties);
+	UMeshSculptToolBase::BrushProperties->bShowPerBrushProps = false;
+	UMeshSculptToolBase::BrushProperties->bShowFalloff = false;
+	SculptProperties->RestoreProperties(this);
+	AddToolPropertySource(SculptProperties);
+	CalculateBrushRadius();
+
+	AlphaProperties = NewObject<UVertexBrushAlphaProperties>(this);
+	AlphaProperties->RestoreProperties(this);
+	AddToolPropertySource(AlphaProperties);
+
+	this->BaseMeshQueryFunc = [&](int32 VertexID, const FVector3d& Position, double MaxDist, FVector3d& PosOut, FVector3d& NormalOut)
+	{
+		return GetBaseMeshNearest(VertexID, Position, MaxDist, PosOut, NormalOut);
+	};
+
+	RegisterBrushType((int32)EMeshVertexSculptBrushType::Smooth,
+		MakeUnique<TBasicMeshSculptBrushOpFactory<FSmoothBrushOp>>(),
+		NewObject<USmoothBrushOpProps>(this));
+
+	RegisterBrushType((int32)EMeshVertexSculptBrushType::SmoothFill,
+		MakeUnique<TBasicMeshSculptBrushOpFactory<FSmoothFillBrushOp>>(),
+		NewObject<USmoothFillBrushOpProps>(this));
+
+	RegisterBrushType((int32)EMeshVertexSculptBrushType::Move,
+		MakeUnique<TBasicMeshSculptBrushOpFactory<FMoveBrushOp>>(),
+		NewObject<UMoveBrushOpProps>(this));
+
+	RegisterBrushType((int32)EMeshVertexSculptBrushType::Offset,
+		MakeUnique<FLambdaMeshSculptBrushOpFactory>([this]() { return MakeUnique<FSurfaceSculptBrushOp>(BaseMeshQueryFunc); }),
+		NewObject<UStandardSculptBrushOpProps>(this));
+
+	RegisterBrushType((int32)EMeshVertexSculptBrushType::SculptView,
+		MakeUnique<FLambdaMeshSculptBrushOpFactory>( [this]() { return MakeUnique<FViewAlignedSculptBrushOp>(BaseMeshQueryFunc); } ),
+		NewObject<UViewAlignedSculptBrushOpProps>(this));
+
+	RegisterBrushType((int32)EMeshVertexSculptBrushType::SculptMax,
+		MakeUnique<FLambdaMeshSculptBrushOpFactory>([this]() { return MakeUnique<FSurfaceMaxSculptBrushOp>(BaseMeshQueryFunc); }),
+		NewObject<USculptMaxBrushOpProps>(this));
+
+	RegisterBrushType((int32)EMeshVertexSculptBrushType::Inflate,
+		MakeUnique<TBasicMeshSculptBrushOpFactory<FInflateBrushOp>>(),
+		NewObject<UInflateBrushOpProps>(this));
+
+	RegisterBrushType((int32)EMeshVertexSculptBrushType::Pinch,
+		MakeUnique<TBasicMeshSculptBrushOpFactory<FPinchBrushOp>>(),
+		NewObject<UPinchBrushOpProps>(this));
+
+	RegisterBrushType((int32)EMeshVertexSculptBrushType::Flatten,
+		MakeUnique<TBasicMeshSculptBrushOpFactory<FFlattenBrushOp>>(),
+		NewObject<UFlattenBrushOpProps>(this));
+
+	RegisterBrushType((int32)EMeshVertexSculptBrushType::Plane,
+		MakeUnique<TBasicMeshSculptBrushOpFactory<FPlaneBrushOp>>(),
+		NewObject<UPlaneBrushOpProps>(this));
+
+	RegisterBrushType((int32)EMeshVertexSculptBrushType::PlaneViewAligned,
+		MakeUnique<TBasicMeshSculptBrushOpFactory<FPlaneBrushOp>>(),
+		NewObject<UViewAlignedPlaneBrushOpProps>(this));
+
+	RegisterBrushType((int32)EMeshVertexSculptBrushType::FixedPlane,
+		MakeUnique<TBasicMeshSculptBrushOpFactory<FPlaneBrushOp>>(),
+		NewObject<UFixedPlaneBrushOpProps>(this));
+
+	RegisterBrushType((int32)EMeshVertexSculptBrushType::ScaleKelvin ,
+		MakeUnique<TBasicMeshSculptBrushOpFactory<FScaleKelvinletBrushOp>>(),
+		NewObject<UScaleKelvinletBrushOpProps>(this));
+
+	RegisterBrushType((int32)EMeshVertexSculptBrushType::PullKelvin,
+		MakeUnique<TBasicMeshSculptBrushOpFactory<FPullKelvinletBrushOp>>(),
+		NewObject<UPullKelvinletBrushOpProps>(this));
+
+	RegisterBrushType((int32)EMeshVertexSculptBrushType::PullSharpKelvin,
+		MakeUnique<TBasicMeshSculptBrushOpFactory<FSharpPullKelvinletBrushOp>>(),
+		NewObject<USharpPullKelvinletBrushOpProps>(this));
+
+	RegisterBrushType((int32)EMeshVertexSculptBrushType::TwistKelvin,
+		MakeUnique<TBasicMeshSculptBrushOpFactory<FTwistKelvinletBrushOp>>(),
+		NewObject<UTwistKelvinletBrushOpProps>(this));
+
+	// secondary brushes
+	RegisterSecondaryBrushType((int32)EMeshVertexSculptBrushType::Smooth,
+		MakeUnique<TBasicMeshSculptBrushOpFactory<FSmoothBrushOp>>(),
+		NewObject<USecondarySmoothBrushOpProps>(this));
+
+
+	AddToolPropertySource(UMeshSculptToolBase::GizmoProperties);
+	SetToolPropertySourceEnabled(UMeshSculptToolBase::GizmoProperties, false);
+
+	AddToolPropertySource(UMeshSculptToolBase::ViewProperties);
+
+	// register watchers
+	SculptProperties->WatchProperty( SculptProperties->PrimaryBrushType,
+		[this](EMeshVertexSculptBrushType NewType) { UpdateBrushType(NewType); });
+
+	SculptProperties->WatchProperty( SculptProperties->PrimaryFalloffType,
+		[this](EMeshSculptFalloffType NewType) { SetPrimaryFalloffType(NewType); });
+
+	SculptProperties->WatchProperty(AlphaProperties->Alpha,
+		[this](UTexture2D* NewAlpha) { UpdateBrushAlpha(AlphaProperties->Alpha); });
+
+	// must call before updating brush type so that we register all brush properties?
+	UMeshSculptToolBase::OnCompleteSetup();
+
+	UpdateBrushType(SculptProperties->PrimaryBrushType);
+	SetPrimaryFalloffType(SculptProperties->PrimaryFalloffType);
+	UpdateBrushAlpha(AlphaProperties->Alpha);
+	SetActiveSecondaryBrushType((int32)EMeshVertexSculptBrushType::Smooth);
+
+	StampRandomStream = FRandomStream(31337);
+}
+
+void UMeshVertexSculptTool::Shutdown(EToolShutdownType ShutdownType)
+{
+	if (DynamicMeshComponent != nullptr)
+	{
+		DynamicMeshComponent->OnMeshChanged.Remove(OnDynamicMeshComponentChangedHandle);
+	}
+
+	SculptProperties->SaveProperties(this);
+	AlphaProperties->SaveProperties(this);
+
+	// this call will commit result, unregister and destroy DynamicMeshComponent
+	UMeshSculptToolBase::Shutdown(ShutdownType);
+}
+
+
+void UMeshVertexSculptTool::OnPropertyModified(UObject* PropertySet, FProperty* Property)
+{
+	CalculateBrushRadius();
+}
+
+
+
+void UMeshVertexSculptTool::OnBeginStroke(const FRay& WorldRay)
+{
+	WaitForPendingUndoRedo();		// cannot start stroke if there is an outstanding undo/redo update
+
+	UpdateBrushPosition(WorldRay);
+
+	if (SculptProperties->PrimaryBrushType == EMeshVertexSculptBrushType::Plane ||
+		SculptProperties->PrimaryBrushType == EMeshVertexSculptBrushType::PlaneViewAligned)
+	{
+		UpdateROI(GetBrushFrameLocal().Origin);
+		UpdateStrokeReferencePlaneForROI(GetBrushFrameLocal(), TriangleROIArray,
+			(SculptProperties->PrimaryBrushType == EMeshVertexSculptBrushType::PlaneViewAligned));
+	}
+	else if (SculptProperties->PrimaryBrushType == EMeshVertexSculptBrushType::FixedPlane)
+	{
+		UpdateStrokeReferencePlaneFromWorkPlane();
+	}
+
+	// initialize first "Last Stamp", so that we can assume all stamps in stroke have a valid previous stamp
+	LastStamp.WorldFrame = GetBrushFrameWorld();
+	LastStamp.LocalFrame = GetBrushFrameLocal();
+	LastStamp.Radius = GetCurrentBrushRadius();
+	LastStamp.Falloff = GetCurrentBrushFalloff();
+	LastStamp.Direction = GetInInvertStroke() ? -1.0 : 1.0;
+	LastStamp.Depth = GetCurrentBrushDepth();
+	LastStamp.Power = GetActivePressure() * GetCurrentBrushStrength();
+	LastStamp.TimeStamp = FDateTime::Now();
+
+	InitialStrokeTriangleID = -1;
+	InitialStrokeTriangleID = GetBrushTriangleID();
+
+	FSculptBrushOptions SculptOptions;
+	//SculptOptions.bPreserveUVFlow = false; // SculptProperties->bPreserveUVFlow;
+	SculptOptions.ConstantReferencePlane = GetCurrentStrokeReferencePlane();
+
+	TUniquePtr<FMeshSculptBrushOp>& UseBrushOp = GetActiveBrushOp();
+	UseBrushOp->ConfigureOptions(SculptOptions);
+	UseBrushOp->BeginStroke(GetSculptMesh(), LastStamp, VertexROI);
+
+	AccumulatedTriangleROI.Reset();
+
+	// begin change here? or wait for first stamp?
+	BeginChange();
+}
+
+void UMeshVertexSculptTool::OnEndStroke()
+{
+	// update spatial
+	bTargetDirty = true;
+
+	GetActiveBrushOp()->EndStroke(GetSculptMesh(), LastStamp, VertexROI);
+
+	// close change record
+	EndChange();
+}
+
+
+
+void UMeshVertexSculptTool::UpdateROI(const FVector3d& BrushPos)
+{
+	TRACE_CPUPROFILER_EVENT_SCOPE(VtxSculptTool_UpdateROI);
+
+	float RadiusSqr = GetCurrentBrushRadius() * GetCurrentBrushRadius();
+	FAxisAlignedBox3d BrushBox(
+		BrushPos - GetCurrentBrushRadius() * FVector3d::One(),
+		BrushPos + GetCurrentBrushRadius() * FVector3d::One());
+
+	// do a parallel range quer
+	RangeQueryTriBuffer.Reset();
+	FDynamicMesh3* Mesh = GetSculptMesh();
+	{
+		TRACE_CPUPROFILER_EVENT_SCOPE(VtxSculptTool_UpdateROI_RangeQuery);
+		Octree.ParallelRangeQuery(BrushBox, RangeQueryTriBuffer);
+	}
+
+	int32 ActiveComponentID = -1;
+	int32 ActiveGroupID = -1;
+	if (SculptProperties->BrushFilter == EMeshVertexSculptBrushFilterType::Component)
+	{
+		ActiveComponentID = (InitialStrokeTriangleID >= 0 && InitialStrokeTriangleID <= TriangleComponentIDs.Num()) ?
+			TriangleComponentIDs[InitialStrokeTriangleID] : -1;
+	}
+	else if (SculptProperties->BrushFilter == EMeshVertexSculptBrushFilterType::PolyGroup)
+	{
+		ActiveGroupID = Mesh->IsTriangle(InitialStrokeTriangleID) ? ActiveGroupSet->GetGroup(InitialStrokeTriangleID) : -1;
+	}
+
+#if 1
+	// in this path we use more memory but this lets us do more in parallel
+
+	// Construct array of inside/outside flags for each triangle's vertices. If no
+	// vertices are inside, clear the triangle ID from the range query buffer.
+	// This can be done in parallel and it's cheaper to do repeated distance computations
+	// than to try to do it inside the ROI building below (todo: profile this some more?)
+	TriangleROIInBuf.SetNum(RangeQueryTriBuffer.Num(), false);
+	{
+		TRACE_CPUPROFILER_EVENT_SCOPE(DynamicMeshSculptTool_UpdateROI_TriVerts);
+		ParallelFor(RangeQueryTriBuffer.Num(), [&](int k)
+		{
+			// check various triangle ROI filters
+			int32 tid = RangeQueryTriBuffer[k];
+			bool bDiscardTriangle = false;
+			if (ActiveComponentID >= 0 && TriangleComponentIDs[tid] != ActiveComponentID)
+			{
+				bDiscardTriangle = true;
+			}
+			if (ActiveGroupID >= 0 && ActiveGroupSet->GetGroup(tid) != ActiveGroupID)
+			{
+				bDiscardTriangle = true;
+			}
+			if (bDiscardTriangle)
+			{
+				TriangleROIInBuf[k].A = TriangleROIInBuf[k].B = TriangleROIInBuf[k].C = 0;
+				RangeQueryTriBuffer[k] = -1;
+				return;
+			}
+
+			const FIndex3i& TriV = Mesh->GetTriangleRef(tid);
+			TriangleROIInBuf[k].A = (DistanceSquared(BrushPos, Mesh->GetVertexRef(TriV.A)) < RadiusSqr) ? 1 : 0;
+			TriangleROIInBuf[k].B = (DistanceSquared(BrushPos, Mesh->GetVertexRef(TriV.B)) < RadiusSqr) ? 1 : 0;
+			TriangleROIInBuf[k].C = (DistanceSquared(BrushPos, Mesh->GetVertexRef(TriV.C)) < RadiusSqr) ? 1 : 0;
+			if (TriangleROIInBuf[k].A + TriangleROIInBuf[k].B + TriangleROIInBuf[k].C == 0)
+			{
+				RangeQueryTriBuffer[k] = -1;
+			}
+		});
+	}
+
+	// Build up vertex and triangle ROIs from the remaining range-query triangles.
+	{
+		TRACE_CPUPROFILER_EVENT_SCOPE(DynamicMeshSculptTool_UpdateROI_3Collect);
+		VertexROIBuilder.Initialize(Mesh->MaxVertexID());
+		TriangleROIBuilder.Initialize(Mesh->MaxTriangleID());
+		int32 N = RangeQueryTriBuffer.Num();
+		for ( int32 k = 0; k < N; ++k )
+		{
+			int32 tid = RangeQueryTriBuffer[k];
+			if (tid == -1) continue;		// triangle was deleted in previous step
+			const FIndex3i& TriV = Mesh->GetTriangleRef(RangeQueryTriBuffer[k]);
+			const FIndex3i& Inside = TriangleROIInBuf[k];
+			int InsideCount = 0;
+			for (int j = 0; j < 3; ++j)
+			{
+				if (Inside[j])
+				{
+					VertexROIBuilder.Add(TriV[j]);
+					InsideCount++;
+				}
+			}
+			if (InsideCount > 0)
+			{
+				TriangleROIBuilder.Add(tid);
+			}
+		}
+		VertexROIBuilder.SwapValuesWith(VertexROI);
+		TriangleROIBuilder.SwapValuesWith(TriangleROIArray);
+	}
+
+#else
+	// In this path we combine everything into one loop. Does fewer distance checks
+	// but nothing can be done in parallel (would change if ROIBuilders had atomic-try-add)
+
+	// TODO would need to support these, this branch is likely dead though
+	ensure(SculptProperties->BrushFilter == EMeshVertexSculptBrushFilterType::None);
+
+	// collect set of vertices and triangles inside brush sphere, from range query result
+	{
+		TRACE_CPUPROFILER_EVENT_SCOPE(DynamicMeshSculptTool_UpdateROI_2Collect);
+		VertexROIBuilder.Initialize(Mesh->MaxVertexID());
+		TriangleROIBuilder.Initialize(Mesh->MaxTriangleID());
+		for (int32 TriIdx : RangeQueryTriBuffer)
+		{
+			FIndex3i TriV = Mesh->GetTriangle(TriIdx);
+			int InsideCount = 0;
+			for (int j = 0; j < 3; ++j)
+			{
+				if (VertexROIBuilder.Contains(TriV[j]))
+				{
+					InsideCount++;
+				} 
+				else if (BrushPos.DistanceSquared(Mesh->GetVertexRef(TriV[j])) < RadiusSqr)
+				{
+					VertexROIBuilder.Add(TriV[j]);
+					InsideCount++;
+				}
+			}
+			if (InsideCount > 0)
+			{
+				TriangleROIBuilder.Add(tid);
+			}
+		}
+		VertexROIBuilder.SwapValuesWith(VertexROI);
+		TriangleROIBuilder.SwapValuesWith(TriangleROIArray);
+	}
+#endif
+
+	{
+		TRACE_CPUPROFILER_EVENT_SCOPE(DynamicMeshSculptTool_UpdateROI_4ROI);
+		ROIPositionBuffer.SetNum(VertexROI.Num(), false);
+		ROIPrevPositionBuffer.SetNum(VertexROI.Num(), false);
+		ParallelFor(VertexROI.Num(), [&](int i)
+		{
+			ROIPrevPositionBuffer[i] = Mesh->GetVertexRef(VertexROI[i]);
+		});
+	}
+}
+
+bool UMeshVertexSculptTool::UpdateStampPosition(const FRay& WorldRay)
+{
+	TRACE_CPUPROFILER_EVENT_SCOPE(VtxSculptTool_UpdateStampPosition);
+
+	CalculateBrushRadius();
+
+	TUniquePtr<FMeshSculptBrushOp>& UseBrushOp = GetActiveBrushOp();
+
+	ESculptBrushOpTargetType TargetType = UseBrushOp->GetBrushTargetType();
+	switch (TargetType)
+	{
+	case ESculptBrushOpTargetType::SculptMesh:
+		UpdateBrushPositionOnSculptMesh(WorldRay, true);
+		break;
+	case ESculptBrushOpTargetType::TargetMesh:
+		UpdateBrushPositionOnTargetMesh(WorldRay, true);
+		break;
+	case ESculptBrushOpTargetType::ActivePlane:
+		UpdateBrushPositionOnActivePlane(WorldRay);
+		break;
+	}
+
+	if (UseBrushOp->GetAlignStampToView())
+	{
+		AlignBrushToView();
+	}
+
+	CurrentStamp = LastStamp;
+	//CurrentStamp.DeltaTime = FMathd::Min((FDateTime::Now() - LastStamp.TimeStamp).GetTotalSeconds(), 1.0);
+	CurrentStamp.DeltaTime = 0.03;		// 30 fps - using actual time is no good now that we support variable stamps!
+	CurrentStamp.WorldFrame = GetBrushFrameWorld();
+	CurrentStamp.LocalFrame = GetBrushFrameLocal();
+	CurrentStamp.Power = GetActivePressure() * GetCurrentBrushStrength();
+
+	if (bHaveBrushAlpha && (AlphaProperties->RotationAngle != 0 || AlphaProperties->bRandomize))
+	{
+		float UseAngle = AlphaProperties->RotationAngle;
+		if (AlphaProperties->bRandomize)
+		{
+			UseAngle += (StampRandomStream.GetFraction() - 0.5f) * 2.0f * AlphaProperties->RandomRange;
+		}
+
+		// possibly should be done in base brush...
+		CurrentStamp.WorldFrame.Rotate(FQuaterniond(CurrentStamp.WorldFrame.Z(), UseAngle, true));
+		CurrentStamp.LocalFrame.Rotate(FQuaterniond(CurrentStamp.LocalFrame.Z(), UseAngle, true));
+	}
+
+	CurrentStamp.PrevLocalFrame = LastStamp.LocalFrame;
+	CurrentStamp.PrevWorldFrame = LastStamp.WorldFrame;
+
+	FVector3d MoveDelta = CurrentStamp.LocalFrame.Origin - CurrentStamp.PrevLocalFrame.Origin;
+	if (UseBrushOp->IgnoreZeroMovements() && MoveDelta.SquaredLength() < FMathd::ZeroTolerance)
+	{
+		return false;
+	}
+
+	return true;
+}
+
+
+TFuture<void> UMeshVertexSculptTool::ApplyStamp()
+{
+	TRACE_CPUPROFILER_EVENT_SCOPE(VtxSculptTool_ApplyStamp);
+
+	TUniquePtr<FMeshSculptBrushOp>& UseBrushOp = GetActiveBrushOp();
+
+	// compute region plane if necessary. This may currently be expensive?
+	if (UseBrushOp->WantsStampRegionPlane())
+	{
+		CurrentStamp.RegionPlane = ComputeStampRegionPlane(CurrentStamp.LocalFrame, TriangleROIArray, true, false, false);
+	}
+
+	// set up alpha function if we have one
+	if (bHaveBrushAlpha)
+	{
+		CurrentStamp.StampAlphaFunc = [this](const FSculptBrushStamp& Stamp, const FVector3d& Position)
+		{
+			return this->SampleBrushAlpha(Stamp, Position);
+		};
+	}
+
+	// apply the stamp, which computes new positions
+	FDynamicMesh3* Mesh = GetSculptMesh();
+	{
+		TRACE_CPUPROFILER_EVENT_SCOPE(VtxSculptTool_ApplyStamp_Apply);
+		UseBrushOp->ApplyStamp(Mesh, CurrentStamp, VertexROI, ROIPositionBuffer);
+	}
+
+	// can discard alpha now
+	CurrentStamp.StampAlphaFunc = nullptr;
+
+	// once stamp is applied, we can start updating vertex change, which can happen async as we saved all necessary info
+	TFuture<void> SaveVertexFuture;
+	if (ActiveVertexChange != nullptr)
+	{
+		SaveVertexFuture = Async(VertexSculptToolAsyncExecTarget, [this]()
+		{
+			TRACE_CPUPROFILER_EVENT_SCOPE(VtxSculptTool_SyncMeshWithPositionBuffer_UpdateChange);
+			const int32 NumV = ROIPositionBuffer.Num();
+			for (int k = 0; k < NumV; ++k)
+			{
+				int VertIdx = VertexROI[k];
+				ActiveVertexChange->UpdateVertex(VertIdx, ROIPrevPositionBuffer[k], ROIPositionBuffer[k]);
+			}
+		});
+	}
+
+	// now actually update the mesh, which happens 
+	{
+		TRACE_CPUPROFILER_EVENT_SCOPE(VtxSculptTool_ApplyStamp_Sync);
+		const int32 NumV = ROIPositionBuffer.Num();
+		ParallelFor(NumV, [&](int32 k)
+		{
+			int VertIdx = VertexROI[k];
+			const FVector3d& NewPos = ROIPositionBuffer[k];
+			Mesh->SetVertex_NoTimeStampUpdate(VertIdx, NewPos);
+		});
+		Mesh->IncrementTimeStamps(1, true, false);
+	}
+
+	LastStamp = CurrentStamp;
+	LastStamp.TimeStamp = FDateTime::Now();
+
+	// let caller wait for this to finish
+	return SaveVertexFuture;
+}
+
+
+
+
+bool UMeshVertexSculptTool::IsHitTriangleBackFacing(int32 TriangleID, const FDynamicMesh3* QueryMesh)
+{
+	if (TriangleID != IndexConstants::InvalidID)
+	{
+		FViewCameraState StateOut;
+		GetToolManager()->GetContextQueriesAPI()->GetCurrentViewState(StateOut);
+		FVector3d LocalEyePosition(CurTargetTransform.InverseTransformPosition((FVector3d)StateOut.Position));
+
+		FVector3d Normal, Centroid;
+		double Area;
+		QueryMesh->GetTriInfo(TriangleID, Normal, Area, Centroid);
+
+		return (Normal.Dot((Centroid - LocalEyePosition)) >= 0);
+	}
+	return false;
+}
+
+
+int32 UMeshVertexSculptTool::FindHitSculptMeshTriangle(const FRay3d& LocalRay)
+{
+	// need this to finish before we can touch Octree
+	WaitForPendingStampUpdate();
+
+	int32 HitTID = Octree.FindNearestHitObject(LocalRay);
+	if (GetBrushCanHitBackFaces() == false && IsHitTriangleBackFacing(HitTID, GetSculptMesh()))
+	{
+		HitTID = IndexConstants::InvalidID;
+	}
+	return HitTID;
+}
+
+int32 UMeshVertexSculptTool::FindHitTargetMeshTriangle(const FRay3d& LocalRay)
+{
+	int32 HitTID = BaseMeshSpatial.FindNearestHitObject(LocalRay);
+	if (GetBrushCanHitBackFaces() == false && IsHitTriangleBackFacing(HitTID, GetBaseMesh()))
+	{
+		HitTID = IndexConstants::InvalidID;
+	}
+	return HitTID;
+}
+
+
+
+bool UMeshVertexSculptTool::UpdateBrushPosition(const FRay& WorldRay)
+{
+	TUniquePtr<FMeshSculptBrushOp>& UseBrushOp = GetActiveBrushOp();
+
+	bool bHit = false; 
+	ESculptBrushOpTargetType TargetType = UseBrushOp->GetBrushTargetType();
+	switch (TargetType)
+	{
+	case ESculptBrushOpTargetType::SculptMesh:
+		bHit = UpdateBrushPositionOnSculptMesh(WorldRay, false);
+		break;
+	case ESculptBrushOpTargetType::TargetMesh:
+		bHit = UpdateBrushPositionOnTargetMesh(WorldRay, false);
+		break;
+	case ESculptBrushOpTargetType::ActivePlane:
+		//UpdateBrushPositionOnActivePlane(WorldRay);
+		bHit = UpdateBrushPositionOnSculptMesh(WorldRay, false);
+		break;
+	}
+
+	if (bHit && UseBrushOp->GetAlignStampToView())
+	{
+		AlignBrushToView();
+	}
+
+	return bHit;
+}
+
+
+
+
+bool UMeshVertexSculptTool::OnUpdateHover(const FInputDeviceRay& DevicePos)
+{
+	// 4.26 HOTFIX: update LastWorldRay position so that we have it for updating WorkPlane position
+	UMeshSurfacePointTool::LastWorldRay = DevicePos.WorldRay;
+
+	PendingStampType = SculptProperties->PrimaryBrushType;
+	if(ensure(InStroke() == false))
+	{
+		UpdateBrushPosition(DevicePos.WorldRay);
+	}
+	return true;
+}
+
+
+
+void UMeshVertexSculptTool::OnTick(float DeltaTime)
+{
+	UMeshSculptToolBase::OnTick(DeltaTime);
+
+	TRACE_CPUPROFILER_EVENT_SCOPE(VtxSculptTool_Tick);
+
+	// process the undo update
+	if (bUndoUpdatePending)
+	{
+		// wait for updates
+		WaitForPendingUndoRedo();
+
+		// post rendering update
+		DynamicMeshComponent->FastNotifyTriangleVerticesUpdated(AccumulatedTriangleROI,
+			EMeshRenderAttributeFlags::Positions | EMeshRenderAttributeFlags::VertexNormals);
+		GetToolManager()->PostInvalidation();
+
+		// ignore stamp and wait for next tick to do anything else
+		bUndoUpdatePending = false;
+		return;
+	}
+
+	// if user changed to not-frozen, we need to reinitialize the target
+	if (bCachedFreezeTarget != SculptProperties->bFreezeTarget)
+	{
+		UpdateBaseMesh(nullptr);
+		bTargetDirty = false;
+	}
+
+	if (InStroke())
+	{
+		//UE_LOG(LogTemp, Warning, TEXT("dt is %.3f, tick fps %.2f - roi size %d/%d"), DeltaTime, 1.0 / DeltaTime, VertexROI.Num(), TriangleROI.Num());
+		TRACE_CPUPROFILER_EVENT_SCOPE(VtxSculptTool_Tick_StrokeUpdate);
+		FDynamicMesh3* Mesh = GetSculptMesh();
+
+		// update brush position
+		if (UpdateStampPosition(GetPendingStampRayWorld()) == false)
+		{
+			return;
+		}
+		UpdateStampPendingState();
+		if (IsStampPending() == false)
+		{
+			return;
+		}
+
+		// need to make sure previous stamp finished
+		WaitForPendingStampUpdate();
+
+		// update sculpt ROI
+		UpdateROI(CurrentStamp.LocalFrame.Origin);
+
+		// Append updated ROI to modified region (async). For some reason this is very expensive,
+		// maybe because of TSet? but we have a lot of time to do it.
+		TFuture<void> AccumulateROI = Async(VertexSculptToolAsyncExecTarget, [this]()
+		{
+			TRACE_CPUPROFILER_EVENT_SCOPE(VtxSculptTool_Tick_AccumROI);
+			for (int32 tid : TriangleROIArray)
+			{
+				AccumulatedTriangleROI.Add(tid);
+			}
+		});
+
+		// Start precomputing the normals ROI. This is currently the most expensive single thing we do next
+		// to Octree re-insertion, despite it being almost trivial. Why?!?
+		bool bUsingOverlayNormalsOut = false;
+		TFuture<void> NormalsROI = Async(VertexSculptToolAsyncExecTarget, [Mesh, &bUsingOverlayNormalsOut, this]()
+		{
+			TRACE_CPUPROFILER_EVENT_SCOPE(VtxSculptTool_Tick_NormalsROI);
+
+			//UE::SculptUtil::PrecalculateNormalsROI(Mesh, TriangleROIArray, NormalsROIBuilder, bUsingOverlayNormals, false);
+			UE::SculptUtil::PrecalculateNormalsROI(Mesh, TriangleROIArray, NormalsFlags, bUsingOverlayNormalsOut, false);
+		});
+
+		// NOTE: you might try to speculatively do the octree remove here, to save doing it later on Reinsert().
+		// This will not improve things, as Reinsert() checks if it needs to actually re-insert, which avoids many
+		// removes, and does much of the work of Remove anyway.
+
+		// Apply the stamp. This will return a future that is updating the vertex-change record, 
+		// which can run until the end of the frame, as it is using cached information
+		TFuture<void> UpdateChangeFuture = ApplyStamp();
+
+		// begin octree rebuild calculation
+		StampUpdateOctreeFuture = Async(VertexSculptToolAsyncExecTarget, [this]()
+		{
+			TRACE_CPUPROFILER_EVENT_SCOPE(VtxSculptTool_Tick_OctreeReinsert);
+			Octree.ReinsertTrianglesParallel(TriangleROIArray, OctreeUpdateTempBuffer, OctreeUpdateTempFlagBuffer);
+		});
+		bStampUpdatePending = true;
+		//TFuture<void> OctreeRebuild = Async(VertexSculptToolAsyncExecTarget, [&]()
+		//{
+		//	TRACE_CPUPROFILER_EVENT_SCOPE(VtxSculptTool_Tick_OctreeReinsert);
+		//	Octree.ReinsertTriangles(TriangleROIArray);
+		//});
+
+		// TODO: first step of RecalculateROINormals() is to convert TriangleROI into vertex or element ROI.
+		// We can do this while we are computing stamp!
+
+		// precompute dynamic mesh update info
+		TArray<int32> RenderUpdateSets; FAxisAlignedBox3d RenderUpdateBounds;
+		TFuture<bool> RenderUpdatePrecompute = DynamicMeshComponent->FastNotifyTriangleVerticesUpdated_TryPrecompute(
+				TriangleROIArray, RenderUpdateSets, RenderUpdateBounds);
+
+		// recalculate normals. This has to complete before we can update component
+		// (in fact we could do it per-chunk...)
+		{
+			TRACE_CPUPROFILER_EVENT_SCOPE(VtxSculptTool_Tick_RecalcNormals);
+			NormalsROI.Wait();
+			UE::SculptUtil::RecalculateROINormals(Mesh, NormalsFlags, bUsingOverlayNormalsOut);
+			//UE::SculptUtil::RecalculateROINormals(Mesh, NormalsROIBuilder.Indices(), bUsingOverlayNormals);
+		}
+
+		{
+			TRACE_CPUPROFILER_EVENT_SCOPE(VtxSculptTool_Tick_UpdateMesh);
+			RenderUpdatePrecompute.Wait();
+			DynamicMeshComponent->FastNotifyTriangleVerticesUpdated_ApplyPrecompute(TriangleROIArray,
+				EMeshRenderAttributeFlags::Positions | EMeshRenderAttributeFlags::VertexNormals,
+				RenderUpdatePrecompute, RenderUpdateSets, RenderUpdateBounds);
+
+			GetToolManager()->PostInvalidation();
+		}
+
+		// we don't really need to wait for these to happen to end Tick()...
+		UpdateChangeFuture.Wait();
+		AccumulateROI.Wait();
+	} 
+	else if (bTargetDirty)
+	{
+		TRACE_CPUPROFILER_EVENT_SCOPE(VtxSculptTool_Tick_UpdateTarget);
+		check(InStroke() == false);
+
+		// this spawns futures that we could allow to run while other things happen...
+		UpdateBaseMesh(&AccumulatedTriangleROI);
+		AccumulatedTriangleROI.Reset();
+
+		bTargetDirty = false;
+	}
+
+}
+
+
+
+void UMeshVertexSculptTool::WaitForPendingStampUpdate()
+{
+	if (bStampUpdatePending)
+	{
+		StampUpdateOctreeFuture.Wait();
+		bStampUpdatePending = true;
+	}
+}
+
+
+
+void UMeshVertexSculptTool::UpdateBaseMesh(const TSet<int32>* TriangleSet)
+{
+	if (SculptProperties != nullptr)
+	{
+		bCachedFreezeTarget = SculptProperties->bFreezeTarget;
+		if (SculptProperties->bFreezeTarget)
+		{
+			return;   // do not update frozen target
+		}
+	}
+
+	const FDynamicMesh3* SculptMesh = GetSculptMesh();
+	if ( ! TriangleSet )
+	{
+		TRACE_CPUPROFILER_EVENT_SCOPE(VtxSculptTool_Target_FullUpdate);
+		BaseMesh.Copy(*SculptMesh, false, false, false, false);
+		BaseMesh.EnableVertexNormals(FVector3f::UnitZ());
+		FMeshNormals::QuickComputeVertexNormals(BaseMesh);
+		BaseMeshSpatial.SetMaxTreeDepth(8);
+		BaseMeshSpatial = FDynamicMeshOctree3();   // need to clear...
+		BaseMeshSpatial.Initialize(&BaseMesh);
+	}
+	else
+	{
+		BaseMeshIndexBuffer.Reset();
+		for ( int32 tid : *TriangleSet)
+		{ 
+			FIndex3i Tri = BaseMesh.GetTriangle(tid);
+			BaseMesh.SetVertex(Tri.A, SculptMesh->GetVertex(Tri.A));
+			BaseMesh.SetVertex(Tri.B, SculptMesh->GetVertex(Tri.B));
+			BaseMesh.SetVertex(Tri.C, SculptMesh->GetVertex(Tri.C));
+			BaseMeshIndexBuffer.Add(tid);
+		}
+		auto UpdateBaseNormals = Async(VertexSculptToolAsyncExecTarget, [this]()
+		{
+			TRACE_CPUPROFILER_EVENT_SCOPE(VtxSculptTool_Target_UpdateBaseNormals);
+			FMeshNormals::QuickComputeVertexNormalsForTriangles(BaseMesh, BaseMeshIndexBuffer);
+		});
+		auto ReinsertTriangles = Async(VertexSculptToolAsyncExecTarget, [TriangleSet, this]()
+		{
+			TRACE_CPUPROFILER_EVENT_SCOPE(VtxSculptTool_Target_Reinsert);
+			BaseMeshSpatial.ReinsertTriangles(*TriangleSet);
+		});
+		UpdateBaseNormals.Wait();
+		ReinsertTriangles.Wait();
+	}
+}
+
+
+bool UMeshVertexSculptTool::GetBaseMeshNearest(int32 VertexID, const FVector3d& Position, double SearchRadius, FVector3d& TargetPosOut, FVector3d& TargetNormalOut)
+{
+	TargetPosOut = BaseMesh.GetVertex(VertexID);
+	TargetNormalOut = (FVector3d)BaseMesh.GetVertexNormal(VertexID);
+	return true;
+}
+
+
+
+
+
+
+void UMeshVertexSculptTool::IncreaseBrushSpeedAction()
+{
+	//SculptProperties->PrimaryBrushSpeed = FMath::Clamp(SculptProperties->PrimaryBrushSpeed + 0.05f, 0.0f, 1.0f);
+}
+
+void UMeshVertexSculptTool::DecreaseBrushSpeedAction()
+{
+	//SculptProperties->PrimaryBrushSpeed = FMath::Clamp(SculptProperties->PrimaryBrushSpeed - 0.05f, 0.0f, 1.0f);
+}
+
+void UMeshVertexSculptTool::UpdateBrushAlpha(UTexture2D* NewAlpha)
+{
+	if (this->BrushAlpha != NewAlpha)
+	{
+		this->BrushAlpha = NewAlpha;
+		if (this->BrushAlpha != nullptr)
+		{
+			TImageBuilder<FVector4f> AlphaValues;
+
+			constexpr bool bPreferPlatformData = false;
+			const bool bReadOK = UE::AssetUtils::ReadTexture(this->BrushAlpha, AlphaValues, bPreferPlatformData);
+			if (bReadOK)
+			{
+				BrushAlphaValues = MoveTemp(AlphaValues);
+				BrushAlphaDimensions = AlphaValues.GetDimensions();
+				bHaveBrushAlpha = true;
+				return;
+			}
+		}
+		bHaveBrushAlpha = false;
+		BrushAlphaValues = TImageBuilder<FVector4f>();
+		BrushAlphaDimensions = FImageDimensions();
+	}
+}
+
+
+double UMeshVertexSculptTool::SampleBrushAlpha(const FSculptBrushStamp& Stamp, const FVector3d& Position) const
+{
+	if (! bHaveBrushAlpha) return 1.0;
+
+	static const FVector4f InvalidValue(0, 0, 0, 0);
+
+	FVector2d AlphaUV = Stamp.LocalFrame.ToPlaneUV(Position, 2);
+	double u = AlphaUV.X / Stamp.Radius;
+	u = 1.0 - (u + 1.0) / 2.0;
+	double v = AlphaUV.Y / Stamp.Radius;
+	v = 1.0 - (v + 1.0) / 2.0;
+	if (u < 0 || u > 1) return 0.0;
+	if (v < 0 || v > 1) return 0.0;
+	FVector4f AlphaValue = BrushAlphaValues.BilinearSampleUV<float>(FVector2d(u, v), InvalidValue);
+	return FMathd::Clamp(AlphaValue.X, 0.0, 1.0);
+}
+
+
+
+//
+// Change Tracking
+//
+void UMeshVertexSculptTool::BeginChange()
+{
+	check(ActiveVertexChange == nullptr);
+	ActiveVertexChange = new FMeshVertexChangeBuilder();
+}
+
+void UMeshVertexSculptTool::EndChange()
+{
+	check(ActiveVertexChange);
+
+	TUniquePtr<TWrappedToolCommandChange<FMeshVertexChange>> NewChange = MakeUnique<TWrappedToolCommandChange<FMeshVertexChange>>();
+	NewChange->WrappedChange = MoveTemp(ActiveVertexChange->Change);
+	NewChange->BeforeModify = [this](bool bRevert)
+	{
+		this->WaitForPendingUndoRedo();
+	};
+
+	GetToolManager()->EmitObjectChange(DynamicMeshComponent, MoveTemp(NewChange), LOCTEXT("VertexSculptChange", "Brush Stroke"));
+
+	delete ActiveVertexChange;
+	ActiveVertexChange = nullptr;
+}
+
+
+void UMeshVertexSculptTool::WaitForPendingUndoRedo()
+{
+	if (bUndoUpdatePending)
+	{
+		UndoNormalsFuture.Wait();
+		UndoUpdateOctreeFuture.Wait();
+		UndoUpdateBaseMeshFuture.Wait();
+		bUndoUpdatePending = false;
+	}
+}
+
+void UMeshVertexSculptTool::OnDynamicMeshComponentChanged(UDynamicMeshComponent* Component, const FMeshVertexChange* Change, bool bRevert)
+{
+	// have to wait for any outstanding stamp update to finish...
+	WaitForPendingStampUpdate();
+	// wait for previous Undo to finish (possibly never hit because the change records do it?)
+	WaitForPendingUndoRedo();
+
+	FDynamicMesh3* Mesh = GetSculptMesh();
+
+	// figure out the set of modified triangles
+	AccumulatedTriangleROI.Reset();
+	UE::Geometry::VertexToTriangleOneRing(Mesh, Change->Vertices, AccumulatedTriangleROI);
+
+	// start the normal recomputation
+	UndoNormalsFuture = Async(VertexSculptToolAsyncExecTarget, [this, Mesh]()
+	{
+		UE::SculptUtil::RecalculateROINormals(Mesh, AccumulatedTriangleROI, NormalsROIBuilder);
+		return true;
+	});
+
+	// start the octree update
+	UndoUpdateOctreeFuture = Async(VertexSculptToolAsyncExecTarget, [this, Mesh]()
+	{
+		Octree.ReinsertTriangles(AccumulatedTriangleROI);
+		return true;
+	});
+
+	// start the base mesh update
+	UndoUpdateBaseMeshFuture = Async(VertexSculptToolAsyncExecTarget, [this, Mesh]()
+	{
+		UpdateBaseMesh(&AccumulatedTriangleROI);
+		return true;
+	});
+
+	// note that we have a pending update
+	bUndoUpdatePending = true;
+}
+
+
+
+
+
+
+void UMeshVertexSculptTool::UpdateBrushType(EMeshVertexSculptBrushType BrushType)
+{
+	static const FText BaseMessage = LOCTEXT("OnStartSculptTool", "Hold Shift to Smooth, Ctrl to Invert (where applicable). [/] and S/D change Size (+Shift to small-step), W/E changes Strength.");
+	FTextBuilder Builder;
+	Builder.AppendLine(BaseMessage);
+
+	SetActivePrimaryBrushType((int32)BrushType);
+
+	SetToolPropertySourceEnabled(GizmoProperties, false);
+	if (BrushType == EMeshVertexSculptBrushType::FixedPlane)
+	{
+		Builder.AppendLine(LOCTEXT("FixedPlaneTip", "Use T to reposition Work Plane at cursor, Shift+T to align to Normal, Ctrl+Shift+T to align to View"));
+		SetToolPropertySourceEnabled(GizmoProperties, true);
+	}
+
+
+	bool bEnableAlpha = false;
+	switch (BrushType)
+	{
+	case EMeshVertexSculptBrushType::Offset:
+	case EMeshVertexSculptBrushType::SculptView:
+	case EMeshVertexSculptBrushType::SculptMax:
+		bEnableAlpha = true; break;
+	}
+	SetToolPropertySourceEnabled(AlphaProperties, bEnableAlpha);
+
+
+	GetToolManager()->DisplayMessage(Builder.ToText(), EToolMessageLevel::UserNotification);
+}
+
+
+
+#undef LOCTEXT_NAMESPACE
diff --git a/Engine/Plugins/Experimental/MeshModelingToolsetExp/Source/MeshModelingToolsExp/Private/MirrorTool.cpp b/Engine/Plugins/Experimental/MeshModelingToolsetExp/Source/MeshModelingToolsExp/Private/MirrorTool.cpp
new file mode 100644
--- /dev/null
+++ b/Engine/Plugins/Experimental/MeshModelingToolsetExp/Source/MeshModelingToolsExp/Private/MirrorTool.cpp
@@ -0,0 +1,558 @@
+// Copyright Epic Games, Inc. All Rights Reserved.
+
+#include "MirrorTool.h"
+
+#include "ModelingObjectsCreationAPI.h"
+#include "BaseBehaviors/KeyAsModifierInputBehavior.h"
+#include "BaseBehaviors/SingleClickBehavior.h"
+#include "CompositionOps/MirrorOp.h"
+#include "Drawing/MeshDebugDrawing.h"
+#include "DynamicMeshToMeshDescription.h"
+#include "InteractiveToolManager.h"
+#include "MeshDescriptionToDynamicMesh.h"
+#include "Misc/MessageDialog.h"
+#include "ToolBuilderUtil.h"
+#include "ToolSetupUtil.h"
+
+#include "TargetInterfaces/MaterialProvider.h"
+#include "TargetInterfaces/MeshDescriptionCommitter.h"
+#include "TargetInterfaces/MeshDescriptionProvider.h"
+#include "TargetInterfaces/PrimitiveComponentBackedTarget.h"
+#include "ToolTargetManager.h"
+
+#include "ExplicitUseGeometryMathTypes.h"		// using UE::Geometry::(math types)
+using namespace UE::Geometry;
+
+#define LOCTEXT_NAMESPACE "UMirrorTool"
+
+
+//Tool builder functions
+
+const FToolTargetTypeRequirements& UMirrorToolBuilder::GetTargetRequirements() const
+{
+	static FToolTargetTypeRequirements TypeRequirements({
+		UMaterialProvider::StaticClass(),
+		UMeshDescriptionCommitter::StaticClass(),
+		UMeshDescriptionProvider::StaticClass(),
+		UPrimitiveComponentBackedTarget::StaticClass()
+		});
+	return TypeRequirements;
+}
+
+bool UMirrorToolBuilder::CanBuildTool(const FToolBuilderState& SceneState) const
+{
+	return SceneState.TargetManager->CountSelectedAndTargetable(SceneState, GetTargetRequirements()) > 0;
+}
+
+UInteractiveTool* UMirrorToolBuilder::BuildTool(const FToolBuilderState& SceneState) const
+{
+	UMirrorTool* NewTool = NewObject<UMirrorTool>(SceneState.ToolManager);
+
+	TArray<TObjectPtr<UToolTarget>> Targets = SceneState.TargetManager->BuildAllSelectedTargetable(SceneState, GetTargetRequirements());
+	NewTool->SetTargets(MoveTemp(Targets));
+	NewTool->SetWorld(SceneState.World);
+
+	return NewTool;
+}
+
+
+// Operator factory
+
+TUniquePtr<FDynamicMeshOperator> UMirrorOperatorFactory::MakeNewOperator()
+{
+	TUniquePtr<FMirrorOp> MirrorOp = MakeUnique<FMirrorOp>();
+
+	// Set up inputs and settings
+	MirrorOp->OriginalMesh = MirrorTool->MeshesToMirror[ComponentIndex]->GetMesh();
+	MirrorOp->bAppendToOriginal = MirrorTool->Settings->OperationMode == EMirrorOperationMode::MirrorAndAppend;
+	MirrorOp->bCropFirst = MirrorTool->Settings->bCropAlongMirrorPlaneFirst;
+	MirrorOp->bWeldAlongPlane = MirrorTool->Settings->bWeldVerticesOnMirrorPlane;
+	MirrorOp->PlaneTolerance = MirrorTool->Settings->PlaneTolerance;
+	MirrorOp->bAllowBowtieVertexCreation = MirrorTool->Settings->bAllowBowtieVertexCreation;
+
+	FTransform LocalToWorld = MirrorTool->TargetComponentInterface(ComponentIndex)->GetWorldTransform();
+	MirrorOp->SetTransform(LocalToWorld);
+
+	// We also need WorldToLocal. Threshold the LocalToWorld scaling transform so we can get the inverse.
+	FVector LocalToWorldScale = LocalToWorld.GetScale3D();
+	for (int i = 0; i < 3; i++)
+	{
+		float DimScale = FMathf::Abs(LocalToWorldScale[i]);
+		float Tolerance = KINDA_SMALL_NUMBER;
+		if (DimScale < Tolerance)
+		{
+			LocalToWorldScale[i] = Tolerance * FMathf::SignNonZero(LocalToWorldScale[i]);
+		}
+	}
+	LocalToWorld.SetScale3D(LocalToWorldScale);
+	UE::Geometry::FTransform3d WorldToLocal = UE::Geometry::FTransform3d(LocalToWorld).Inverse();
+
+	// Now we can get the plane parameters in local space.
+	MirrorOp->LocalPlaneOrigin = WorldToLocal.TransformPosition(MirrorTool->MirrorPlaneOrigin);;
+
+	FVector3d WorldNormal = MirrorTool->MirrorPlaneNormal;
+	MirrorOp->LocalPlaneNormal = WorldToLocal.TransformNormal(MirrorTool->MirrorPlaneNormal);
+
+	return MirrorOp;
+}
+
+
+// Tool property functions
+
+void UMirrorToolActionPropertySet::PostAction(EMirrorToolAction Action)
+{
+	if (ParentTool.IsValid())
+	{
+		ParentTool->RequestAction(Action);
+	}
+}
+
+
+// Tool itself
+UMirrorTool::UMirrorTool()
+{
+}
+
+bool UMirrorTool::CanAccept() const
+{
+	for (UMeshOpPreviewWithBackgroundCompute* Preview : Previews)
+	{
+		if (!Preview->HaveValidResult())
+		{
+			return false;
+		}
+	}
+	return Super::CanAccept();
+}
+
+void UMirrorTool::SetWorld(UWorld* World)
+{
+	TargetWorld = World;
+}
+
+void UMirrorTool::OnPropertyModified(UObject* PropertySet, FProperty* Property)
+{
+	// Editing the "show preview" option changes whether we need to be displaying the preview or the original mesh.
+	if (Property && (Property->GetFName() == GET_MEMBER_NAME_CHECKED(UMirrorToolProperties, bShowPreview)))
+	{
+		for (int32 ComponentIdx = 0; ComponentIdx < Targets.Num(); ComponentIdx++)
+		{
+			TargetComponentInterface(ComponentIdx)->SetOwnerVisibility(!Settings->bShowPreview);
+		}
+		for (UMeshOpPreviewWithBackgroundCompute* Preview : Previews)
+		{
+			Preview->SetVisibility(Settings->bShowPreview);
+		}
+	}
+
+	// Regardless of what changed, update the previews.
+	for (UMeshOpPreviewWithBackgroundCompute* Preview : Previews)
+	{
+		Preview->InvalidateResult();
+	}
+}
+
+void UMirrorTool::OnTick(float DeltaTime)
+{
+	// Deal with any buttons that may have been clicked
+	if (PendingAction != EMirrorToolAction::NoAction)
+	{
+		ApplyAction(PendingAction);
+		PendingAction = EMirrorToolAction::NoAction;
+	}
+
+	if (PlaneMechanic != nullptr)
+	{
+		PlaneMechanic->Tick(DeltaTime);
+	}
+	
+	for (UMeshOpPreviewWithBackgroundCompute* Preview : Previews)
+	{
+		Preview->Tick(DeltaTime);
+	}
+}
+
+void UMirrorTool::Render(IToolsContextRenderAPI* RenderAPI)
+{
+	// Have the plane draw itself.
+	PlaneMechanic->Render(RenderAPI);
+}
+
+void UMirrorTool::Setup()
+{
+	UInteractiveTool::Setup();
+
+	SetToolDisplayName(LOCTEXT("ToolName", "Mirror"));
+	GetToolManager()->DisplayMessage(
+		LOCTEXT("OnStartMirrorTool", "Mirror one or more meshes across a plane. The plane can be set by using the preset buttons, moving the gizmo, or ctrl+clicking on a spot on the original mesh."),
+		EToolMessageLevel::UserNotification);
+
+	// Set up the properties
+	Settings = NewObject<UMirrorToolProperties>(this, TEXT("Mirror Tool Settings"));
+	Settings->RestoreProperties(this);
+	AddToolPropertySource(Settings);
+
+	ToolActions = NewObject<UMirrorToolActionPropertySet>(this);
+	ToolActions->Initialize(this);
+	AddToolPropertySource(ToolActions);
+
+	CheckAndDisplayWarnings();
+
+	// Fill in the MeshesToMirror array with suitably converted meshes.
+	for (int i = 0; i < Targets.Num(); i++)
+	{
+		IPrimitiveComponentBackedTarget* TargetComponent = TargetComponentInterface(i);
+		IMeshDescriptionProvider* TargetMeshProvider = TargetMeshProviderInterface(i);
+
+		// Convert into dynamic mesh
+		TSharedPtr<FDynamicMesh3, ESPMode::ThreadSafe> DynamicMesh = MakeShared<FDynamicMesh3, ESPMode::ThreadSafe>();
+		FMeshDescriptionToDynamicMesh Converter;
+		Converter.Convert(TargetMeshProvider->GetMeshDescription(), *DynamicMesh);
+
+		// Wrap the dynamic mesh in a replacement change target
+		UDynamicMeshReplacementChangeTarget* WrappedTarget = MeshesToMirror.Add_GetRef(NewObject<UDynamicMeshReplacementChangeTarget>());
+
+		// Set callbacks so previews are invalidated on undo/redo changing the meshes
+		WrappedTarget->SetMesh(DynamicMesh);
+		WrappedTarget->OnMeshChanged.AddLambda([this, i]() { Previews[i]->InvalidateResult(); });
+	}
+
+	// Set the visibility of the StaticMeshComponents depending on whether we are showing them or the preview.
+	for (int32 ComponentIdx = 0; ComponentIdx < Targets.Num(); ComponentIdx++)
+	{
+		TargetComponentInterface(ComponentIdx)->SetOwnerVisibility(!Settings->bShowPreview);
+	}
+
+	// Initialize the PreviewMesh and BackgroundCompute objects
+	SetupPreviews();
+
+	// Update the bounding box of the meshes.
+	CombinedBounds.Init();
+	for (int32 ComponentIdx = 0; ComponentIdx < Targets.Num(); ComponentIdx++)
+	{
+		FVector ComponentOrigin, ComponentExtents;
+		TargetComponentInterface(ComponentIdx)->GetOwnerActor()->GetActorBounds(false, ComponentOrigin, ComponentExtents);
+		CombinedBounds += FBox::BuildAABB(ComponentOrigin, ComponentExtents);
+	}
+
+	// Set the initial mirror plane. We want the plane to start in the middle if we're doing a simple
+	// mirror (i.e., not appending, and not cropping). Otherwise, we want the plane to start to one side.
+	MirrorPlaneOrigin = (FVector3d)CombinedBounds.GetCenter();
+	MirrorPlaneNormal = FVector3d(0, -1, 0);
+	if (Settings->OperationMode == EMirrorOperationMode::MirrorAndAppend || Settings->bCropAlongMirrorPlaneFirst)
+	{
+		MirrorPlaneOrigin.Y = CombinedBounds.Min.Y;
+	}
+
+	// Set up the mirror plane mechanic, which manages the gizmo
+	PlaneMechanic = NewObject<UConstructionPlaneMechanic>(this);
+	PlaneMechanic->Setup(this);
+	PlaneMechanic->Initialize(TargetWorld, FFrame3d(MirrorPlaneOrigin, MirrorPlaneNormal));
+
+	// Have the plane mechanic update things properly
+	PlaneMechanic->OnPlaneChanged.AddLambda([this]() {
+		MirrorPlaneNormal = PlaneMechanic->Plane.Rotation.AxisZ();
+		MirrorPlaneOrigin = PlaneMechanic->Plane.Origin;
+		for (UMeshOpPreviewWithBackgroundCompute* Preview : Previews)
+		{
+			Preview->InvalidateResult();
+		}
+		});
+
+	// Modify the Ctrl+click set plane behavior to respond to our CtrlClickBehavior property
+	PlaneMechanic->SetPlaneCtrlClickBehaviorTarget->OnClickedPositionFunc = [this](const FHitResult& Hit)
+	{
+		bool bIgnoreNormal = (Settings->CtrlClickBehavior == EMirrorCtrlClickBehavior::Reposition);
+		PlaneMechanic->SetDrawPlaneFromWorldPos((FVector3d)Hit.ImpactPoint, (FVector3d)Hit.ImpactNormal, bIgnoreNormal);
+	};
+	// Also include the original components in the ctrl+click hit testing even though we made them 
+	// invisible, since we want to be able to reposition the plane onto the original mesh.
+	for (int32 ComponentIdx = 0; ComponentIdx < Targets.Num(); ComponentIdx++)
+	{
+		IPrimitiveComponentBackedTarget* TargetComponent = TargetComponentInterface(ComponentIdx);
+		PlaneMechanic->SetPlaneCtrlClickBehaviorTarget->InvisibleComponentsToHitTest.Add(TargetComponent->GetOwnerComponent());
+	}
+
+	// Start the preview calculations
+	for (UMeshOpPreviewWithBackgroundCompute* Preview : Previews)
+	{
+		Preview->InvalidateResult();
+	}
+}
+
+void UMirrorTool::SetupPreviews()
+{
+	// Create a preview (with an op) for each selected component.
+	int32 NumMeshes = MeshesToMirror.Num();
+	for (int32 PreviewIndex = 0; PreviewIndex < NumMeshes; ++PreviewIndex)
+	{
+		UMirrorOperatorFactory* MirrorOpCreator = NewObject<UMirrorOperatorFactory>();
+		MirrorOpCreator->MirrorTool = this;
+		MirrorOpCreator->ComponentIndex = PreviewIndex;
+
+		UMeshOpPreviewWithBackgroundCompute* Preview = Previews.Add_GetRef(
+			NewObject<UMeshOpPreviewWithBackgroundCompute>(MirrorOpCreator, "Preview"));
+		Preview->Setup(TargetWorld, MirrorOpCreator);
+		Preview->PreviewMesh->SetTangentsMode(EDynamicMeshComponentTangentsMode::AutoCalculated);
+
+		FComponentMaterialSet MaterialSet;
+		TargetMaterialInterface(PreviewIndex)->GetMaterialSet(MaterialSet);
+		Preview->ConfigureMaterials(MaterialSet.Materials, ToolSetupUtil::GetDefaultWorkingMaterial(GetToolManager()));
+
+		// Set initial preview to unprocessed mesh, so that things don't disappear initially
+		Preview->PreviewMesh->UpdatePreview(MeshesToMirror[PreviewIndex]->GetMesh().Get());
+		Preview->PreviewMesh->SetTransform(TargetComponentInterface(PreviewIndex)->GetWorldTransform());
+		Preview->SetVisibility(Settings->bShowPreview);
+	}
+}
+
+void UMirrorTool::CheckAndDisplayWarnings()
+{
+	// We can have more than one warning, which makes this a bit more work.
+	FText SameSourceWarning;
+	FText NonUniformScaleWarning;
+
+	// See if any of the selected components have the same source.
+	TArray<int32> MapToFirstOccurrences;
+	bool bAnyHaveSameSource = GetMapToSharedSourceData(MapToFirstOccurrences);
+	if (bAnyHaveSameSource)
+	{
+		SameSourceWarning = LOCTEXT("MirrorMultipleAssetsWithSameSource", "WARNING: Multiple meshes in your selection use the same source asset! Only the \"Create New Assets\" save mode is supported.");
+		
+		// We could forcefully set the save mode to CreateNewAssets, but the setting will persist on new invocations
+		// of the tool, which may surprise the user. So, it's up to them to set it.
+	}
+
+	// See if any of the selected components have a nonuniform scaling transform.
+	IPrimitiveComponentBackedTarget* NonUniformScalingTarget = nullptr;
+	for (int32 i = 0; i < Targets.Num(); ++i)
+	{
+		FVector Scaling = TargetComponentInterface(i)->GetWorldTransform().GetScale3D();
+		if (Scaling.X != Scaling.Y || Scaling.Y != Scaling.Z)
+		{
+			NonUniformScalingTarget = TargetComponentInterface(i);
+			break;
+		}
+	}
+
+	if (NonUniformScalingTarget)
+	{
+		NonUniformScaleWarning = FText::Format(
+			LOCTEXT("MirrorNonUniformScaledAsset", "WARNING: The item \"{0}\" has a non-uniform scaling transform. This is not supported because mirroring acts on the underlying mesh, and mirroring is not commutative with non-uniform scaling. Consider deforming the mesh rather than scaling it non-uniformly."),
+			FText::FromString(NonUniformScalingTarget->GetOwnerActor()->GetName()));
+	}
+
+	if (bAnyHaveSameSource && NonUniformScalingTarget)
+	{
+		// Concatenates the two warnings with an extra line in between.
+		GetToolManager()->DisplayMessage(FText::Format(LOCTEXT("CombinedWarnings", "{0}\n\n{1}"),
+			SameSourceWarning, NonUniformScaleWarning), EToolMessageLevel::UserWarning);
+	}
+	else if (bAnyHaveSameSource)
+	{
+		GetToolManager()->DisplayMessage(SameSourceWarning, EToolMessageLevel::UserWarning);
+	}
+	else if (NonUniformScalingTarget)
+	{
+		GetToolManager()->DisplayMessage(NonUniformScaleWarning, EToolMessageLevel::UserWarning);
+	}
+}
+
+void UMirrorTool::Shutdown(EToolShutdownType ShutdownType)
+{
+	Settings->SaveProperties(this);
+
+	PlaneMechanic->Shutdown();
+
+	// Restore (unhide) the source meshes
+	for (int32 ComponentIdx = 0; ComponentIdx < Targets.Num(); ComponentIdx++)
+	{
+		TargetComponentInterface(ComponentIdx)->SetOwnerVisibility(true);
+	}
+
+	// Swap in results, if appropriate
+	if (ShutdownType == EToolShutdownType::Accept)
+	{
+		// Gather results
+		TArray<FDynamicMeshOpResult> Results;
+		for (UMeshOpPreviewWithBackgroundCompute* Preview : Previews)
+		{
+			Results.Emplace(Preview->Shutdown());
+		}
+
+		// Convert to output. This will also edit the selection.
+		GenerateAsset(Results);
+	}
+	else
+	{
+		for (UMeshOpPreviewWithBackgroundCompute* Preview : Previews)
+		{
+			Preview->Cancel();
+		}
+	}
+}
+
+void UMirrorTool::GenerateAsset(const TArray<FDynamicMeshOpResult>& Results)
+{
+	if (Results.Num() == 0)
+	{
+		return;
+	}
+
+	GetToolManager()->BeginUndoTransaction(LOCTEXT("MirrorToolTransactionName", "Mirror Tool"));
+
+	ensure(Results.Num() > 0);
+
+	int32 NumSourceMeshes = MeshesToMirror.Num();
+
+	// check if we entirely cut away any meshes
+	bool bWantToDestroy = false;
+	for (int OrigMeshIdx = 0; OrigMeshIdx < NumSourceMeshes; OrigMeshIdx++)
+	{
+		if (Results[OrigMeshIdx].Mesh->TriangleCount() == 0)
+		{
+			bWantToDestroy = true;
+			break;
+		}
+	}
+	// if so ask user what to do
+	if (bWantToDestroy)
+	{
+		FText Title = LOCTEXT("MirrorDestroyTitle", "Delete mesh components?");
+		EAppReturnType::Type Ret = FMessageDialog::Open(EAppMsgType::YesNo,
+			LOCTEXT("PlaneCutDestroyQuestion", "The mirror plane cropping has entirely cut away at least one mesh.  Actually destroy these mesh components?"), &Title);
+		if (Ret == EAppReturnType::No)
+		{
+			bWantToDestroy = false;
+		}
+	}
+
+	// Properly deal with each result, setting up the selection at the same time.
+	FSelectedOjectsChangeList NewSelection;
+	NewSelection.ModificationType = ESelectedObjectsModificationType::Replace;
+	for (int OrigMeshIdx = 0; OrigMeshIdx < NumSourceMeshes; OrigMeshIdx++)
+	{
+		IPrimitiveComponentBackedTarget* TargetComponent = TargetComponentInterface(OrigMeshIdx);
+		FDynamicMesh3* Mesh = Results[OrigMeshIdx].Mesh.Get();
+		check(Mesh != nullptr);
+
+		if (Mesh->TriangleCount() == 0)
+		{
+			if (bWantToDestroy)
+			{
+				TargetComponent->GetOwnerComponent()->DestroyComponent();
+			}
+			continue;
+		}
+		else if (Settings->SaveMode == EMirrorSaveMode::UpdateAssets)
+		{
+			NewSelection.Actors.Add(TargetComponent->GetOwnerActor());
+
+			TargetMeshCommitterInterface(OrigMeshIdx)->CommitMeshDescription([&Mesh](const IMeshDescriptionCommitter::FCommitterParams& CommitParams)
+				{
+					FDynamicMeshToMeshDescription Converter;
+					Converter.Convert(Mesh, *CommitParams.MeshDescriptionOut);
+				});
+		}
+		else
+		{
+			// Build array of materials from the original.
+			TArray<UMaterialInterface*> Materials;
+			IMaterialProvider* TargetMaterial = TargetMaterialInterface(OrigMeshIdx);
+			for (int MaterialIdx = 0, NumMaterials = TargetMaterial->GetNumMaterials(); MaterialIdx < NumMaterials; MaterialIdx++)
+			{
+				Materials.Add(TargetMaterial->GetMaterial(MaterialIdx));
+			}
+
+			FCreateMeshObjectParams NewMeshObjectParams;
+			NewMeshObjectParams.TargetWorld = TargetWorld;
+			NewMeshObjectParams.Transform = (FTransform)Results[OrigMeshIdx].Transform;
+			NewMeshObjectParams.BaseName = TEXT("Mirror");
+			NewMeshObjectParams.Materials = Materials;
+			NewMeshObjectParams.SetMesh(Mesh);
+			FCreateMeshObjectResult Result = UE::Modeling::CreateMeshObject(GetToolManager(), MoveTemp(NewMeshObjectParams));
+			if (Result.IsOK() && Result.NewActor != nullptr)
+			{
+				NewSelection.Actors.Add(Result.NewActor);
+			}
+
+			// Remove the original actor
+			TargetComponent->GetOwnerComponent()->DestroyComponent();
+		}
+	}
+
+	// Update the selection
+	if (NewSelection.Actors.Num() > 0)
+	{
+		GetToolManager()->RequestSelectionChange(NewSelection);
+	}
+
+	GetToolManager()->EndUndoTransaction();
+}
+
+
+// Action support
+
+void UMirrorTool::RequestAction(EMirrorToolAction ActionType)
+{
+	if (PendingAction == EMirrorToolAction::NoAction)
+	{
+		PendingAction = ActionType;
+	}
+}
+
+void UMirrorTool::ApplyAction(EMirrorToolAction ActionType)
+{
+	FVector3d ShiftedPlaneOrigin = (FVector3d)CombinedBounds.GetCenter();
+
+	if (ActionType == EMirrorToolAction::ShiftToCenter)
+	{
+		// We keep the same orientation here
+		PlaneMechanic->SetDrawPlaneFromWorldPos(ShiftedPlaneOrigin, FVector3d(), true);
+	}
+	else
+	{
+		// We still start from the center, but adjust one of the coordinates and set direction.
+		FVector3d DirectionVector;
+		switch (ActionType)
+		{
+		case EMirrorToolAction::Left:
+			ShiftedPlaneOrigin.Y = CombinedBounds.Min.Y;
+			DirectionVector = FVector3d(0, -1.0, 0);
+			break;
+		case EMirrorToolAction::Right:
+			ShiftedPlaneOrigin.Y = CombinedBounds.Max.Y;
+			DirectionVector = FVector3d(0, 1.0, 0);
+			break;
+		case EMirrorToolAction::Up:
+			ShiftedPlaneOrigin.Z = CombinedBounds.Max.Z;
+			DirectionVector = FVector3d(0, 0, 1.0);
+			break;
+		case EMirrorToolAction::Down:
+			ShiftedPlaneOrigin.Z = CombinedBounds.Min.Z;
+			DirectionVector = FVector3d(0, 0, -1.0);
+			break;
+		case EMirrorToolAction::Forward:
+			ShiftedPlaneOrigin.X = CombinedBounds.Max.X;
+			DirectionVector = FVector3d(1.0, 0, 0);
+			break;
+		case EMirrorToolAction::Backward:
+			ShiftedPlaneOrigin.X = CombinedBounds.Min.X;
+			DirectionVector = FVector3d(-1.0, 0, 0);
+			break;
+		}
+
+		// The user can optionally have the button change the direction only
+		if (Settings->bButtonsOnlyChangeOrientation)
+		{
+			ShiftedPlaneOrigin = MirrorPlaneOrigin;	// Keeps the same
+		}
+		PlaneMechanic->SetDrawPlaneFromWorldPos(ShiftedPlaneOrigin, DirectionVector, false);
+	}
+}
+
+void UMirrorTool::OnUpdateModifierState(int ModifierID, bool bIsOn)
+{
+}
+
+#undef LOCTEXT_NAMESPACE
diff --git a/Engine/Plugins/Experimental/MeshModelingToolsetExp/Source/MeshModelingToolsExp/Private/Physics/CollisionGeometryVisualization.cpp b/Engine/Plugins/Experimental/MeshModelingToolsetExp/Source/MeshModelingToolsExp/Private/Physics/CollisionGeometryVisualization.cpp
new file mode 100644
--- /dev/null
+++ b/Engine/Plugins/Experimental/MeshModelingToolsetExp/Source/MeshModelingToolsExp/Private/Physics/CollisionGeometryVisualization.cpp
@@ -0,0 +1,121 @@
+// Copyright Epic Games, Inc. All Rights Reserved.
+
+#include "Physics/CollisionGeometryVisualization.h"
+#include "Generators/LineSegmentGenerators.h"
+
+#include "ExplicitUseGeometryMathTypes.h"		// using UE::Geometry::(math types)
+using namespace UE::Geometry;
+
+void UE::PhysicsTools::InitializePreviewGeometryLines(const FPhysicsDataCollection& PhysicsData, UPreviewGeometry* PreviewGeom,
+	const FColor& LineColor, float LineThickness, float DepthBias, int32 CircleStepResolution)
+{
+	int32 CircleSteps = FMath::Max(4, CircleStepResolution);
+	FColor SphereColor = LineColor;
+	FColor BoxColor = LineColor;
+	FColor ConvexColor = LineColor;
+	FColor CapsuleColor = LineColor;
+
+	const FKAggregateGeom& AggGeom = PhysicsData.AggGeom;
+
+	// spheres are draw as 3 orthogonal circles
+	PreviewGeom->CreateOrUpdateLineSet(TEXT("Spheres"), AggGeom.SphereElems.Num(), [&](int32 Index, TArray<FRenderableLine>& LinesOut)
+	{
+		const FKSphereElem& Sphere = AggGeom.SphereElems[Index];
+		FTransform ElemTransform = Sphere.GetTransform();
+		ElemTransform.ScaleTranslation(PhysicsData.ExternalScale3D);
+		UE::Geometry::FTransform3f ElemTransformf(ElemTransform);
+		float Radius = PhysicsData.ExternalScale3D.GetAbsMin() * Sphere.Radius;
+		UE::Geometry::GenerateCircleSegments<float>(CircleSteps, Radius, FVector3f::Zero(), FVector3f::UnitX(), FVector3f::UnitY(), ElemTransformf,
+			[&](const FVector3f& A, const FVector3f& B) { LinesOut.Add(FRenderableLine((FVector)A, (FVector)B, SphereColor, LineThickness, DepthBias)); });
+		UE::Geometry::GenerateCircleSegments<float>(CircleSteps, Radius, FVector3f::Zero(), FVector3f::UnitX(), FVector3f::UnitZ(), ElemTransformf,
+			[&](const FVector3f& A, const FVector3f& B) { LinesOut.Add(FRenderableLine((FVector)A, (FVector)B, SphereColor, LineThickness, DepthBias)); });
+		UE::Geometry::GenerateCircleSegments<float>(CircleSteps, Radius, FVector3f::Zero(), FVector3f::UnitY(), FVector3f::UnitZ(), ElemTransformf,
+			[&](const FVector3f& A, const FVector3f& B) { LinesOut.Add(FRenderableLine((FVector)A, (FVector)B, SphereColor, LineThickness, DepthBias)); });
+	});
+
+
+	// boxes are drawn as boxes
+	PreviewGeom->CreateOrUpdateLineSet(TEXT("Boxes"), AggGeom.BoxElems.Num(), [&](int32 Index, TArray<FRenderableLine>& LinesOut)
+	{
+		const FKBoxElem& Box = AggGeom.BoxElems[Index];
+		FTransform ElemTransform = Box.GetTransform();
+		ElemTransform.ScaleTranslation(PhysicsData.ExternalScale3D);
+		UE::Geometry::FTransform3f ElemTransformf(ElemTransform);
+		FVector3f HalfDimensions(
+			PhysicsData.ExternalScale3D.X * Box.X * 0.5f,
+			PhysicsData.ExternalScale3D.Y * Box.Y * 0.5f,
+			PhysicsData.ExternalScale3D.Z * Box.Z * 0.5f);
+		UE::Geometry::GenerateBoxSegments<float>(HalfDimensions, FVector3f::Zero(), FVector3f::UnitX(), FVector3f::UnitY(), FVector3f::UnitZ(), ElemTransformf,
+			[&](const FVector3f& A, const FVector3f& B) { LinesOut.Add(FRenderableLine((FVector)A, (FVector)B, BoxColor, LineThickness, DepthBias)); });
+	});
+
+
+	// capsules are draw as two hemispheres (with 3 intersecting arcs/circles) and connecting lines
+	PreviewGeom->CreateOrUpdateLineSet(TEXT("Capsules"), AggGeom.SphylElems.Num(), [&](int32 Index, TArray<FRenderableLine>& LinesOut)
+	{
+		const FKSphylElem& Capsule = AggGeom.SphylElems[Index];
+		FTransform ElemTransform = Capsule.GetTransform();
+		ElemTransform.ScaleTranslation(PhysicsData.ExternalScale3D);
+		UE::Geometry::FTransform3f ElemTransformf(ElemTransform);
+		const float HalfLength = Capsule.GetScaledCylinderLength(PhysicsData.ExternalScale3D) * .5f;
+		const float Radius = Capsule.GetScaledRadius(PhysicsData.ExternalScale3D);
+		FVector3f Top(0, 0, HalfLength), Bottom(0, 0, -HalfLength);
+
+		// top and bottom circles
+		UE::Geometry::GenerateCircleSegments<float>(CircleSteps, Radius, Top, FVector3f::UnitX(), FVector3f::UnitY(), ElemTransformf,
+			[&](const FVector3f& A, const FVector3f& B) { LinesOut.Add(FRenderableLine((FVector)A, (FVector)B, CapsuleColor, LineThickness, DepthBias)); });
+		UE::Geometry::GenerateCircleSegments<float>(CircleSteps, Radius, Bottom, FVector3f::UnitX(), FVector3f::UnitY(), ElemTransformf,
+			[&](const FVector3f& A, const FVector3f& B) { LinesOut.Add(FRenderableLine((FVector)A, (FVector)B, CapsuleColor, LineThickness, DepthBias)); });
+
+		// top dome
+		UE::Geometry::GenerateArcSegments<float>(CircleSteps, Radius, 0.0, PI, Top, FVector3f::UnitY(), FVector3f::UnitZ(), ElemTransformf,
+			[&](const FVector3f& A, const FVector3f& B) { LinesOut.Add(FRenderableLine((FVector)A, (FVector)B, CapsuleColor, LineThickness, DepthBias)); });
+		UE::Geometry::GenerateArcSegments<float>(CircleSteps, Radius, 0.0, PI, Top, FVector3f::UnitX(), FVector3f::UnitZ(), ElemTransformf,
+			[&](const FVector3f& A, const FVector3f& B) { LinesOut.Add(FRenderableLine((FVector)A, (FVector)B, CapsuleColor, LineThickness, DepthBias)); });
+
+		// bottom dome
+		UE::Geometry::GenerateArcSegments<float>(CircleSteps, Radius, 0.0, -PI, Bottom, FVector3f::UnitY(), FVector3f::UnitZ(), ElemTransformf,
+			[&](const FVector3f& A, const FVector3f& B) { LinesOut.Add(FRenderableLine((FVector)A, (FVector)B, CapsuleColor, LineThickness, DepthBias)); });
+		UE::Geometry::GenerateArcSegments<float>(CircleSteps, Radius, 0.0, -PI, Bottom, FVector3f::UnitX(), FVector3f::UnitZ(), ElemTransformf,
+			[&](const FVector3f& A, const FVector3f& B) { LinesOut.Add(FRenderableLine((FVector)A, (FVector)B, CapsuleColor, LineThickness, DepthBias)); });
+
+		// connecting lines
+		for (int k = 0; k < 2; ++k)
+		{
+			FVector DX = (k < 1) ? FVector(-Radius, 0, 0) : FVector(Radius, 0, 0);
+			LinesOut.Add(FRenderableLine(
+				ElemTransform.TransformPosition((FVector)Top + DX), 
+				ElemTransform.TransformPosition((FVector)Bottom + DX), CapsuleColor, LineThickness, DepthBias));
+			FVector DY = (k < 1) ? FVector(0, -Radius, 0) : FVector(0, Radius, 0);
+			LinesOut.Add(FRenderableLine(
+				ElemTransform.TransformPosition((FVector)Top + DY), 
+				ElemTransform.TransformPosition((FVector)Bottom + DY), CapsuleColor, LineThickness, DepthBias));
+		}
+	});
+
+
+
+	// convexes are drawn as mesh edges
+	PreviewGeom->CreateOrUpdateLineSet(TEXT("Convexes"), AggGeom.ConvexElems.Num(), [&](int32 Index, TArray<FRenderableLine>& LinesOut)
+	{
+		const FKConvexElem& Convex = AggGeom.ConvexElems[Index];
+		FTransform ElemTransform = Convex.GetTransform();
+		ElemTransform.ScaleTranslation(PhysicsData.ExternalScale3D);
+		ElemTransform.SetScale3D(PhysicsData.ExternalScale3D);
+		int32 NumTriangles = Convex.IndexData.Num() / 3;
+		for (int32 k = 0; k < NumTriangles; ++k)
+		{
+			FVector A = ElemTransform.TransformPosition(Convex.VertexData[Convex.IndexData[3 * k]]);
+			FVector B = ElemTransform.TransformPosition(Convex.VertexData[Convex.IndexData[3 * k + 1]]);
+			FVector C = ElemTransform.TransformPosition(Convex.VertexData[Convex.IndexData[3 * k + 2]]);
+			LinesOut.Add(FRenderableLine(A, B, ConvexColor, LineThickness, DepthBias));
+			LinesOut.Add(FRenderableLine(B, C, ConvexColor, LineThickness, DepthBias));
+			LinesOut.Add(FRenderableLine(C, A, ConvexColor, LineThickness, DepthBias));
+		}
+	});
+
+
+	// Unclear whether we actually use these in the Engine, for UBodySetup? Does not appear to be supported by UxX import system,
+	// and online documentation suggests they may only be supported for cloth?
+	ensure(AggGeom.TaperedCapsuleElems.Num() == 0);
+}
\ No newline at end of file
diff --git a/Engine/Plugins/Experimental/MeshModelingToolsetExp/Source/MeshModelingToolsExp/Private/Physics/ExtractCollisionGeometryTool.cpp b/Engine/Plugins/Experimental/MeshModelingToolsetExp/Source/MeshModelingToolsExp/Private/Physics/ExtractCollisionGeometryTool.cpp
new file mode 100644
--- /dev/null
+++ b/Engine/Plugins/Experimental/MeshModelingToolsetExp/Source/MeshModelingToolsExp/Private/Physics/ExtractCollisionGeometryTool.cpp
@@ -0,0 +1,413 @@
+// Copyright Epic Games, Inc. All Rights Reserved.
+
+#include "Physics/ExtractCollisionGeometryTool.h"
+#include "InteractiveToolManager.h"
+#include "ToolBuilderUtil.h"
+#include "ToolSetupUtil.h"
+#include "ModelingObjectsCreationAPI.h"
+#include "Selection/ToolSelectionUtil.h"
+#include "Drawing/PreviewGeometryActor.h"
+#include "Util/ColorConstants.h"
+#include "PreviewMesh.h"
+
+#include "DynamicMeshEditor.h"
+#include "DynamicMesh/MeshNormals.h"
+#include "Generators/SphereGenerator.h"
+#include "Generators/MinimalBoxMeshGenerator.h"
+#include "Generators/CapsuleGenerator.h"
+#include "DynamicMesh/MeshTransforms.h"
+#include "Parameterization/DynamicMeshUVEditor.h"
+#include "DynamicMesh/Operations/MergeCoincidentMeshEdges.h"
+
+#include "Physics/PhysicsDataCollection.h"
+#include "Physics/CollisionGeometryVisualization.h"
+
+// physics data
+#include "Engine/Classes/Engine/StaticMesh.h"
+#include "Engine/Classes/Components/StaticMeshComponent.h"
+#include "Engine/Classes/PhysicsEngine/BodySetup.h"
+#include "Engine/Classes/PhysicsEngine/AggregateGeom.h"
+
+#include "TargetInterfaces/MeshDescriptionProvider.h"
+#include "TargetInterfaces/PrimitiveComponentBackedTarget.h"
+#include "TargetInterfaces/StaticMeshBackedTarget.h"
+#include "ModelingToolTargetUtil.h"
+
+#include "ExplicitUseGeometryMathTypes.h"		// using UE::Geometry::(math types)
+using namespace UE::Geometry;
+
+#define LOCTEXT_NAMESPACE "UExtractCollisionGeometryTool"
+
+const FToolTargetTypeRequirements& UExtractCollisionGeometryToolBuilder::GetTargetRequirements() const
+{
+	static FToolTargetTypeRequirements TypeRequirements({
+		UMeshDescriptionProvider::StaticClass(),
+		UPrimitiveComponentBackedTarget::StaticClass(),
+		UStaticMeshBackedTarget::StaticClass()
+		});
+	return TypeRequirements;
+}
+
+
+USingleSelectionMeshEditingTool* UExtractCollisionGeometryToolBuilder::CreateNewTool(const FToolBuilderState& SceneState) const
+{
+	return NewObject<UExtractCollisionGeometryTool>(SceneState.ToolManager);
+}
+
+
+void UExtractCollisionGeometryTool::Setup()
+{
+	UInteractiveTool::Setup();
+
+	// create preview mesh
+	PreviewMesh = NewObject<UPreviewMesh>(this);
+	PreviewMesh->bBuildSpatialDataStructure = false;
+	PreviewMesh->CreateInWorld(TargetWorld, FTransform::Identity);
+	PreviewMesh->SetTransform((FTransform)UE::ToolTarget::GetLocalToWorldTransform(Target));
+	PreviewMesh->SetMaterial(ToolSetupUtil::GetDefaultSculptMaterial(GetToolManager()));
+	PreviewMesh->SetOverrideRenderMaterial(ToolSetupUtil::GetSelectionMaterial(GetToolManager()));
+	PreviewMesh->SetTriangleColorFunction([this](const FDynamicMesh3* Mesh, int TriangleID)
+	{
+		return LinearColors::SelectFColor(Mesh->GetTriangleGroup(TriangleID));
+	});
+
+	OutputTypeProperties = NewObject<UCreateMeshObjectTypeProperties>(this);
+	OutputTypeProperties->OutputType = UCreateMeshObjectTypeProperties::VolumeIdentifier;		// prefer volumes for extracting simple collision
+	OutputTypeProperties->RestoreProperties(this, TEXT("ExtractCollisionTool"));
+	OutputTypeProperties->InitializeDefault();
+	OutputTypeProperties->WatchProperty(OutputTypeProperties->OutputType, [this](FString) { OutputTypeProperties->UpdatePropertyVisibility(); });
+	AddToolPropertySource(OutputTypeProperties);
+
+	Settings = NewObject<UExtractCollisionToolProperties>(this);
+	Settings->RestoreProperties(this);
+	AddToolPropertySource(Settings);
+	Settings->WatchProperty(Settings->CollisionType, [this](EExtractCollisionOutputType NewValue) { bResultValid = false; });
+	Settings->WatchProperty(Settings->bWeldEdges, [this](bool bNewValue) { bResultValid = false; });
+	Settings->WatchProperty(Settings->bShowPreview, [this](bool bNewValue) { PreviewMesh->SetVisible(bNewValue); });
+	PreviewMesh->SetVisible(Settings->bShowPreview);
+	Settings->WatchProperty(Settings->bShowInputMesh, [this](bool bNewValue) { UE::ToolTarget::SetSourceObjectVisible(Target, bNewValue); });
+	UE::ToolTarget::SetSourceObjectVisible(Target, Settings->bShowInputMesh);
+
+	VizSettings = NewObject<UCollisionGeometryVisualizationProperties>(this);
+	VizSettings->RestoreProperties(this);
+	AddToolPropertySource(VizSettings);
+	VizSettings->WatchProperty(VizSettings->LineThickness, [this](float NewValue) { bVisualizationDirty = true; });
+	VizSettings->WatchProperty(VizSettings->Color, [this](FColor NewValue) { bVisualizationDirty = true; });
+
+
+	UPrimitiveComponent* TargetComponent = UE::ToolTarget::GetTargetComponent(Target);
+	UStaticMeshComponent* StaticMeshComponent = (TargetComponent) ? Cast<UStaticMeshComponent>(TargetComponent) : nullptr;
+	TObjectPtr<UStaticMesh> StaticMesh = (StaticMeshComponent) ? StaticMeshComponent->GetStaticMesh() : nullptr;
+	UBodySetup* BodySetup = (StaticMesh) ? StaticMesh->GetBodySetup() : nullptr;
+	if (ensure(BodySetup))
+	{
+		PhysicsInfo = MakeShared<FPhysicsDataCollection>();
+		PhysicsInfo->InitializeFromComponent(StaticMeshComponent, true);
+
+		PreviewElements = NewObject<UPreviewGeometry>(this);
+		FTransform TargetTransform = (FTransform)UE::ToolTarget::GetLocalToWorldTransform(Target);
+		//PhysicsInfo->ExternalScale3D = TargetTransform.GetScale3D();
+		//TargetTransform.SetScale3D(FVector::OneVector);
+		PreviewElements->CreateInWorld(UE::ToolTarget::GetTargetActor(Target)->GetWorld(), TargetTransform);
+		
+		UE::PhysicsTools::InitializePreviewGeometryLines(*PhysicsInfo, PreviewElements,
+			VizSettings->Color, VizSettings->LineThickness, 0.0f, 16);
+
+		ObjectProps = NewObject<UPhysicsObjectToolPropertySet>(this);
+		UE::PhysicsTools::InitializePhysicsToolObjectPropertySet(PhysicsInfo.Get(), ObjectProps);
+		AddToolPropertySource(ObjectProps);
+	}
+
+	SetToolDisplayName(LOCTEXT("ToolName", "Collision To Mesh"));
+	GetToolManager()->DisplayMessage(
+		LOCTEXT("OnStartTool", "Convert Collision Geometry to Mesh Objects"),
+		EToolMessageLevel::UserNotification);
+}
+
+
+void UExtractCollisionGeometryTool::Shutdown(EToolShutdownType ShutdownType)
+{
+	OutputTypeProperties->SaveProperties(this, TEXT("ExtractCollisionTool"));
+	Settings->SaveProperties(this);
+	VizSettings->SaveProperties(this);
+
+	UE::Geometry::FTransform3d ActorTransform(PreviewMesh->GetTransform());
+
+	PreviewElements->Disconnect();
+	PreviewMesh->SetVisible(false);
+	PreviewMesh->Disconnect();
+	PreviewMesh = nullptr;
+
+	UE::ToolTarget::ShowSourceObject(Target);
+
+	if (ShutdownType == EToolShutdownType::Accept)
+	{
+		UMaterialInterface* UseMaterial = UMaterial::GetDefaultMaterial(MD_Surface);
+
+		FString TargetName = UE::ToolTarget::GetTargetComponent(Target)->GetName();
+
+		GetToolManager()->BeginUndoTransaction(LOCTEXT("CreateCollisionMesh", "Collision To Mesh"));
+
+		TArray<AActor*> OutputSelection;
+
+		auto EmitNewMesh = [&](FDynamicMesh3&& Mesh, FTransform3d UseTransform, FString UseName)
+		{
+			FCreateMeshObjectParams NewMeshObjectParams;
+			NewMeshObjectParams.TargetWorld = TargetWorld;
+			NewMeshObjectParams.Transform = (FTransform)UseTransform;
+			NewMeshObjectParams.BaseName = UseName;
+			NewMeshObjectParams.Materials.Add(UseMaterial);
+			NewMeshObjectParams.SetMesh(MoveTemp(Mesh));
+			OutputTypeProperties->ConfigureCreateMeshObjectParams(NewMeshObjectParams);
+			FCreateMeshObjectResult Result = UE::Modeling::CreateMeshObject(GetToolManager(), MoveTemp(NewMeshObjectParams));
+			if (Result.IsOK() && Result.NewActor != nullptr)
+			{
+				OutputSelection.Add(Result.NewActor);
+			}
+		};
+
+		int32 NumParts = CurrentMeshParts.Num();
+		if (Settings->bOutputSeparateMeshes && NumParts > 1)
+		{
+			for ( int32 k = 0; k < NumParts; ++k)
+			{
+				FDynamicMesh3& MeshPart = CurrentMeshParts[k];
+				FAxisAlignedBox3d Bounds = MeshPart.GetBounds();
+				MeshTransforms::Translate(MeshPart, -Bounds.Center());
+				FTransform3d CenterTransform = ActorTransform;
+				CenterTransform.SetTranslation(CenterTransform.GetTranslation() + Bounds.Center());
+				FString NewName = FString::Printf(TEXT("%s_Collision%d"), *TargetName, k);
+				EmitNewMesh(MoveTemp(MeshPart), CenterTransform, NewName);
+			}
+		}
+		else
+		{
+			FString NewName = FString::Printf(TEXT("%s_Collision"), *TargetName);
+			EmitNewMesh(MoveTemp(CurrentMesh), ActorTransform, NewName);
+		}
+
+		ToolSelectionUtil::SetNewActorSelection(GetToolManager(), OutputSelection);
+
+		GetToolManager()->EndUndoTransaction();
+	}
+}
+
+
+bool UExtractCollisionGeometryTool::CanAccept() const
+{
+	return Super::CanAccept() && CurrentMesh.TriangleCount() > 0;
+}
+
+
+
+
+void UExtractCollisionGeometryTool::OnTick(float DeltaTime)
+{
+	if (bResultValid == false)
+	{
+		GetToolManager()->DisplayMessage(FText(), EToolMessageLevel::UserWarning);
+		if (Settings->CollisionType == EExtractCollisionOutputType::Simple)
+		{
+			RecalculateMesh_Simple();
+		}
+		else
+		{
+			RecalculateMesh_Complex();
+		}
+		bResultValid = true;
+	}
+
+	if (bVisualizationDirty)
+	{
+		UpdateVisualization();
+		bVisualizationDirty = false;
+	}
+}
+
+
+
+void UExtractCollisionGeometryTool::UpdateVisualization()
+{
+	float UseThickness = VizSettings->LineThickness;
+	FColor UseColor = VizSettings->Color;
+	PreviewElements->UpdateAllLineSets([&](ULineSetComponent* LineSet)
+		{
+			LineSet->SetAllLinesThickness(UseThickness);
+			LineSet->SetAllLinesColor(UseColor);
+		});
+	
+	UMaterialInterface* LineMaterial = ToolSetupUtil::GetDefaultLineComponentMaterial(GetToolManager(), !VizSettings->bShowHidden);
+	PreviewElements->SetAllLineSetsMaterial(LineMaterial);
+}
+
+
+
+void UExtractCollisionGeometryTool::RecalculateMesh_Simple()
+{
+	int32 SphereResolution = 16;
+
+	CurrentMesh = FDynamicMesh3(EMeshComponents::FaceGroups);
+	CurrentMesh.EnableAttributes();
+
+	CurrentMeshParts.Reset();
+
+	FDynamicMeshEditor Editor(&CurrentMesh);
+
+	const FKAggregateGeom& AggGeom = PhysicsInfo->AggGeom;
+
+	for (const FKSphereElem& Sphere : AggGeom.SphereElems)
+	{
+		FSphereGenerator SphereGen;
+		SphereGen.Radius = Sphere.Radius;
+		SphereGen.NumPhi = SphereGen.NumTheta = SphereResolution;
+		SphereGen.bPolygroupPerQuad = false;
+		SphereGen.Generate();
+		FDynamicMesh3 SphereMesh(&SphereGen);
+
+		MeshTransforms::Translate(SphereMesh, FVector3d(Sphere.Center));
+
+		FMeshIndexMappings Mappings;
+		Editor.AppendMesh(&SphereMesh, Mappings);
+		CurrentMeshParts.Add(MoveTemp(SphereMesh));
+	}
+
+
+	for (const FKBoxElem& Box : AggGeom.BoxElems)
+	{
+		FMinimalBoxMeshGenerator BoxGen;
+		BoxGen.Box = FOrientedBox3d(
+			FFrame3d(FVector3d(Box.Center), FQuaterniond(Box.Rotation.Quaternion())),
+			0.5*FVector3d(Box.X, Box.Y, Box.Z));
+		BoxGen.Generate();
+		FDynamicMesh3 BoxMesh(&BoxGen);
+
+		// transform not applied because it is just the Center/Rotation
+
+		FMeshIndexMappings Mappings;
+		Editor.AppendMesh(&BoxMesh, Mappings);
+		CurrentMeshParts.Add(MoveTemp(BoxMesh));
+	}
+
+
+	for (const FKSphylElem& Capsule: AggGeom.SphylElems)
+	{
+		FCapsuleGenerator CapsuleGen;
+		CapsuleGen.Radius = Capsule.Radius;
+		CapsuleGen.SegmentLength = Capsule.Length;
+		CapsuleGen.NumHemisphereArcSteps = SphereResolution/4+1;
+		CapsuleGen.NumCircleSteps = SphereResolution;
+		CapsuleGen.bPolygroupPerQuad = false;
+		CapsuleGen.Generate();
+		FDynamicMesh3 CapsuleMesh(&CapsuleGen);
+
+		MeshTransforms::Translate(CapsuleMesh, FVector3d(0,0,-0.5*Capsule.Length) );
+
+		UE::Geometry::FTransform3d Transform(Capsule.GetTransform());
+		MeshTransforms::ApplyTransform(CapsuleMesh, Transform);
+
+		FMeshIndexMappings Mappings;
+		Editor.AppendMesh(&CapsuleMesh, Mappings);
+		CurrentMeshParts.Add(MoveTemp(CapsuleMesh));
+	}
+
+
+	for (const FKConvexElem& Convex : AggGeom.ConvexElems)
+	{
+		UE::Geometry::FTransform3d ElemTransform(Convex.GetTransform());
+		FDynamicMesh3 ConvexMesh(EMeshComponents::None);
+		int32 NumVertices = Convex.VertexData.Num();
+		for (int32 k = 0; k < NumVertices; ++k)
+		{
+			ConvexMesh.AppendVertex( FVector3d(Convex.VertexData[k]) );
+		}
+		int32 NumTriangles = Convex.IndexData.Num() / 3;
+		for (int32 k = 0; k < NumTriangles; ++k)
+		{
+			ConvexMesh.AppendTriangle(Convex.IndexData[3*k], Convex.IndexData[3*k+1], Convex.IndexData[3*k+2]);
+		}
+
+		ConvexMesh.ReverseOrientation();
+		ConvexMesh.EnableTriangleGroups(0);
+		ConvexMesh.EnableAttributes();
+		FDynamicMeshUVEditor UVEditor(&ConvexMesh, 0, true);
+		UVEditor.SetPerTriangleUVs();
+		FMeshIndexMappings Mappings;
+		Editor.AppendMesh(&ConvexMesh, Mappings);
+		CurrentMeshParts.Add(MoveTemp(ConvexMesh));
+	}
+
+	for (FDynamicMesh3& MeshPart : CurrentMeshParts)
+	{
+		FMeshNormals::InitializeMeshToPerTriangleNormals(&MeshPart);
+	}
+	FMeshNormals::InitializeMeshToPerTriangleNormals(&CurrentMesh);
+
+	PreviewMesh->UpdatePreview(&CurrentMesh);
+
+	if (CurrentMeshParts.Num() == 0)
+	{
+		GetToolManager()->DisplayMessage(LOCTEXT("NoSimpleCollisionShapes", "This Mesh has no Simple Collision Shapes"), EToolMessageLevel::UserWarning);
+	}
+}
+
+
+
+void UExtractCollisionGeometryTool::RecalculateMesh_Complex()
+{
+	CurrentMesh = FDynamicMesh3(EMeshComponents::FaceGroups);
+	CurrentMesh.EnableAttributes();
+	CurrentMeshParts.Reset();
+
+	bool bMeshErrors = false;
+
+	UStaticMeshComponent* StaticMeshComponent = Cast<UStaticMeshComponent>(UE::ToolTarget::GetTargetComponent(Target));
+	TObjectPtr<UStaticMesh> StaticMesh = (StaticMeshComponent) ? StaticMeshComponent->GetStaticMesh() : nullptr;
+	IInterface_CollisionDataProvider* CollisionProvider = (StaticMesh) ? Cast<IInterface_CollisionDataProvider>(StaticMesh) : nullptr;
+	if (CollisionProvider && CollisionProvider->ContainsPhysicsTriMeshData(true))
+	{
+		FTriMeshCollisionData CollisionData;
+		if (CollisionProvider->GetPhysicsTriMeshData(&CollisionData, true))
+		{
+			TArray<int32> VertexIDMap;
+			for (int32 k = 0; k < CollisionData.Vertices.Num(); ++k)
+			{
+				int32 vid = CurrentMesh.AppendVertex(CollisionData.Vertices[k]);
+				VertexIDMap.Add(vid);
+			}
+			for (const FTriIndices& TriIndices : CollisionData.Indices)
+			{
+				FIndex3i Triangle(TriIndices.v0, TriIndices.v1, TriIndices.v2);
+				int32 tid = CurrentMesh.AppendTriangle(Triangle);
+				if (tid < 0)
+				{
+					bMeshErrors = true;
+				}
+			}
+
+			if (Settings->bWeldEdges)
+			{
+				FMergeCoincidentMeshEdges Weld(&CurrentMesh);
+				Weld.OnlyUniquePairs = true;
+				Weld.Apply();
+				Weld.OnlyUniquePairs = false;
+				Weld.Apply();
+			}
+
+			if (!CollisionData.bFlipNormals)		// collision mesh has reversed orientation
+			{
+				CurrentMesh.ReverseOrientation(false);
+			}
+
+			FMeshNormals::InitializeMeshToPerTriangleNormals(&CurrentMesh);
+		}
+	}
+
+	PreviewMesh->UpdatePreview(&CurrentMesh);
+
+	if (CurrentMesh.TriangleCount() == 0)
+	{
+		GetToolManager()->DisplayMessage(LOCTEXT("EmptyComplexCollision", "This Mesh has no Complex Collision geometry"), EToolMessageLevel::UserWarning);
+	}
+}
+
+#undef LOCTEXT_NAMESPACE
\ No newline at end of file
diff --git a/Engine/Plugins/Experimental/MeshModelingToolsetExp/Source/MeshModelingToolsExp/Private/Physics/PhysicsInspectorTool.cpp b/Engine/Plugins/Experimental/MeshModelingToolsetExp/Source/MeshModelingToolsExp/Private/Physics/PhysicsInspectorTool.cpp
new file mode 100644
--- /dev/null
+++ b/Engine/Plugins/Experimental/MeshModelingToolsetExp/Source/MeshModelingToolsExp/Private/Physics/PhysicsInspectorTool.cpp
@@ -0,0 +1,160 @@
+// Copyright Epic Games, Inc. All Rights Reserved.
+
+#include "Physics/PhysicsInspectorTool.h"
+#include "InteractiveToolManager.h"
+#include "ToolBuilderUtil.h"
+#include "ToolSetupUtil.h"
+#include "Drawing/PreviewGeometryActor.h"
+
+#include "Physics/PhysicsDataCollection.h"
+#include "Physics/CollisionGeometryVisualization.h"
+
+// physics data
+#include "Engine/Classes/Engine/StaticMesh.h"
+#include "Engine/Classes/Components/StaticMeshComponent.h"
+#include "Engine/Classes/PhysicsEngine/BodySetup.h"
+#include "Engine/Classes/PhysicsEngine/AggregateGeom.h"
+
+#include "TargetInterfaces/PrimitiveComponentBackedTarget.h"
+#include "TargetInterfaces/StaticMeshBackedTarget.h"
+#include "ToolTargetManager.h"
+
+#include "ExplicitUseGeometryMathTypes.h"		// using UE::Geometry::(math types)
+using namespace UE::Geometry;
+
+#define LOCTEXT_NAMESPACE "UPhysicsInspectorTool"
+
+
+const FToolTargetTypeRequirements& UPhysicsInspectorToolBuilder::GetTargetRequirements() const
+{
+	static FToolTargetTypeRequirements TypeRequirements({
+		UPrimitiveComponentBackedTarget::StaticClass(),
+		UStaticMeshBackedTarget::StaticClass()
+		});
+	return TypeRequirements;
+}
+
+bool UPhysicsInspectorToolBuilder::CanBuildTool(const FToolBuilderState& SceneState) const
+{
+	return SceneState.TargetManager->CountSelectedAndTargetable(SceneState, GetTargetRequirements()) > 0;
+}
+
+
+UInteractiveTool* UPhysicsInspectorToolBuilder::BuildTool(const FToolBuilderState& SceneState) const
+{
+	UPhysicsInspectorTool* NewTool = NewObject<UPhysicsInspectorTool>(SceneState.ToolManager);
+
+	TArray<TObjectPtr<UToolTarget>> Targets = SceneState.TargetManager->BuildAllSelectedTargetable(SceneState, GetTargetRequirements());
+	NewTool->SetTargets(MoveTemp(Targets));
+
+	return NewTool;
+}
+
+
+void UPhysicsInspectorTool::Setup()
+{
+	UInteractiveTool::Setup();
+
+	VizSettings = NewObject<UCollisionGeometryVisualizationProperties>(this);
+	VizSettings->RestoreProperties(this);
+	AddToolPropertySource(VizSettings);
+	VizSettings->WatchProperty(VizSettings->LineThickness, [this](float NewValue) { bVisualizationDirty = true; });
+	VizSettings->WatchProperty(VizSettings->Color, [this](FColor NewValue) { bVisualizationDirty = true; });
+	VizSettings->WatchProperty(VizSettings->bShowHidden, [this](bool bNewValue) { bVisualizationDirty = true; });
+
+	for (int32 ComponentIdx = 0; ComponentIdx < Targets.Num(); ComponentIdx++)
+	{
+		IPrimitiveComponentBackedTarget* TargetComponent = TargetComponentInterface(ComponentIdx);
+		const UStaticMeshComponent* Component = CastChecked<UStaticMeshComponent>(TargetComponent->GetOwnerComponent());
+		const UStaticMesh* StaticMesh = Component->GetStaticMesh();
+		if (ensure(StaticMesh && StaticMesh->GetBodySetup()))
+		{
+			TSharedPtr<FPhysicsDataCollection> PhysicsData = MakeShared<FPhysicsDataCollection>();
+			PhysicsData->SourceComponent = Component;
+			PhysicsData->BodySetup = StaticMesh->GetBodySetup();
+			PhysicsData->AggGeom = StaticMesh->GetBodySetup()->AggGeom;
+
+			PhysicsInfos.Add(PhysicsData);
+
+			UPreviewGeometry* PreviewGeom = NewObject<UPreviewGeometry>(this);
+			FTransform TargetTransform = TargetComponent->GetWorldTransform();
+			PhysicsData->ExternalScale3D = TargetTransform.GetScale3D();
+			TargetTransform.SetScale3D(FVector::OneVector);
+			PreviewGeom->CreateInWorld(TargetComponent->GetOwnerActor()->GetWorld(), TargetTransform);
+			PreviewElements.Add(PreviewGeom);
+
+			InitializeGeometry(*PhysicsData, PreviewGeom);
+
+			UPhysicsObjectToolPropertySet* ObjectProps = NewObject<UPhysicsObjectToolPropertySet>(this);
+			InitializeObjectProperties(*PhysicsData, ObjectProps);
+			AddToolPropertySource(ObjectProps);
+		}
+	}
+
+	SetToolDisplayName(LOCTEXT("ToolName", "Physics Inspector"));
+	GetToolManager()->DisplayMessage(
+		LOCTEXT("OnStartTool", "Inspect Physics data for the seleced Static Meshes"),
+		EToolMessageLevel::UserNotification);
+}
+
+
+void UPhysicsInspectorTool::Shutdown(EToolShutdownType ShutdownType)
+{
+	VizSettings->SaveProperties(this);
+
+	for (UPreviewGeometry* Preview : PreviewElements)
+	{
+		Preview->Disconnect();
+	}
+}
+
+
+
+
+
+void UPhysicsInspectorTool::OnTick(float DeltaTime)
+{
+	if (bVisualizationDirty)
+	{
+		UpdateVisualization();
+		bVisualizationDirty = false;
+	}
+}
+
+
+
+void UPhysicsInspectorTool::UpdateVisualization()
+{
+	float UseThickness = VizSettings->LineThickness;
+	FColor UseColor = VizSettings->Color;
+	LineMaterial = ToolSetupUtil::GetDefaultLineComponentMaterial(GetToolManager(), !VizSettings->bShowHidden);
+
+	for (UPreviewGeometry* Preview : PreviewElements)
+	{
+		Preview->UpdateAllLineSets([&](ULineSetComponent* LineSet)
+		{
+			LineSet->SetAllLinesThickness(UseThickness);
+			LineSet->SetAllLinesColor(UseColor);
+		});
+		Preview->SetAllLineSetsMaterial(LineMaterial);
+	}
+}
+
+
+
+void UPhysicsInspectorTool::InitializeObjectProperties(const FPhysicsDataCollection& PhysicsData, UPhysicsObjectToolPropertySet* PropSet)
+{
+	UE::PhysicsTools::InitializePhysicsToolObjectPropertySet(&PhysicsData, PropSet);
+}
+
+
+
+void UPhysicsInspectorTool::InitializeGeometry(const FPhysicsDataCollection& PhysicsData, UPreviewGeometry* PreviewGeom)
+{
+	UE::PhysicsTools::InitializePreviewGeometryLines(PhysicsData, PreviewGeom,
+		VizSettings->Color, VizSettings->LineThickness, 0.0f, 16);
+}
+
+
+
+#undef LOCTEXT_NAMESPACE
\ No newline at end of file
