MeshModelingToolset: Split plugin into non-experimental and experimental

diff --git a/Engine/Plugins/Experimental/MeshModelingToolsetExp/Source/MeshModelingToolsExp/Private/Physics/SetCollisionGeometryTool.cpp b/Engine/Plugins/Experimental/MeshModelingToolsetExp/Source/MeshModelingToolsExp/Private/Physics/SetCollisionGeometryTool.cpp
new file mode 100644
--- /dev/null
+++ b/Engine/Plugins/Experimental/MeshModelingToolsetExp/Source/MeshModelingToolsExp/Private/Physics/SetCollisionGeometryTool.cpp
@@ -0,0 +1,572 @@
+// Copyright Epic Games, Inc. All Rights Reserved.
+
+#include "Physics/SetCollisionGeometryTool.h"
+#include "InteractiveToolManager.h"
+#include "ToolBuilderUtil.h"
+#include "ToolSetupUtil.h"
+#include "DynamicMesh/MeshTransforms.h"
+#include "DynamicMeshEditor.h"
+#include "Selections/MeshConnectedComponents.h"
+#include "DynamicSubmesh3.h"
+#include "Polygroups/PolygroupUtil.h"
+
+#include "ShapeApproximation/ShapeDetection3.h"
+#include "ShapeApproximation/MeshSimpleShapeApproximation.h"
+
+#include "Physics/PhysicsDataCollection.h"
+#include "Physics/CollisionGeometryVisualization.h"
+
+// physics data
+#include "Engine/Classes/Engine/StaticMesh.h"
+#include "Engine/Classes/Components/StaticMeshComponent.h"
+#include "Engine/Classes/PhysicsEngine/BodySetup.h"
+
+#include "Async/ParallelFor.h"
+
+#include "TargetInterfaces/MeshDescriptionProvider.h"
+#include "TargetInterfaces/PrimitiveComponentBackedTarget.h"
+#include "ToolTargetManager.h"
+#include "ModelingToolTargetUtil.h"
+
+#include "ExplicitUseGeometryMathTypes.h"		// using UE::Geometry::(math types)
+using namespace UE::Geometry;
+
+#define LOCTEXT_NAMESPACE "USetCollisionGeometryTool"
+
+
+const FToolTargetTypeRequirements& USetCollisionGeometryToolBuilder::GetTargetRequirements() const
+{
+	static FToolTargetTypeRequirements TypeRequirements({
+		UMeshDescriptionProvider::StaticClass(),
+		UPrimitiveComponentBackedTarget::StaticClass()
+		});
+	return TypeRequirements;
+}
+
+bool USetCollisionGeometryToolBuilder::CanBuildTool(const FToolBuilderState& SceneState) const
+{
+	UActorComponent* LastValidTarget = nullptr;
+	SceneState.TargetManager->EnumerateSelectedAndTargetableComponents(SceneState, GetTargetRequirements(),
+		[&](UActorComponent* Component) { LastValidTarget = Component; });
+	return (LastValidTarget != nullptr && Cast<UStaticMeshComponent>(LastValidTarget) != nullptr);
+}
+
+
+UInteractiveTool* USetCollisionGeometryToolBuilder::BuildTool(const FToolBuilderState& SceneState) const
+{
+	USetCollisionGeometryTool* NewTool = NewObject<USetCollisionGeometryTool>(SceneState.ToolManager);
+	TArray<TObjectPtr<UToolTarget>> Targets = SceneState.TargetManager->BuildAllSelectedTargetable(SceneState, GetTargetRequirements());
+	NewTool->SetTargets(MoveTemp(Targets));
+	return NewTool;
+}
+
+
+void USetCollisionGeometryTool::Setup()
+{
+	UInteractiveTool::Setup();
+
+	// if we have one selection, use it as the source, otherwise use all but the last selected mesh
+	bSourcesHidden = (Targets.Num() > 1);
+	if (Targets.Num() == 1)
+	{
+		SourceObjectIndices.Add(0);
+	}
+	else
+	{
+		for (int32 k = 0; k < Targets.Num() -1; ++k)
+		{
+			SourceObjectIndices.Add(k);
+			UE::ToolTarget::HideSourceObject(Targets[k]);
+		}
+	}
+
+	// collect input meshes
+	InitialSourceMeshes.SetNum(SourceObjectIndices.Num());
+	ParallelFor(SourceObjectIndices.Num(), [&](int32 k)
+	{
+		InitialSourceMeshes[k] = UE::ToolTarget::GetDynamicMeshCopy(Targets[k]);
+	});
+
+	// why is this here?
+	bInputMeshesValid = true;
+
+	UToolTarget* CollisionTarget = Targets[Targets.Num() - 1];
+	PreviewGeom = NewObject<UPreviewGeometry>(this);
+	FTransform PreviewTransform = (FTransform)UE::ToolTarget::GetLocalToWorldTransform(CollisionTarget);
+	OrigTargetTransform = PreviewTransform;
+	TargetScale3D = PreviewTransform.GetScale3D();
+	PreviewTransform.SetScale3D(FVector::OneVector);
+	PreviewGeom->CreateInWorld(UE::ToolTarget::GetTargetActor(CollisionTarget)->GetWorld(), PreviewTransform);
+
+	// initialize initial collision object
+	InitialCollision = MakeShared<FPhysicsDataCollection, ESPMode::ThreadSafe>();
+	InitialCollision->InitializeFromComponent(UE::ToolTarget::GetTargetComponent(CollisionTarget), true);
+	InitialCollision->ExternalScale3D = TargetScale3D;
+
+	// create tool options
+	Settings = NewObject<USetCollisionGeometryToolProperties>(this);
+	Settings->RestoreProperties(this);
+	AddToolPropertySource(Settings);
+	Settings->bUseWorldSpace = (SourceObjectIndices.Num() > 1);
+	Settings->WatchProperty(Settings->InputMode, [this](ESetCollisionGeometryInputMode) { OnInputModeChanged(); });
+	Settings->WatchProperty(Settings->GeometryType, [this](ECollisionGeometryType) { bResultValid = false; });
+	Settings->WatchProperty(Settings->bUseWorldSpace, [this](bool) { bInputMeshesValid = false; });
+	Settings->WatchProperty(Settings->bAppendToExisting, [this](bool) { bResultValid = false; });
+	Settings->WatchProperty(Settings->bRemoveContained, [this](bool) { bResultValid = false; });
+	Settings->WatchProperty(Settings->bEnableMaxCount, [this](bool) { bResultValid = false; });
+	Settings->WatchProperty(Settings->MaxCount, [this](int32) { bResultValid = false; });
+	Settings->WatchProperty(Settings->MinThickness, [this](float) { bResultValid = false; });
+	Settings->WatchProperty(Settings->bDetectBoxes, [this](int32) { bResultValid = false; });
+	Settings->WatchProperty(Settings->bDetectSpheres, [this](int32) { bResultValid = false; });
+	Settings->WatchProperty(Settings->bDetectCapsules, [this](int32) { bResultValid = false; });
+	Settings->WatchProperty(Settings->bSimplifyHulls, [this](bool) { bResultValid = false; });
+	Settings->WatchProperty(Settings->HullTargetFaceCount, [this](int32) { bResultValid = false; });
+	Settings->WatchProperty(Settings->bSimplifyPolygons, [this](bool) { bResultValid = false; });
+	Settings->WatchProperty(Settings->HullTolerance, [this](float) { bResultValid = false; });
+	Settings->WatchProperty(Settings->SweepAxis, [this](EProjectedHullAxis) { bResultValid = false; });
+
+	if (InitialSourceMeshes.Num() == 1)
+	{
+		PolygroupLayerProperties = NewObject<UPolygroupLayersProperties>(this);
+		PolygroupLayerProperties->RestoreProperties(this, TEXT("SetCollisionGeometryTool"));
+		PolygroupLayerProperties->InitializeGroupLayers(&InitialSourceMeshes[0]);
+		PolygroupLayerProperties->WatchProperty(PolygroupLayerProperties->ActiveGroupLayer, [&](FName) { OnSelectedGroupLayerChanged(); });
+		AddToolPropertySource(PolygroupLayerProperties);
+	}
+
+	bResultValid = false;
+
+	VizSettings = NewObject<UCollisionGeometryVisualizationProperties>(this);
+	VizSettings->RestoreProperties(this);
+	AddToolPropertySource(VizSettings);
+	VizSettings->WatchProperty(VizSettings->LineThickness, [this](float NewValue) { bVisualizationDirty = true; });
+	VizSettings->WatchProperty(VizSettings->Color, [this](FColor NewValue) { bVisualizationDirty = true; });
+	VizSettings->WatchProperty(VizSettings->bShowHidden, [this](bool bNewValue) { bVisualizationDirty = true; });
+
+	// add option for collision properties
+	CollisionProps = NewObject<UPhysicsObjectToolPropertySet>(this);
+	AddToolPropertySource(CollisionProps);
+
+	SetToolDisplayName(LOCTEXT("ToolName", "Mesh To Collision"));
+	GetToolManager()->DisplayMessage(
+		LOCTEXT("OnStartTool", "Initialize Simple Collision geometry for a Mesh from one or more input Meshes (including itself)."),
+		EToolMessageLevel::UserNotification);
+}
+
+
+void USetCollisionGeometryTool::Shutdown(EToolShutdownType ShutdownType)
+{
+	VizSettings->SaveProperties(this);
+	Settings->SaveProperties(this);
+	if (PolygroupLayerProperties)
+	{
+		PolygroupLayerProperties->SaveProperties(this, TEXT("SetCollisionGeometryTool"));
+	}
+
+	PreviewGeom->Disconnect();
+
+	// show hidden sources
+	if (bSourcesHidden)
+	{
+		for (int32 k : SourceObjectIndices)
+		{
+			UE::ToolTarget::ShowSourceObject(Targets[k]);
+		}
+	}
+
+	if (ShutdownType == EToolShutdownType::Accept)
+	{
+		// Make sure rendering is done so that we are not changing data being used by collision drawing.
+		FlushRenderingCommands();
+
+		GetToolManager()->BeginUndoTransaction(LOCTEXT("UpdateCollision", "Update Collision"));
+
+		// code below derived from FStaticMeshEditor::DuplicateSelectedPrims(), FStaticMeshEditor::OnCollisionSphere(), and GeomFitUtils.cpp::GenerateSphylAsSimpleCollision()
+
+		UPrimitiveComponent* Component = UE::ToolTarget::GetTargetComponent(Targets[Targets.Num() - 1]);
+		UStaticMeshComponent* StaticMeshComponent = Cast<UStaticMeshComponent>(Component);
+		TObjectPtr<UStaticMesh> StaticMesh = (StaticMeshComponent) ? StaticMeshComponent->GetStaticMesh() : nullptr;
+		UBodySetup* BodySetup = (StaticMesh) ? StaticMesh->GetBodySetup() : nullptr;
+		if (BodySetup != nullptr)
+		{
+			// mark the BodySetup for modification. Do we need to modify the UStaticMesh??
+			BodySetup->Modify();
+
+			// clear existing simple collision. This will call BodySetup->InvalidatePhysicsData()
+			BodySetup->RemoveSimpleCollision();
+
+			// set new collision geometry
+			BodySetup->AggGeom = GeneratedCollision->AggGeom;
+
+			// update collision type
+			BodySetup->CollisionTraceFlag = (ECollisionTraceFlag)(int32)Settings->SetCollisionType;
+
+			// rebuild physics meshes
+			BodySetup->CreatePhysicsMeshes();
+
+			// rebuild nav collision (? StaticMeshEditor does this)
+			StaticMesh->CreateNavCollision(/*bIsUpdate=*/true);
+
+			// update physics state on all components using this StaticMesh
+			for (FThreadSafeObjectIterator Iter(UStaticMeshComponent::StaticClass()); Iter; ++Iter)
+			{
+				UStaticMeshComponent* SMComponent = Cast<UStaticMeshComponent>(*Iter);
+				if (SMComponent->GetStaticMesh() == StaticMesh)
+				{
+					if (SMComponent->IsPhysicsStateCreated())
+					{
+						SMComponent->RecreatePhysicsState();
+					}
+				}
+			}
+
+			// do we need to do a post edit change here??
+
+			// mark static mesh as dirty so it gets resaved?
+			StaticMesh->MarkPackageDirty();
+
+#if WITH_EDITORONLY_DATA
+			// mark the static mesh as having customized collision so it is not regenerated on reimport
+			StaticMesh->bCustomizedCollision = true;
+#endif // WITH_EDITORONLY_DATA
+		}
+
+		// post the undo transaction
+		GetToolManager()->EndUndoTransaction();
+	}
+
+}
+
+
+
+
+void USetCollisionGeometryTool::OnTick(float DeltaTime)
+{
+	if (bInputMeshesValid == false)
+	{
+		PrecomputeInputMeshes();
+		bInputMeshesValid = true;
+		bResultValid = false;
+	}
+
+	if (bResultValid == false)
+	{
+		UpdateGeneratedCollision();
+		bResultValid = true;
+	}
+
+	if (bVisualizationDirty)
+	{
+		UpdateVisualization();
+		bVisualizationDirty = false;
+	}
+}
+
+
+
+void USetCollisionGeometryTool::OnInputModeChanged()
+{
+	if (PolygroupLayerProperties != nullptr)
+	{
+		SetToolPropertySourceEnabled(PolygroupLayerProperties, Settings->InputMode == ESetCollisionGeometryInputMode::PerMeshGroup);
+	}
+	bResultValid = false;
+}
+
+void USetCollisionGeometryTool::OnSelectedGroupLayerChanged()
+{
+	bInputMeshesValid = false;
+	bResultValid = false;
+}
+
+
+void USetCollisionGeometryTool::UpdateActiveGroupLayer()
+{
+	if (InitialSourceMeshes.Num() != 1)
+	{
+		ensure(false);		// should not get here
+		return;
+	}
+	FDynamicMesh3* GroupLayersMesh = &InitialSourceMeshes[0];
+
+	if (PolygroupLayerProperties->HasSelectedPolygroup() == false)
+	{
+		ActiveGroupSet = MakeUnique<UE::Geometry::FPolygroupSet>(GroupLayersMesh);
+	}
+	else
+	{
+		FName SelectedName = PolygroupLayerProperties->ActiveGroupLayer;
+		FDynamicMeshPolygroupAttribute* FoundAttrib = UE::Geometry::FindPolygroupLayerByName(*GroupLayersMesh, SelectedName);
+		ensureMsgf(FoundAttrib, TEXT("Selected Attribute Not Found! Falling back to Default group layer."));
+		ActiveGroupSet = MakeUnique<UE::Geometry::FPolygroupSet>(GroupLayersMesh, FoundAttrib);
+	}
+}
+
+
+
+
+void USetCollisionGeometryTool::UpdateVisualization()
+{
+	float UseThickness = VizSettings->LineThickness;
+	FColor UseColor = VizSettings->Color;
+	PreviewGeom->UpdateAllLineSets([&](ULineSetComponent* LineSet)
+	{
+		LineSet->SetAllLinesThickness(UseThickness);
+		LineSet->SetAllLinesColor(UseColor);
+	});
+
+	LineMaterial = ToolSetupUtil::GetDefaultLineComponentMaterial(GetToolManager(), !VizSettings->bShowHidden);
+	PreviewGeom->SetAllLineSetsMaterial(LineMaterial);
+}
+
+
+void USetCollisionGeometryTool::UpdateGeneratedCollision()
+{
+	// calculate new collision
+	ECollisionGeometryType ComputeType = Settings->GeometryType;
+
+
+	TSharedPtr<FPhysicsDataCollection, ESPMode::ThreadSafe> NewCollision = MakeShared<FPhysicsDataCollection, ESPMode::ThreadSafe>();
+	NewCollision->InitializeFromExisting(*InitialCollision);
+	if (Settings->bAppendToExisting || ComputeType == ECollisionGeometryType::KeepExisting)
+	{
+		NewCollision->CopyGeometryFromExisting(*InitialCollision);
+	}
+
+	TSharedPtr<FMeshSimpleShapeApproximation, ESPMode::ThreadSafe> UseShapeGenerator = GetApproximator(Settings->InputMode);
+
+	UseShapeGenerator->bDetectSpheres = Settings->bDetectSpheres;
+	UseShapeGenerator->bDetectBoxes = Settings->bDetectBoxes;
+	UseShapeGenerator->bDetectCapsules = Settings->bDetectCapsules;
+	//UseShapeGenerator->bDetectConvexes = Settings->bDetectConvexes;
+
+	UseShapeGenerator->MinDimension = Settings->MinThickness;
+
+	switch (ComputeType)
+	{
+	case ECollisionGeometryType::KeepExisting:
+	case ECollisionGeometryType::None:
+		break;
+	case ECollisionGeometryType::AlignedBoxes:
+		UseShapeGenerator->Generate_AlignedBoxes(NewCollision->Geometry);
+		break;
+	case ECollisionGeometryType::OrientedBoxes:
+		UseShapeGenerator->Generate_OrientedBoxes(NewCollision->Geometry);
+		break;
+	case ECollisionGeometryType::MinimalSpheres:
+		UseShapeGenerator->Generate_MinimalSpheres(NewCollision->Geometry);
+		break;
+	case ECollisionGeometryType::Capsules:
+		UseShapeGenerator->Generate_Capsules(NewCollision->Geometry);
+		break;
+	case ECollisionGeometryType::ConvexHulls:
+		UseShapeGenerator->bSimplifyHulls = Settings->bSimplifyHulls;
+		UseShapeGenerator->HullTargetFaceCount = Settings->HullTargetFaceCount;
+		UseShapeGenerator->Generate_ConvexHulls(NewCollision->Geometry);
+		break;
+	case ECollisionGeometryType::SweptHulls:
+		UseShapeGenerator->bSimplifyHulls = Settings->bSimplifyPolygons;
+		UseShapeGenerator->HullSimplifyTolerance = Settings->HullTolerance;
+		UseShapeGenerator->Generate_ProjectedHulls(NewCollision->Geometry, 
+			(FMeshSimpleShapeApproximation::EProjectedHullAxisMode)(int32)Settings->SweepAxis);
+		break;
+	case ECollisionGeometryType::MinVolume:
+		UseShapeGenerator->Generate_MinVolume(NewCollision->Geometry);
+		break;
+	}
+
+
+	if (!NewCollision)
+	{
+		ensure(false);
+		return;
+	}
+	GeneratedCollision = NewCollision;
+
+	if (Settings->bRemoveContained)
+	{
+		GeneratedCollision->Geometry.RemoveContainedGeometry();
+	}
+
+	bool bUseMaxCount = (Settings->bEnableMaxCount);
+	if (bUseMaxCount)
+	{
+		GeneratedCollision->Geometry.FilterByVolume(Settings->MaxCount);
+	}
+
+	GeneratedCollision->CopyGeometryToAggregate();
+
+	// update visualization
+	PreviewGeom->RemoveAllLineSets();
+	UE::PhysicsTools::InitializePreviewGeometryLines(*GeneratedCollision, PreviewGeom,
+		VizSettings->Color, VizSettings->LineThickness, 0.0f, 16);
+
+	// update property set
+	CollisionProps->Reset();
+	UE::PhysicsTools::InitializePhysicsToolObjectPropertySet(GeneratedCollision.Get(), CollisionProps);
+}
+
+
+
+
+
+void USetCollisionGeometryTool::InitializeDerivedMeshSet(
+	const TArray<TSharedPtr<FDynamicMesh3, ESPMode::ThreadSafe>>& FromInputMeshes,
+	TArray<TSharedPtr<FDynamicMesh3, ESPMode::ThreadSafe>>& ToMeshes,
+	TFunctionRef<bool(const FDynamicMesh3* Mesh, int32 Tri0, int32 Tri1)> TrisConnectedPredicate)
+{
+	// find connected-components on input meshes, under given connectivity predicate
+	TArray<TUniquePtr<FMeshConnectedComponents>> ComponentSets;
+	ComponentSets.SetNum(FromInputMeshes.Num());
+	ParallelFor(FromInputMeshes.Num(), [&](int32 k)
+	{
+		const FDynamicMesh3* Mesh = FromInputMeshes[k].Get();
+		ComponentSets[k] = MakeUnique<FMeshConnectedComponents>(Mesh);
+		ComponentSets[k]->FindConnectedTriangles(
+			[Mesh, &TrisConnectedPredicate](int32 Tri0, int32 Tri1) 
+			{ 
+				return TrisConnectedPredicate(Mesh, Tri0, Tri1); 
+			}
+		);
+	});
+
+	// Assemble a list of all the submeshes we want to compute, so we can do them all in parallel
+	struct FSubmeshSource
+	{
+		const FDynamicMesh3* SourceMesh;
+		FIndex2i ComponentIdx;
+	};
+	TArray<FSubmeshSource> AllSubmeshes;
+	for (int32 k = 0; k < FromInputMeshes.Num(); ++k)
+	{
+		const FDynamicMesh3* Mesh = FromInputMeshes[k].Get();
+		int32 NumComponents = ComponentSets[k]->Num();
+		for ( int32 j = 0; j < NumComponents; ++j )
+		{
+			const FMeshConnectedComponents::FComponent& Component = ComponentSets[k]->GetComponent(j);
+			if (Component.Indices.Num() > 1)		// ignore single triangles
+			{
+				AllSubmeshes.Add(FSubmeshSource{ Mesh, FIndex2i(k,j) });
+			}
+		}
+	}
+
+
+	// compute all the submeshes
+	ToMeshes.Reset();
+	ToMeshes.SetNum(AllSubmeshes.Num());
+	ParallelFor(AllSubmeshes.Num(), [&](int32 k)
+	{
+		const FSubmeshSource& Source = AllSubmeshes[k];
+		const FMeshConnectedComponents::FComponent& Component = ComponentSets[Source.ComponentIdx.A]->GetComponent(Source.ComponentIdx.B);
+		FDynamicSubmesh3 Submesh(Source.SourceMesh, Component.Indices, (int32)EMeshComponents::None, false);
+		ToMeshes[k] = MakeShared<FDynamicMesh3, ESPMode::ThreadSafe>( MoveTemp(Submesh.GetSubmesh()) );
+	});
+}
+
+
+template<typename T>
+TArray<const T*> MakeRawPointerList(const TArray<TSharedPtr<T, ESPMode::ThreadSafe>>& InputList)
+{
+	TArray<const T*> Result;
+	Result.Reserve(InputList.Num());
+	for (const TSharedPtr<T, ESPMode::ThreadSafe>& Ptr : InputList)
+	{
+		Result.Add(Ptr.Get());
+	}
+	return MoveTemp(Result);
+}
+
+
+void USetCollisionGeometryTool::PrecomputeInputMeshes()
+{
+	if (InitialSourceMeshes.Num() == 1)
+	{
+		UpdateActiveGroupLayer();
+	}
+
+	UToolTarget* CollisionTarget = Targets[Targets.Num() - 1];
+	UE::Geometry::FTransform3d TargetTransform(UE::ToolTarget::GetLocalToWorldTransform(CollisionTarget));
+	UE::Geometry::FTransform3d TargetTransformInv = TargetTransform.Inverse();
+
+	InputMeshes.Reset();
+	InputMeshes.SetNum(SourceObjectIndices.Num());
+	ParallelFor(SourceObjectIndices.Num(), [&](int32 k)
+	{
+		FDynamicMesh3 SourceMesh = InitialSourceMeshes[k];
+		if (Settings->bUseWorldSpace)
+		{
+			UE::Geometry::FTransform3d ToWorld(UE::ToolTarget::GetLocalToWorldTransform(Targets[k]));
+			MeshTransforms::ApplyTransform(SourceMesh, ToWorld);
+			MeshTransforms::ApplyTransform(SourceMesh, TargetTransformInv);
+		}
+		SourceMesh.DiscardAttributes();
+		InputMeshes[k] = MakeShared<FDynamicMesh3, ESPMode::ThreadSafe>(MoveTemp(SourceMesh));
+	});
+	InputMeshesApproximator = MakeShared<FMeshSimpleShapeApproximation, ESPMode::ThreadSafe>();
+	InputMeshesApproximator->InitializeSourceMeshes(MakeRawPointerList<FDynamicMesh3>(InputMeshes));
+
+
+	// build combined input
+	CombinedInputMeshes.Reset();
+	FDynamicMesh3 CombinedMesh;
+	CombinedMesh.EnableTriangleGroups();
+	FDynamicMeshEditor Appender(&CombinedMesh);
+	FMeshIndexMappings TmpMappings;
+	for (const TSharedPtr<FDynamicMesh3, ESPMode::ThreadSafe>& InputMesh : InputMeshes)
+	{
+		TmpMappings.Reset();
+		Appender.AppendMesh(InputMesh.Get(), TmpMappings);
+	}
+	CombinedInputMeshes.Add( MakeShared<FDynamicMesh3, ESPMode::ThreadSafe>(MoveTemp(CombinedMesh)) );
+	CombinedInputMeshesApproximator = MakeShared<FMeshSimpleShapeApproximation, ESPMode::ThreadSafe>();
+	CombinedInputMeshesApproximator->InitializeSourceMeshes(MakeRawPointerList<FDynamicMesh3>(CombinedInputMeshes));
+
+	// build separated input meshes
+	SeparatedInputMeshes.Reset();
+	InitializeDerivedMeshSet(InputMeshes, SeparatedInputMeshes, 
+		[&](const FDynamicMesh3* Mesh, int32 Tri0, int32 Tri1) { return true; });
+	SeparatedMeshesApproximator = MakeShared<FMeshSimpleShapeApproximation, ESPMode::ThreadSafe>();
+	SeparatedMeshesApproximator->InitializeSourceMeshes(MakeRawPointerList<FDynamicMesh3>(SeparatedInputMeshes));
+
+	// build per-group input meshes
+	PerGroupInputMeshes.Reset();
+	if (ActiveGroupSet.IsValid())
+	{
+		check(InputMeshes.Num() == 1);
+		InitializeDerivedMeshSet(InputMeshes, PerGroupInputMeshes,
+			[&](const FDynamicMesh3* Mesh, int32 Tri0, int32 Tri1) { return ActiveGroupSet->GetTriangleGroup(Tri0) == ActiveGroupSet->GetTriangleGroup(Tri1); });
+	}
+	else
+	{
+		InitializeDerivedMeshSet(InputMeshes, PerGroupInputMeshes,
+			[&](const FDynamicMesh3* Mesh, int32 Tri0, int32 Tri1) { return Mesh->GetTriangleGroup(Tri0) == Mesh->GetTriangleGroup(Tri1); });
+	}
+	PerGroupMeshesApproximator = MakeShared<FMeshSimpleShapeApproximation, ESPMode::ThreadSafe>();
+	PerGroupMeshesApproximator->InitializeSourceMeshes(MakeRawPointerList<FDynamicMesh3>(PerGroupInputMeshes));
+
+}
+
+
+TSharedPtr<FMeshSimpleShapeApproximation, ESPMode::ThreadSafe>& USetCollisionGeometryTool::GetApproximator(ESetCollisionGeometryInputMode MeshSetMode)
+{
+	if (MeshSetMode == ESetCollisionGeometryInputMode::CombineAll)
+	{
+		return CombinedInputMeshesApproximator;
+	}
+	else if ( MeshSetMode == ESetCollisionGeometryInputMode::PerMeshComponent)
+	{
+		return SeparatedMeshesApproximator;
+	}
+	else if (MeshSetMode == ESetCollisionGeometryInputMode::PerMeshGroup)
+	{
+		return PerGroupMeshesApproximator;
+	}
+	else
+	{
+		return InputMeshesApproximator;
+	}
+}
+
+
+#undef LOCTEXT_NAMESPACE
\ No newline at end of file
diff --git a/Engine/Plugins/Experimental/MeshModelingToolsetExp/Source/MeshModelingToolsExp/Private/ProjectToTargetTool.cpp b/Engine/Plugins/Experimental/MeshModelingToolsetExp/Source/MeshModelingToolsExp/Private/ProjectToTargetTool.cpp
new file mode 100644
--- /dev/null
+++ b/Engine/Plugins/Experimental/MeshModelingToolsetExp/Source/MeshModelingToolsExp/Private/ProjectToTargetTool.cpp
@@ -0,0 +1,100 @@
+// Copyright Epic Games, Inc. All Rights Reserved.
+
+#include "ProjectToTargetTool.h"
+#include "MeshDescriptionToDynamicMesh.h"
+#include "InteractiveToolManager.h"
+
+#include "TargetInterfaces/MaterialProvider.h"
+#include "TargetInterfaces/MeshDescriptionCommitter.h"
+#include "TargetInterfaces/MeshDescriptionProvider.h"
+#include "TargetInterfaces/PrimitiveComponentBackedTarget.h"
+#include "ToolTargetManager.h"
+
+#include "ExplicitUseGeometryMathTypes.h"		// using UE::Geometry::(math types)
+using namespace UE::Geometry;
+
+#define LOCTEXT_NAMESPACE "UProjectToTargetTool"
+
+const FToolTargetTypeRequirements& UProjectToTargetToolBuilder::GetTargetRequirements() const
+{
+	static FToolTargetTypeRequirements TypeRequirements({
+		UMeshDescriptionCommitter::StaticClass(),
+		UMeshDescriptionProvider::StaticClass(),
+		UPrimitiveComponentBackedTarget::StaticClass(),
+		UMaterialProvider::StaticClass()
+		});
+	return TypeRequirements;
+}
+
+bool UProjectToTargetToolBuilder::CanBuildTool(const FToolBuilderState& SceneState) const
+{
+	return SceneState.TargetManager->CountSelectedAndTargetable(SceneState, GetTargetRequirements()) == 2;
+}
+
+UInteractiveTool* UProjectToTargetToolBuilder::BuildTool(const FToolBuilderState& SceneState) const
+{
+	UProjectToTargetTool* NewTool = NewObject<UProjectToTargetTool>(SceneState.ToolManager);
+
+	TArray<TObjectPtr<UToolTarget>> Targets = SceneState.TargetManager->BuildAllSelectedTargetable(SceneState, GetTargetRequirements());
+	NewTool->SetTargets(MoveTemp(Targets));
+	NewTool->SetWorld(SceneState.World);
+
+	return NewTool;
+}
+
+void UProjectToTargetTool::Setup()
+{
+	// ProjectionTarget and ProjectionTargetSpatial are setup before calling the parent class's Setup
+	FMeshDescriptionToDynamicMesh ProjectionConverter;
+	check(Targets.Num() == 2);
+	ProjectionTarget = MakeUnique<FDynamicMesh3>();
+	ProjectionConverter.Convert(TargetMeshProviderInterface(1)->GetMeshDescription(), *ProjectionTarget);
+	ProjectionTargetSpatial = MakeUnique<FDynamicMeshAABBTree3>(ProjectionTarget.Get(), true);
+
+	// Now setup parent RemeshMeshTool class
+	URemeshMeshTool::Setup();
+
+	SetToolDisplayName(LOCTEXT("ProjectToTargetToolName", "Project To Target"));
+	GetToolManager()->DisplayMessage(
+		LOCTEXT("ProjectToTargetToolDescription", "Incrementally deform the first selected mesh towards the second, while applying Remeshing. This can be used to improve the accuracy of shrink-wrapping strategies."),
+		EToolMessageLevel::UserNotification);
+}
+
+
+TUniquePtr<FDynamicMeshOperator> UProjectToTargetTool::MakeNewOperator()
+{
+	UProjectToTargetToolProperties* ProjectProperties = Cast<UProjectToTargetToolProperties>(BasicProperties);
+	if (!ensure(ProjectProperties))
+	{
+		return nullptr;
+	}
+
+	TUniquePtr<FDynamicMeshOperator> Op = URemeshMeshTool::MakeNewOperator();
+
+	FDynamicMeshOperator* RawOp = Op.Get();
+	FRemeshMeshOp* RemeshOp = static_cast<FRemeshMeshOp*>(RawOp);
+	check(RemeshOp);
+
+	RemeshOp->ProjectionTarget = ProjectionTarget.Get();
+	RemeshOp->ProjectionTargetSpatial = ProjectionTargetSpatial.Get();
+
+	RemeshOp->ToolMeshLocalToWorld = FTransform3d(Cast<IPrimitiveComponentBackedTarget>(Targets[0])->GetWorldTransform());
+	RemeshOp->TargetMeshLocalToWorld = FTransform3d(Cast<IPrimitiveComponentBackedTarget>(Targets[1])->GetWorldTransform());
+	RemeshOp->bUseWorldSpace = ProjectProperties->bWorldSpace;
+	RemeshOp->bParallel = ProjectProperties->bParallel;
+
+	if (ProjectProperties->RemeshType == ERemeshType::NormalFlow)
+	{
+		RemeshOp->FaceProjectionPassesPerRemeshIteration = ProjectProperties->FaceProjectionPassesPerRemeshIteration;
+		RemeshOp->SurfaceProjectionSpeed = ProjectProperties->SurfaceProjectionSpeed;
+		RemeshOp->NormalAlignmentSpeed = ProjectProperties->NormalAlignmentSpeed;
+		RemeshOp->bSmoothInFillAreas = ProjectProperties->bSmoothInFillAreas;
+		RemeshOp->FillAreaDistanceMultiplier = ProjectProperties->FillAreaDistanceMultiplier;
+		RemeshOp->FillAreaSmoothMultiplier = ProjectProperties->FillAreaSmoothMultiplier;
+	}
+
+	return Op;
+}
+
+
+#undef LOCTEXT_NAMESPACE
diff --git a/Engine/Plugins/Experimental/MeshModelingToolsetExp/Source/MeshModelingToolsExp/Private/Properties/MeshAnalysisProperties.cpp b/Engine/Plugins/Experimental/MeshModelingToolsetExp/Source/MeshModelingToolsExp/Private/Properties/MeshAnalysisProperties.cpp
new file mode 100644
--- /dev/null
+++ b/Engine/Plugins/Experimental/MeshModelingToolsetExp/Source/MeshModelingToolsExp/Private/Properties/MeshAnalysisProperties.cpp
@@ -0,0 +1,24 @@
+// Copyright Epic Games, Inc. All Rights Reserved.
+
+#include "Properties/MeshAnalysisProperties.h"
+
+#include "DynamicMesh/DynamicMesh3.h"
+#include "MeshQueries.h"
+#include "MeshAdapter.h"
+#include "DynamicMesh/MeshAdapterUtil.h"
+
+#define LOCTEXT_NAMESPACE "UMeshAnalysisProperites"
+
+using namespace UE::Geometry;
+
+void UMeshAnalysisProperties::Update(const FDynamicMesh3& MeshIn, const FTransform& Transform)
+{
+	FTriangleMeshAdapterd TransformedMesh = UE::Geometry::MakeTransformedDynamicMeshAdapter(&MeshIn, Transform);
+	FVector2d VolArea = TMeshQueries<FTriangleMeshAdapterd>::GetVolumeArea(TransformedMesh);
+	this->SurfaceArea = FString::Printf(TEXT("%.2f m^2"), VolArea.Y / 10000);
+	this->Volume = FString::Printf(TEXT("%.2f m^3"), VolArea.X / 1000000);
+}
+
+
+
+#undef LOCTEXT_NAMESPACE
diff --git a/Engine/Plugins/Experimental/MeshModelingToolsetExp/Source/MeshModelingToolsExp/Public/AddPatchTool.h b/Engine/Plugins/Experimental/MeshModelingToolsetExp/Source/MeshModelingToolsExp/Public/AddPatchTool.h
new file mode 100644
--- /dev/null
+++ b/Engine/Plugins/Experimental/MeshModelingToolsetExp/Source/MeshModelingToolsExp/Public/AddPatchTool.h
@@ -0,0 +1,121 @@
+// Copyright Epic Games, Inc. All Rights Reserved.
+
+#pragma once
+
+#include "CoreMinimal.h"
+#include "MaterialShared.h"
+#include "UObject/NoExportTypes.h"
+#include "InteractiveToolBuilder.h"
+#include "BaseTools/SingleClickTool.h"
+#include "PreviewMesh.h"
+#include "Properties/MeshMaterialProperties.h"
+#include "AddPatchTool.generated.h"
+
+PREDECLARE_USE_GEOMETRY_CLASS(FDynamicMesh3);
+
+/**
+ *
+ */
+UCLASS()
+class MESHMODELINGTOOLSEXP_API UAddPatchToolBuilder : public UInteractiveToolBuilder
+{
+	GENERATED_BODY()
+public:
+	virtual bool CanBuildTool(const FToolBuilderState& SceneState) const override;
+	virtual UInteractiveTool* BuildTool(const FToolBuilderState& SceneState) const override;
+};
+
+
+
+UCLASS()
+class MESHMODELINGTOOLSEXP_API UAddPatchToolProperties : public UInteractiveToolPropertySet
+{
+	GENERATED_BODY()
+
+public:
+	UAddPatchToolProperties();
+
+	/** Width of Shape */
+	UPROPERTY(EditAnywhere, Category = PatchSettings, meta = (DisplayName = "Width", UIMin = "1.0", UIMax = "1000.0", ClampMin = "0.0001", ClampMax = "1000000.0"))
+	float Width;
+
+	/** Rotation around up axis */
+	UPROPERTY(EditAnywhere, Category = PatchSettings, meta = (DisplayName = "Rotation", UIMin = "0.0", UIMax = "360.0"))
+	float Rotation;
+
+	/** Subdivisions */
+	UPROPERTY(EditAnywhere, Category = PatchSettings, meta = (DisplayName = "Subdivisions", UIMin = "0", UIMax = "100", ClampMin = "0", ClampMax = "4000"))
+	int Subdivisions;
+
+	/** Rotation around up axis */
+	UPROPERTY(EditAnywhere, Category = PatchSettings, meta = (DisplayName = "Shift", UIMin = "-1000", UIMax = "1000"))
+	float Shift;
+};
+
+
+
+
+
+
+
+/**
+ *
+ */
+UCLASS()
+class MESHMODELINGTOOLSEXP_API UAddPatchTool : public USingleClickTool, public IHoverBehaviorTarget
+{
+	GENERATED_BODY()
+
+public:
+	virtual void SetWorld(UWorld* World);
+
+	virtual void Setup() override;
+	virtual void Shutdown(EToolShutdownType ShutdownType) override;
+
+	virtual void Render(IToolsContextRenderAPI* RenderAPI) override;
+	virtual void OnTick(float DeltaTime) override;
+
+	virtual bool HasCancel() const override { return false; }
+	virtual bool HasAccept() const override { return false; }
+	virtual bool CanAccept() const override { return false; }
+
+	virtual void OnPropertyModified(UObject* PropertySet, FProperty* Property) override;
+
+	virtual void OnClicked(const FInputDeviceRay& ClickPos) override;
+
+
+	// IHoverBehaviorTarget interface
+	virtual FInputRayHit BeginHoverSequenceHitTest(const FInputDeviceRay& PressPos) override;
+	virtual void OnBeginHover(const FInputDeviceRay& DevicePos) override;
+	virtual bool OnUpdateHover(const FInputDeviceRay& DevicePos) override;
+	virtual void OnEndHover() override;
+
+
+protected:
+	UPROPERTY()
+	TObjectPtr<UAddPatchToolProperties> ShapeSettings;
+
+	UPROPERTY()
+	TObjectPtr<UNewMeshMaterialProperties> MaterialProperties;
+
+
+
+	UPROPERTY()
+	TObjectPtr<UPreviewMesh> PreviewMesh;
+
+protected:
+	UWorld* TargetWorld;
+
+	FBox WorldBounds;
+
+	UE::Geometry::FFrame3f ShapeFrame;
+	bool bPreviewValid = true;
+
+	void UpdatePreviewPosition(const FInputDeviceRay& ClickPos);
+	void UpdatePreviewMesh();
+
+	TUniquePtr<UE::Geometry::FDynamicMesh3> BaseMesh;
+	void GeneratePreviewBaseMesh();
+
+	void GeneratePlane(UE::Geometry::FDynamicMesh3* OutMesh);
+};
diff --git a/Engine/Plugins/Experimental/MeshModelingToolsetExp/Source/MeshModelingToolsExp/Public/AlignObjectsTool.h b/Engine/Plugins/Experimental/MeshModelingToolsetExp/Source/MeshModelingToolsExp/Public/AlignObjectsTool.h
new file mode 100644
--- /dev/null
+++ b/Engine/Plugins/Experimental/MeshModelingToolsetExp/Source/MeshModelingToolsExp/Public/AlignObjectsTool.h
@@ -0,0 +1,165 @@
+// Copyright Epic Games, Inc. All Rights Reserved.
+
+#pragma once
+
+#include "CoreMinimal.h"
+#include "MultiSelectionTool.h"
+#include "InteractiveToolBuilder.h"
+#include "BaseBehaviors/BehaviorTargetInterfaces.h"
+#include "Changes/TransformChange.h"
+#include "FrameTypes.h"
+#include "BoxTypes.h"
+#include "AlignObjectsTool.generated.h"
+
+class UPrimitiveComponent;
+
+/**
+ *
+ */
+UCLASS()
+class MESHMODELINGTOOLSEXP_API UAlignObjectsToolBuilder : public UInteractiveToolBuilder
+{
+	GENERATED_BODY()
+
+public:
+	virtual bool CanBuildTool(const FToolBuilderState& SceneState) const override;
+	virtual UInteractiveTool* BuildTool(const FToolBuilderState& SceneState) const override;
+
+protected:
+	virtual const FToolTargetTypeRequirements& GetTargetRequirements() const override;
+};
+
+
+UENUM()
+enum class EAlignObjectsAlignTypes
+{
+	Pivots,
+	BoundingBoxes
+};
+
+UENUM()
+enum class EAlignObjectsAlignToOptions
+{
+	FirstSelected,
+	LastSelected,
+	Combined
+};
+
+UENUM()
+enum class EAlignObjectsBoxPoint
+{
+	Center,
+	Bottom,
+	Top,
+	Left,
+	Right,
+	Front,
+	Back,
+	Min,
+	Max
+};
+
+
+
+/**
+ * Standard properties of the Align Objects Operation
+ */
+UCLASS()
+class MESHMODELINGTOOLSEXP_API UAlignObjectsToolProperties : public UInteractiveToolPropertySet
+{
+	GENERATED_BODY()
+public:
+
+	UPROPERTY(EditAnywhere, Category = Options)
+	EAlignObjectsAlignTypes AlignType = EAlignObjectsAlignTypes::BoundingBoxes;
+
+	UPROPERTY(EditAnywhere, Category = Options)
+	EAlignObjectsAlignToOptions AlignTo = EAlignObjectsAlignToOptions::Combined;
+
+	UPROPERTY(EditAnywhere, Category = Options, meta = (EditCondition = "AlignType == EAlignObjectsAlignTypes::BoundingBoxes || AlignTo == EAlignObjectsAlignToOptions::Combined"))
+	EAlignObjectsBoxPoint BoxPosition = EAlignObjectsBoxPoint::Center;
+
+	UPROPERTY(EditAnywhere, Category = Axes)
+	bool bAlignX = false;
+
+	UPROPERTY(EditAnywhere, Category = Axes)
+	bool bAlignY = false;
+
+	UPROPERTY(EditAnywhere, Category = Axes)
+	bool bAlignZ = true;
+};
+
+
+
+
+
+/**
+ * UAlignObjectsTool transforms the input Components so that they are aligned in various ways, depending on the current settings.
+ * The object positions move after every change in the parameters. Currently those changes are not transacted.
+ * On cancel the original positions are restored, and on accept the positions are updated with a transaction.
+ */
+UCLASS()
+class MESHMODELINGTOOLSEXP_API UAlignObjectsTool : public UMultiSelectionTool, public IClickDragBehaviorTarget
+{
+	GENERATED_BODY()
+	using FTransform3d = UE::Geometry::FTransform3d;
+	using FAxisAlignedBox3d = UE::Geometry::FAxisAlignedBox3d;
+public:
+	UAlignObjectsTool();
+
+	virtual void RegisterActions(FInteractiveToolActionSet& ActionSet) override;
+
+	virtual void SetWorld(UWorld* World, UInteractiveGizmoManager* GizmoManager);
+
+	virtual void Setup() override;
+	virtual void Shutdown(EToolShutdownType ShutdownType) override;
+
+	virtual void OnTick(float DeltaTime) override;
+	virtual void Render(IToolsContextRenderAPI* RenderAPI) override;
+
+	virtual bool HasCancel() const override { return true; }
+	virtual bool HasAccept() const override { return true; }
+	virtual bool CanAccept() const override { return true; }
+
+	virtual void OnPropertyModified(UObject* PropertySet, FProperty* Property) override;
+
+
+	// ICLickDragBehaviorTarget interface
+	virtual FInputRayHit CanBeginClickDragSequence(const FInputDeviceRay& PressPos) override;
+	virtual void OnClickPress(const FInputDeviceRay& PressPos) override;
+	virtual void OnClickDrag(const FInputDeviceRay& DragPos) override;
+	virtual void OnClickRelease(const FInputDeviceRay& ReleasePos) override;
+	virtual void OnTerminateDragSequence() override;
+
+public:
+	UPROPERTY()
+	TObjectPtr<UAlignObjectsToolProperties> AlignProps;
+
+
+protected:
+	UWorld* TargetWorld;
+	UInteractiveGizmoManager* GizmoManager;
+
+
+	struct FAlignInfo
+	{
+		UPrimitiveComponent* Component;
+		FTransform SavedTransform;
+		FTransform3d WorldTransform;
+		FVector3d WorldPivot;
+		FAxisAlignedBox3d WorldBounds;
+	};
+
+	TArray<FAlignInfo> ComponentInfo;
+
+	FAxisAlignedBox3d CombinedBounds;
+	FAxisAlignedBox3d PivotBounds;
+	FVector3d AveragePivot;
+
+	void Precompute();
+
+	bool bAlignDirty = false;
+	void UpdateAlignment();
+	void UpdateAlignment_Pivots();
+	void UpdateAlignment_BoundingBoxes();
+};
diff --git a/Engine/Plugins/Experimental/MeshModelingToolsetExp/Source/MeshModelingToolsExp/Public/BakeMeshAttributeMapsTool.h b/Engine/Plugins/Experimental/MeshModelingToolsetExp/Source/MeshModelingToolsExp/Public/BakeMeshAttributeMapsTool.h
new file mode 100644
--- /dev/null
+++ b/Engine/Plugins/Experimental/MeshModelingToolsetExp/Source/MeshModelingToolsExp/Public/BakeMeshAttributeMapsTool.h
@@ -0,0 +1,297 @@
+// Copyright Epic Games, Inc. All Rights Reserved.
+
+#pragma once
+
+#include "CoreMinimal.h"
+#include "MultiSelectionTool.h"
+#include "InteractiveToolBuilder.h"
+#include "DynamicMesh/DynamicMesh3.h"
+#include "DynamicMesh/DynamicMeshAABBTree3.h"
+#include "Image/ImageDimensions.h"
+#include "Image/ImageBuilder.h"
+#include "Sampling/MeshMapBaker.h"
+#include "ModelingOperators.h"
+#include "MeshOpPreviewHelpers.h"
+#include "BakeMeshAttributeToolCommon.h"
+#include "BakeMeshAttributeMapsTool.generated.h"
+
+
+// predeclarations
+struct FMeshDescription;
+class UDynamicMeshComponent;
+class UMaterialInstanceDynamic;
+class UTexture2D;
+PREDECLARE_GEOMETRY(template<typename RealType> class TMeshTangents);
+PREDECLARE_GEOMETRY(class FMeshImageBakingCache);
+using UE::Geometry::FImageDimensions;
+class IPrimitiveComponentBackedTarget;
+class IMeshDescriptionProvider;
+class IMaterialProvider;
+class FBakeNormalMapOp;
+class FBakeOcclusionMapOp;
+class FBakeCurvatureMapOp;
+class FBakeMeshPropertyMapOp;
+class FBakeTexture2DImageMapOp;
+
+/**
+ * Tool Builder
+ */
+UCLASS()
+class MESHMODELINGTOOLSEXP_API UBakeMeshAttributeMapsToolBuilder : public UInteractiveToolBuilder
+{
+	GENERATED_BODY()
+
+public:
+	virtual bool CanBuildTool(const FToolBuilderState& SceneState) const override;
+	virtual UInteractiveTool* BuildTool(const FToolBuilderState& SceneState) const override;
+
+protected:
+	virtual const FToolTargetTypeRequirements& GetTargetRequirements() const override;
+};
+
+
+
+UENUM(meta=(Bitflags, UseEnumValuesAsMaskValuesInEditor="true"))
+enum class EBakeMapType
+{
+	None                   = 0 UMETA(Hidden),
+	TangentSpaceNormalMap  = 1 << 0,
+	AmbientOcclusion       = 1 << 1,
+	BentNormal             = 1 << 2,
+	Curvature              = 1 << 3,
+	Texture2DImage         = 1 << 4,
+	NormalImage            = 1 << 5,
+	FaceNormalImage        = 1 << 6,
+	PositionImage          = 1 << 7,
+	MaterialID             = 1 << 8,
+	MultiTexture           = 1 << 9,
+	VertexColorImage       = 1 << 10,
+	Occlusion              = (AmbientOcclusion | BentNormal) UMETA(Hidden),
+	All                    = 0x7FF UMETA(Hidden)
+};
+ENUM_CLASS_FLAGS(EBakeMapType);
+
+
+UENUM()
+enum class EBakeTextureResolution 
+{
+	Resolution16 = 16 UMETA(DisplayName = "16 x 16"),
+	Resolution32 = 32 UMETA(DisplayName = "32 x 32"),
+	Resolution64 = 64 UMETA(DisplayName = "64 x 64"),
+	Resolution128 = 128 UMETA(DisplayName = "128 x 128"),
+	Resolution256 = 256 UMETA(DisplayName = "256 x 256"),
+	Resolution512 = 512 UMETA(DisplayName = "512 x 512"),
+	Resolution1024 = 1024 UMETA(DisplayName = "1024 x 1024"),
+	Resolution2048 = 2048 UMETA(DisplayName = "2048 x 2048"),
+	Resolution4096 = 4096 UMETA(DisplayName = "4096 x 4096"),
+	Resolution8192 = 8192 UMETA(DisplayName = "8192 x 8192")
+};
+
+
+UENUM()
+enum class EBakeMultisampling
+{
+	None = 1 UMETA(DisplayName = "None"),
+	Sample2x2 = 2 UMETA(DisplayName = "2 x 2"),
+	Sample4x4 = 4 UMETA(DisplayName = "4 x 4"),
+	Sample8x8 = 8 UMETA(DisplayName = "8 x 8"),
+	Sample16x16 = 16 UMETA(DisplayName = "16 x 16")
+};
+
+
+UCLASS()
+class MESHMODELINGTOOLSEXP_API UBakeMeshAttributeMapsToolProperties : public UInteractiveToolPropertySet
+{
+	GENERATED_BODY()
+
+public:
+	/** The map types to generate */
+	UPROPERTY(EditAnywhere, Category = MapSettings, meta=(Bitmask, BitmaskEnum=EBakeMapType))
+	int32 MapTypes = (int32) EBakeMapType::None;
+
+	/** The map type index to preview */
+	UPROPERTY(EditAnywhere, Category = MapSettings, meta=(ArrayClamp="Result"))
+	int MapPreview = 0;
+
+	/** The pixel resolution of the generated map */
+	UPROPERTY(EditAnywhere, Category = MapSettings, meta = (TransientToolProperty))
+	EBakeTextureResolution Resolution = EBakeTextureResolution::Resolution256;
+
+	/** The multisampling configuration per texel */
+	UPROPERTY(EditAnywhere, Category = MapSettings)
+	EBakeMultisampling Multisampling = EBakeMultisampling::None;
+
+	UPROPERTY(EditAnywhere, Category = MapSettings)
+	bool bUseWorldSpace = false;
+
+	/** Distance to search for the correspondence between the source and target meshes */
+	UPROPERTY(EditAnywhere, Category = MapSettings, meta = (ClampMin = "0.001"))
+	float Thickness = 3.0;
+
+	/** Which UV layer to use to create the map */
+	UPROPERTY(EditAnywhere, Category = MapSettings, meta = (GetOptions = GetUVLayerNamesFunc))
+	FString UVLayer;
+
+	UFUNCTION()
+	TArray<FString> GetUVLayerNamesFunc();
+	UPROPERTY(meta = (TransientToolProperty))
+	TArray<FString> UVLayerNamesList;
+
+	UPROPERTY(VisibleAnywhere, Category = MapSettings, meta = (TransientToolProperty))
+	TArray<TObjectPtr<UTexture2D>> Result;
+
+};
+
+
+/**
+ * Detail Map Baking Tool
+ */
+UCLASS()
+class MESHMODELINGTOOLSEXP_API UBakeMeshAttributeMapsTool : public UMultiSelectionTool, public UE::Geometry::IGenericDataOperatorFactory<UE::Geometry::FMeshMapBaker>
+{
+	GENERATED_BODY()
+
+public:
+	UBakeMeshAttributeMapsTool() = default;
+
+	virtual void Setup() override;
+	virtual void Shutdown(EToolShutdownType ShutdownType) override;
+
+	virtual void OnTick(float DeltaTime) override;
+	virtual void Render(IToolsContextRenderAPI* RenderAPI) override;
+
+	virtual bool HasCancel() const override { return true; }
+	virtual bool HasAccept() const override { return true; }
+	virtual bool CanAccept() const override;
+
+	// IGenericDataOperatorFactory API
+	virtual TUniquePtr<UE::Geometry::TGenericDataOperator<UE::Geometry::FMeshMapBaker>> MakeNewOperator() override;
+
+protected:
+	// need to update bResultValid if these are modified, so we don't publicly expose them. 
+	// @todo setters/getters for these
+
+	UPROPERTY()
+	TObjectPtr<UBakeMeshAttributeMapsToolProperties> Settings;
+
+	UPROPERTY()
+	TObjectPtr<UBakedNormalMapToolProperties> NormalMapProps;
+
+	UPROPERTY()
+	TObjectPtr<UBakedOcclusionMapToolProperties> OcclusionMapProps;
+
+	UPROPERTY()
+	TObjectPtr<UBakedCurvatureMapToolProperties> CurvatureMapProps;
+
+	UPROPERTY()
+	TObjectPtr<UBakedTexture2DImageProperties> Texture2DProps;
+
+	UPROPERTY()
+	TObjectPtr<UBakedMultiTexture2DImageProperties> MultiTextureProps;
+
+	UPROPERTY()
+	TObjectPtr<UBakedOcclusionMapVisualizationProperties> VisualizationProps;
+
+
+protected:
+	friend class FMeshMapBakerOp;
+	
+	UDynamicMeshComponent* DynamicMeshComponent;
+
+	UPROPERTY()
+	TObjectPtr<UMaterialInstanceDynamic> PreviewMaterial;
+
+    UPROPERTY()
+	TObjectPtr<UMaterialInstanceDynamic> BentNormalPreviewMaterial;
+
+	UPROPERTY()
+	TObjectPtr<UMaterialInstanceDynamic> WorkingPreviewMaterial;
+	float SecondsBeforeWorkingMaterial = 0.75;
+
+	TSharedPtr<UE::Geometry::TMeshTangents<double>, ESPMode::ThreadSafe> BaseMeshTangents;
+	UE::Geometry::FDynamicMesh3 BaseMesh;
+	UE::Geometry::FDynamicMeshAABBTree3 BaseSpatial;
+
+	bool bIsBakeToSelf = false;
+
+	TSharedPtr<UE::Geometry::FDynamicMesh3, ESPMode::ThreadSafe> DetailMesh;
+	TSharedPtr<UE::Geometry::FDynamicMeshAABBTree3, ESPMode::ThreadSafe> DetailSpatial;
+	int32 DetailMeshTimestamp = 0;
+	void UpdateDetailMesh();
+	bool bDetailMeshValid = false;
+
+	bool bInputsDirty = false;
+	void UpdateResult();
+
+	void UpdateOnModeChange();
+	void UpdateVisualization();
+
+	TUniquePtr<TGenericDataBackgroundCompute<UE::Geometry::FMeshMapBaker>> Compute = nullptr;
+	void OnMapsUpdated(const TUniquePtr<UE::Geometry::FMeshMapBaker>& NewResult);
+
+	/** @return A single bitfield of map types from an array of map types. */
+	EBakeMapType GetMapTypes(const int32& MapTypes) const;
+	TArray<EBakeMapType> GetMapTypesArray(const int32& MapTypes) const;
+
+	struct FBakeCacheSettings
+	{
+		EBakeMapType BakeMapTypes = EBakeMapType::None;
+		FImageDimensions Dimensions;
+		int32 UVLayer;
+		int32 DetailTimestamp;
+		float Thickness;
+		int32 Multisampling;
+
+		bool operator==(const FBakeCacheSettings& Other) const
+		{
+			return BakeMapTypes == Other.BakeMapTypes && Dimensions == Other.Dimensions &&
+				UVLayer == Other.UVLayer && DetailTimestamp == Other.DetailTimestamp &&
+				Thickness == Other.Thickness && Multisampling == Other.Multisampling;
+		}
+	};
+	FBakeCacheSettings CachedBakeCacheSettings;
+	TArray<EBakeMapType> ResultTypes;
+
+	EBakeOpState OpState = EBakeOpState::Evaluate;
+
+	UPROPERTY()
+	TArray<TObjectPtr<UTexture2D>> CachedMaps;
+	using CachedMapIndex = TMap<EBakeMapType, int32>;
+	CachedMapIndex CachedMapIndices;
+
+	FNormalMapSettings CachedNormalMapSettings;
+	EBakeOpState UpdateResult_Normal();
+
+	FOcclusionMapSettings CachedOcclusionMapSettings;
+	EBakeOpState UpdateResult_Occlusion();
+
+	FCurvatureMapSettings CachedCurvatureMapSettings;
+	EBakeOpState UpdateResult_Curvature();
+
+	FMeshPropertyMapSettings CachedMeshPropertyMapSettings;
+	EBakeOpState UpdateResult_MeshProperty();
+
+	TSharedPtr<UE::Geometry::TImageBuilder<FVector4f>, ESPMode::ThreadSafe> CachedTextureImage;
+	FTexture2DImageSettings CachedTexture2DImageSettings;
+	bool bPreferPlatformData = true;
+	EBakeOpState UpdateResult_Texture2DImage();
+
+	TMap<int32, TSharedPtr<UE::Geometry::TImageBuilder<FVector4f>, ESPMode::ThreadSafe>> CachedMultiTextures;
+	EBakeOpState UpdateResult_MultiTexture();
+
+
+	// empty maps are shown when nothing is computed
+	UPROPERTY()
+	TObjectPtr<UTexture2D> EmptyNormalMap;
+
+	UPROPERTY()
+	TObjectPtr<UTexture2D> EmptyColorMapBlack;
+
+	UPROPERTY()
+	TObjectPtr<UTexture2D> EmptyColorMapWhite;
+
+	void InitializeEmptyMaps();
+
+	void GetTexturesFromDetailMesh(const UPrimitiveComponent* DetailComponent);
+
+};
diff --git a/Engine/Plugins/Experimental/MeshModelingToolsetExp/Source/MeshModelingToolsExp/Public/BakeTransformTool.h b/Engine/Plugins/Experimental/MeshModelingToolsetExp/Source/MeshModelingToolsExp/Public/BakeTransformTool.h
new file mode 100644
--- /dev/null
+++ b/Engine/Plugins/Experimental/MeshModelingToolsetExp/Source/MeshModelingToolsExp/Public/BakeTransformTool.h
@@ -0,0 +1,102 @@
+// Copyright Epic Games, Inc. All Rights Reserved.
+
+#pragma once
+
+#include "CoreMinimal.h"
+#include "UObject/NoExportTypes.h"
+#include "MultiSelectionTool.h"
+#include "InteractiveToolBuilder.h"
+#include "MeshOpPreviewHelpers.h"
+#include "DynamicMesh/DynamicMesh3.h"
+#include "BaseTools/SingleClickTool.h"
+#include "BakeTransformTool.generated.h"
+
+
+// predeclarations
+struct FMeshDescription;
+class UDynamicMeshComponent;
+
+
+/**
+ *
+ */
+UCLASS()
+class MESHMODELINGTOOLSEXP_API UBakeTransformToolBuilder : public UInteractiveToolBuilder
+{
+	GENERATED_BODY()
+
+public:
+	virtual bool CanBuildTool(const FToolBuilderState& SceneState) const override;
+	virtual UInteractiveTool* BuildTool(const FToolBuilderState& SceneState) const override;
+
+protected:
+	virtual const FToolTargetTypeRequirements& GetTargetRequirements() const override;
+};
+
+
+UENUM()
+enum class EBakeScaleMethod : uint8
+{
+	BakeFullScale,			// bake all scale information, so the component has scale of 1 on all axes
+	BakeNonuniformScale,	// bake the non-uniform scale, so the component has a uniform scale
+	DoNotBakeScale			// do not bake any scaling
+};
+
+
+/**
+ * Standard properties
+ */
+UCLASS()
+class MESHMODELINGTOOLSEXP_API UBakeTransformToolProperties : public UInteractiveToolPropertySet
+{
+	GENERATED_BODY()
+
+public:
+	UBakeTransformToolProperties();
+
+	/** Bake rotation */
+	UPROPERTY(EditAnywhere, Category = Options)
+	bool bBakeRotation = true;
+
+	/** Bake scale */
+	UPROPERTY(EditAnywhere, Category = Options)
+	EBakeScaleMethod BakeScale = EBakeScaleMethod::BakeFullScale;
+
+	/** Recenter pivot after baking transform */
+	UPROPERTY(EditAnywhere, Category = Options)
+	bool bRecenterPivot = false;
+};
+
+
+
+/**
+ * Simple tool to bake scene transform on meshes into the mesh assets
+ */
+UCLASS()
+class MESHMODELINGTOOLSEXP_API UBakeTransformTool : public UMultiSelectionTool
+{
+	GENERATED_BODY()
+
+public:
+	UBakeTransformTool();
+
+	virtual void Setup() override;
+	virtual void Shutdown(EToolShutdownType ShutdownType) override;
+
+	virtual void SetWorld(UWorld* World);
+
+	virtual bool HasCancel() const override { return true; }
+	virtual bool HasAccept() const override { return true; }
+
+protected:
+
+	UPROPERTY()
+	TObjectPtr<UBakeTransformToolProperties> BasicProperties;
+
+protected:
+	UWorld* TargetWorld;
+
+	TArray<int> MapToFirstOccurrences;
+
+	void UpdateAssets();
+};
diff --git a/Engine/Plugins/Experimental/MeshModelingToolsetExp/Source/MeshModelingToolsExp/Public/ConvertMeshesTool.h b/Engine/Plugins/Experimental/MeshModelingToolsetExp/Source/MeshModelingToolsExp/Public/ConvertMeshesTool.h
new file mode 100644
--- /dev/null
+++ b/Engine/Plugins/Experimental/MeshModelingToolsetExp/Source/MeshModelingToolsExp/Public/ConvertMeshesTool.h
@@ -0,0 +1,63 @@
+// Copyright Epic Games, Inc. All Rights Reserved.
+
+#pragma once
+
+#include "CoreMinimal.h"
+#include "MultiSelectionTool.h"
+#include "InteractiveToolBuilder.h"
+#include "PropertySets/CreateMeshObjectTypeProperties.h"
+#include "ConvertMeshesTool.generated.h"
+
+/**
+ *
+ */
+UCLASS()
+class MESHMODELINGTOOLSEXP_API UConvertMeshesToolBuilder : public UInteractiveToolBuilder
+{
+	GENERATED_BODY()
+public:
+	virtual bool CanBuildTool(const FToolBuilderState& SceneState) const override;
+	virtual UInteractiveTool* BuildTool(const FToolBuilderState& SceneState) const override;
+
+protected:
+	virtual const FToolTargetTypeRequirements& GetTargetRequirements() const override;
+};
+
+/**
+ * Standard properties of the Transfer operation
+ */
+UCLASS()
+class MESHMODELINGTOOLSEXP_API UConvertMeshesToolProperties : public UInteractiveToolPropertySet
+{
+	GENERATED_BODY()
+public:
+	UPROPERTY(EditAnywhere, Category = Options)
+	bool bTransferMaterials = true;
+};
+
+
+
+UCLASS()
+class MESHMODELINGTOOLSEXP_API UConvertMeshesTool : public UMultiSelectionTool
+{
+	GENERATED_BODY()
+
+public:
+	virtual void SetWorld(UWorld* World);
+
+	virtual void Setup() override;
+	virtual void Shutdown(EToolShutdownType ShutdownType) override;
+
+	virtual bool HasCancel() const override { return true; }
+	virtual bool HasAccept() const override { return true; }
+	virtual bool CanAccept() const override;
+
+	UPROPERTY()
+	TObjectPtr<UConvertMeshesToolProperties> BasicProperties;
+
+	UPROPERTY()
+	TObjectPtr<UCreateMeshObjectTypeProperties> OutputTypeProperties;
+
+protected:
+	UWorld* TargetWorld;
+};
diff --git a/Engine/Plugins/Experimental/MeshModelingToolsetExp/Source/MeshModelingToolsExp/Public/ConvertToPolygonsTool.h b/Engine/Plugins/Experimental/MeshModelingToolsetExp/Source/MeshModelingToolsExp/Public/ConvertToPolygonsTool.h
new file mode 100644
--- /dev/null
+++ b/Engine/Plugins/Experimental/MeshModelingToolsetExp/Source/MeshModelingToolsExp/Public/ConvertToPolygonsTool.h
@@ -0,0 +1,145 @@
+// Copyright Epic Games, Inc. All Rights Reserved.
+
+#pragma once
+
+#include "CoreMinimal.h"
+#include "BaseTools/SingleSelectionMeshEditingTool.h"
+#include "PreviewMesh.h"
+#include "ModelingOperators.h"
+#include "MeshOpPreviewHelpers.h"
+#include "ConvertToPolygonsTool.generated.h"
+
+// predeclaration
+class UConvertToPolygonsTool;
+class FConvertToPolygonsOp;
+class UPreviewGeometry;
+PREDECLARE_GEOMETRY(class FDynamicMesh3);
+
+/**
+ *
+ */
+UCLASS()
+class MESHMODELINGTOOLSEXP_API UConvertToPolygonsToolBuilder : public USingleSelectionMeshEditingToolBuilder
+{
+	GENERATED_BODY()
+public:
+	virtual USingleSelectionMeshEditingTool* CreateNewTool(const FToolBuilderState & SceneState) const override;
+};
+
+
+
+
+UENUM()
+enum class EConvertToPolygonsMode
+{
+	/** Convert based on Angle Tolerance between Face Normals */
+	FaceNormalDeviation UMETA(DisplayName = "Face Normal Deviation"),
+	/** Create PolyGroups based on UV Islands */
+	FromUVIslands  UMETA(DisplayName = "From UV Islands"),
+	/** Create Polygroups based on Connected Triangles */
+	FromConnectedTris UMETA(DisplayName = "From Connected Tris"),
+	/** Create Polygroups based on Connected Triangles */
+	FromFurthestPointSampling UMETA(DisplayName = "Furthest Point Sampling")
+};
+
+
+
+UCLASS()
+class MESHMODELINGTOOLSEXP_API UConvertToPolygonsToolProperties : public UInteractiveToolPropertySet
+{
+	GENERATED_BODY()
+
+public:
+	/** Strategy to use to group triangles */
+	UPROPERTY(EditAnywhere, Category = PolyGroups)
+	EConvertToPolygonsMode ConversionMode = EConvertToPolygonsMode::FaceNormalDeviation;
+
+	/** Tolerance for planarity */
+	UPROPERTY(EditAnywhere, Category = PolyGroups, meta = (UIMin = "0.001", UIMax = "20.0", ClampMin = "0.0", ClampMax = "90.0", EditCondition = "ConversionMode == EConvertToPolygonsMode::FaceNormalDeviation", EditConditionHides))
+	float AngleTolerance = 0.1f;
+
+	/** Furthest-Point Sample count, approximately this number of polygroups will be generated */
+	UPROPERTY(EditAnywhere, Category = PolyGroups, meta = (UIMin = "1", UIMax = "100", ClampMin = "1", ClampMax = "10000", EditCondition = "ConversionMode == EConvertToPolygonsMode::FromFurthestPointSampling", EditConditionHides))
+	int32 NumPoints = 100;
+
+	/** If true, region-growing in Sampling modes will be controlled by face normals, resulting in regions with borders that are more-aligned with curvature ridges */
+	UPROPERTY(EditAnywhere, Category = PolyGroups, meta = (EditCondition = "ConversionMode == EConvertToPolygonsMode::FromFurthestPointSampling", EditConditionHides))
+	bool bNormalWeighted = true;
+
+	/** This parameter modulates the effect of normal weighting during region-growing */
+	UPROPERTY(EditAnywhere, Category = PolyGroups, meta = (UIMin = "0.1", UIMax = "2.0", ClampMin = "0.01", ClampMax = "100.0", EditCondition = "ConversionMode == EConvertToPolygonsMode::FromFurthestPointSampling", EditConditionHides))
+	float NormalWeighting = 1.0f;
+
+
+	/** group filtering */
+	UPROPERTY(EditAnywhere, Category = Filtering, meta = (UIMin = "1", UIMax = "100", ClampMin = "1", ClampMax = "10000"))
+	int32 MinGroupSize = 2;
+
+
+	/** If true, normals are recomputed per-group, with hard edges at group boundaries */
+	UPROPERTY(EditAnywhere, Category = Output)
+	bool bCalculateNormals = false;
+	
+	/** Display each group with a different auto-generated color */
+	UPROPERTY(EditAnywhere, Category = Display)
+	bool bShowGroupColors = true;
+};
+
+UCLASS()
+class MESHMODELINGTOOLSEXP_API UConvertToPolygonsOperatorFactory : public UObject, public UE::Geometry::IDynamicMeshOperatorFactory
+{
+	GENERATED_BODY()
+
+public:
+	// IDynamicMeshOperatorFactory API
+	virtual TUniquePtr<UE::Geometry::FDynamicMeshOperator> MakeNewOperator() override;
+
+	UPROPERTY()
+	TObjectPtr<UConvertToPolygonsTool> ConvertToPolygonsTool;  // back pointer
+};
+
+/**
+ *
+ */
+UCLASS()
+class MESHMODELINGTOOLSEXP_API UConvertToPolygonsTool : public USingleSelectionMeshEditingTool
+{
+	GENERATED_BODY()
+
+public:
+	UConvertToPolygonsTool();
+
+	virtual void Setup() override;
+	virtual void Shutdown(EToolShutdownType ShutdownType) override;
+
+	virtual void OnTick(float DeltaTime) override;
+
+	virtual bool HasCancel() const override { return true; }
+	virtual bool HasAccept() const override { return true; }
+	virtual bool CanAccept() const override;
+
+	// update parameters in ConvertToPolygonsOp based on current Settings
+	void UpdateOpParameters(FConvertToPolygonsOp& ConvertToPolygonsOp) const;
+
+protected:
+	
+	void OnSettingsModified();
+
+protected:
+	UPROPERTY()
+	TObjectPtr<UConvertToPolygonsToolProperties> Settings;
+
+	UPROPERTY()
+	TObjectPtr<UMeshOpPreviewWithBackgroundCompute> PreviewCompute = nullptr;
+
+	UPROPERTY()
+	TObjectPtr<UPreviewGeometry> PreviewGeometry = nullptr;
+
+protected:
+	TSharedPtr<UE::Geometry::FDynamicMesh3, ESPMode::ThreadSafe> OriginalDynamicMesh;
+
+	// for visualization
+	TArray<int> PolygonEdges;
+	
+	void UpdateVisualization();
+};
diff --git a/Engine/Plugins/Experimental/MeshModelingToolsetExp/Source/MeshModelingToolsExp/Public/DeformMeshPolygonsTool.h b/Engine/Plugins/Experimental/MeshModelingToolsetExp/Source/MeshModelingToolsExp/Public/DeformMeshPolygonsTool.h
new file mode 100644
--- /dev/null
+++ b/Engine/Plugins/Experimental/MeshModelingToolsetExp/Source/MeshModelingToolsExp/Public/DeformMeshPolygonsTool.h
@@ -0,0 +1,257 @@
+// Copyright Epic Games, Inc. All Rights Reserved.
+
+#pragma once
+
+#include "CoreMinimal.h"
+#include "Templates/PimplPtr.h"
+#include "UObject/NoExportTypes.h"
+#include "BaseTools/MeshSurfacePointTool.h"
+#include "Components/DynamicMeshComponent.h"
+#include "DynamicMesh/DynamicMeshAABBTree3.h"
+#include "ToolDataVisualizer.h"
+#include "Transforms/QuickAxisTranslater.h"
+#include "Transforms/QuickAxisRotator.h"
+#include "Changes/MeshVertexChange.h"
+#include "GroupTopology.h"
+#include "Spatial/GeometrySet3.h"
+#include "Selection/GroupTopologySelector.h"
+#include "Operations/GroupTopologyDeformer.h"
+#include "Solvers/MeshLaplacian.h"
+#include "Curves/CurveFloat.h"
+#include "DeformMeshPolygonsTool.generated.h"
+
+class FMeshVertexChangeBuilder;
+class FGroupTopologyLaplacianDeformer;
+struct FDeformerVertexConstraintData;
+
+/**
+ * ToolBuilder
+ */
+UCLASS()
+class MESHMODELINGTOOLSEXP_API UDeformMeshPolygonsToolBuilder : public UMeshSurfacePointToolBuilder
+{
+	GENERATED_BODY()
+
+public:
+	UDeformMeshPolygonsToolBuilder()
+	{
+	}
+
+	virtual UMeshSurfacePointTool* CreateNewTool(const FToolBuilderState& SceneState) const override;
+};
+
+/** Deformation strategies */
+UENUM()
+enum class EGroupTopologyDeformationStrategy : uint8
+{
+	/** Deforms the mesh using linear translations*/
+	Linear UMETA(DisplayName = "Linear"),
+
+	/** Deforms the mesh using laplacian deformation*/
+	Laplacian UMETA(DisplayName = "Smooth")
+};
+
+/** Laplacian weight schemes determine how we will look at the curvature at a given vertex in relation to its neighborhood*/
+UENUM()
+enum class EWeightScheme 
+{
+	Uniform				UMETA(DisplayName = "Uniform"),
+	Umbrella			UMETA(DisplayName = "Umbrella"),
+	Valence				UMETA(DisplayName = "Valence"),
+	MeanValue			UMETA(DisplayName = "MeanValue"),
+	Cotangent			UMETA(DisplayName = "Cotangent"),
+	ClampedCotangent	UMETA(DisplayName = "ClampedCotangent")
+};
+
+/** The ELaplacianWeightScheme enum is the same..*/
+static ELaplacianWeightScheme ConvertToLaplacianWeightScheme(const EWeightScheme WeightScheme)
+{
+	return static_cast<ELaplacianWeightScheme>(WeightScheme);
+}
+
+
+/** Modes for quick transformer */
+UENUM()
+enum class EQuickTransformerMode : uint8
+{
+	/** Translation along frame axes */
+	AxisTranslation = 0 UMETA(DisplayName = "Translate"),
+
+	/** Rotation around frame axes*/
+	AxisRotation = 1 UMETA(DisplayName = "Rotate"),
+};
+
+
+
+
+UCLASS()
+class MESHMODELINGTOOLSEXP_API UDeformMeshPolygonsTransformProperties : public UInteractiveToolPropertySet
+{
+	GENERATED_BODY()
+
+public:
+	UDeformMeshPolygonsTransformProperties();
+
+	//Options
+
+	UPROPERTY(EditAnywhere, Category = Options, meta = (DisplayName = "Deformation Type", ToolTip = "Select the type of deformation you wish to employ on a polygroup."))
+	EGroupTopologyDeformationStrategy DeformationStrategy;
+
+	UPROPERTY(EditAnywhere, Category = Options)
+	EQuickTransformerMode TransformMode;
+
+	UPROPERTY(EditAnywhere, Category = SelectionFilter)
+	bool bSelectFaces;
+
+	UPROPERTY(EditAnywhere, Category = SelectionFilter)
+	bool bSelectEdges;
+
+	UPROPERTY(EditAnywhere, Category = SelectionFilter)
+	bool bSelectVertices;
+
+	UPROPERTY(EditAnywhere, Category = Display)
+	bool bShowWireframe;
+
+	//Laplacian Deformation Options, currently not exposed.
+
+	UPROPERTY() 
+	EWeightScheme SelectedWeightScheme = EWeightScheme::ClampedCotangent;
+
+	UPROPERTY() 
+	double HandleWeight = 1000.0;
+
+	UPROPERTY()
+	bool bPostFixHandles = false;
+/**
+	// How to add a weight curve
+	UPROPERTY(EditAnywhere, Category = LaplacianOptions, meta = (EditCondition="DeformationStrategy == EGroupTopologyDeformationStrategy::Laplacian", DisplayName = "Localize Deformation", ToolTip = "When enabled, only the vertices in the polygroups immediately adjacent to the selected group will be affected by the deformation.\nWhen disabled, the deformer will solve for the curvature of the entire mesh (slower)"))
+	bool bLocalizeDeformation{ true};
+
+	UPROPERTY(EditAnywhere, Category = LaplacianOptions, meta = (EditCondition="DeformationStrategy == EGroupTopologyDeformationStrategy::Laplacian", DisplayName = "Apply Weight Attenuation Curve", ToolTip = "When enabled, the curve below will be used to calculate the weight at a given vertex based on distance from the handles"))
+	bool bApplyAttenuationCurve{ false };
+	
+	FRichCurve DefaultFalloffCurve;
+
+	UPROPERTY(EditAnywhere,  Category = LaplacianOptions, meta = ( EditCondition="DeformationStrategy == EGroupTopologyDeformationStrategy::Laplacian && bApplyAttenuationCurve", DisplayName = "Distance-Weight Attenuation Curve",UIMin = "1.0", UIMax = "1.0", ClampMin = "1.0", ClampMax = "1.0", ToolTip = "This curve determines the weight attenuation over the distance of the mesh.\nThe selected polygroup handle is t=0.0, and t=1.0 is roughly the farthest vertices from the handles.\nThe value of the curve at each time interval represents the weight of the vertices at that distance from the selection."))
+	FRuntimeFloatCurve WeightAttenuationCurve;
+*/ 
+};
+
+
+/**
+ *
+ */
+UCLASS()
+class MESHMODELINGTOOLSEXP_API UDeformMeshPolygonsTool : public UMeshSurfacePointTool
+{
+	GENERATED_BODY()
+
+public:
+	UDeformMeshPolygonsTool();
+
+	virtual void RegisterActions(FInteractiveToolActionSet& ActionSet) override;
+
+	virtual void Setup() override;
+	virtual void Shutdown(EToolShutdownType ShutdownType) override;
+
+	virtual void OnTick(float DeltaTime) override;
+	virtual void Render(IToolsContextRenderAPI* RenderAPI) override;
+
+	virtual bool HasCancel() const override { return true; }
+	virtual bool HasAccept() const override { return true; }
+
+	virtual bool HitTest(const FRay& Ray, FHitResult& OutHit) override;
+	
+	virtual void OnBeginDrag(const FRay& Ray) override;
+	virtual void OnUpdateDrag(const FRay& Ray) override;
+	virtual void OnEndDrag(const FRay& Ray) override;
+	virtual bool OnUpdateHover(const FInputDeviceRay& DevicePos) override;
+
+public:
+	virtual void NextTransformTypeAction();
+
+	//
+	float VisualAngleSnapThreshold = 0.5;
+
+protected:
+	UPROPERTY()
+	TObjectPtr<UDynamicMeshComponent> DynamicMeshComponent;
+
+	UPROPERTY()
+	TObjectPtr<UDeformMeshPolygonsTransformProperties> TransformProps;
+
+	// realtime visualization
+	void OnDynamicMeshComponentChanged();
+	FDelegateHandle OnDynamicMeshComponentChangedHandle;
+
+	virtual void OnPropertyModified(UObject* PropertySet, FProperty* Property) override;
+	
+	// camera state at last render
+	FViewCameraState CameraState;
+	UE::Geometry::FTransform3d WorldTransform;
+
+	FToolDataVisualizer PolyEdgesRenderer;
+
+	// True for the duration of UI click+drag
+	bool bInDrag;
+
+	FPlane ActiveDragPlane;
+	FVector StartHitPosWorld;
+	FVector StartHitNormalWorld;
+	FVector LastHitPosWorld;
+	FVector LastBrushPosLocal;
+	FVector StartBrushPosLocal;
+
+	UE::Geometry::FFrame3d ActiveSurfaceFrame;
+	UE::Geometry::FQuickTransformer* GetActiveQuickTransformer();
+	void UpdateActiveSurfaceFrame(FGroupTopologySelection& Selection);
+	void UpdateQuickTransformer();
+
+	FRay UpdateRay;
+	bool bUpdatePending = false;
+	void ComputeUpdate();
+
+	FVector3d LastMoveDelta;
+	UE::Geometry::FQuickAxisTranslater QuickAxisTranslater;
+	void ComputeUpdate_Translate();
+
+	UE::Geometry::FQuickAxisRotator QuickAxisRotator;
+	FVector3d RotationStartPointWorld;
+	UE::Geometry::FFrame3d RotationStartFrame;
+	void ComputeUpdate_Rotate();
+
+	FGroupTopology Topology;
+	void PrecomputeTopology();
+
+	FGroupTopologySelector TopoSelector;
+	FGroupTopologySelector::FSelectionSettings GetTopoSelectorSettings();
+	
+	//
+	// data for current drag
+	//
+
+	FGroupTopologySelection HilightSelection;
+	FToolDataVisualizer HilightRenderer;
+
+	UE::Geometry::FDynamicMeshAABBTree3 MeshSpatial;
+	UE::Geometry::FDynamicMeshAABBTree3& GetSpatial();
+
+	FMeshVertexChangeBuilder* ActiveVertexChange;
+
+	EGroupTopologyDeformationStrategy DeformationStrategy;
+
+	// The two deformer type options.
+	UE::Geometry::FGroupTopologyDeformer LinearDeformer;
+	TPimplPtr<FGroupTopologyLaplacianDeformer> LaplacianDeformer;
+
+
+
+	// This is true when the spatial index needs to reflect a modification
+	bool bSpatialDirty; 
+
+	void BeginChange();
+	void EndChange();
+	void UpdateChangeFromROI(bool bFinal);
+
+};
+
diff --git a/Engine/Plugins/Experimental/MeshModelingToolsetExp/Source/MeshModelingToolsExp/Public/DrawPolyPathTool.h b/Engine/Plugins/Experimental/MeshModelingToolsetExp/Source/MeshModelingToolsExp/Public/DrawPolyPathTool.h
new file mode 100644
--- /dev/null
+++ b/Engine/Plugins/Experimental/MeshModelingToolsetExp/Source/MeshModelingToolsExp/Public/DrawPolyPathTool.h
@@ -0,0 +1,255 @@
+// Copyright Epic Games, Inc. All Rights Reserved.
+
+#pragma once
+
+#include "CoreMinimal.h"
+#include "UObject/NoExportTypes.h"
+#include "BaseTools/MeshSurfacePointTool.h"
+#include "Mechanics/PlaneDistanceFromHitMechanic.h"
+#include "Mechanics/SpatialCurveDistanceMechanic.h"
+#include "Mechanics/CollectSurfacePathMechanic.h"
+#include "Mechanics/ConstructionPlaneMechanic.h"
+#include "Drawing/PolyEditPreviewMesh.h"
+#include "PropertySets/CreateMeshObjectTypeProperties.h"
+#include "Properties/MeshMaterialProperties.h"
+#include "DrawPolyPathTool.generated.h"
+
+class FMeshVertexChangeBuilder;
+class UTransformGizmo;
+class UTransformProxy;
+
+/**
+ * ToolBuilder
+ */
+UCLASS()
+class MESHMODELINGTOOLSEXP_API UDrawPolyPathToolBuilder : public UMeshSurfacePointToolBuilder
+{
+	GENERATED_BODY()
+public:
+	virtual bool CanBuildTool(const FToolBuilderState& SceneState) const override;
+
+	virtual UInteractiveTool* BuildTool(const FToolBuilderState& SceneState) const override;
+};
+
+
+UENUM()
+enum class EDrawPolyPathOutputMode
+{
+	Ribbon,
+	Extrusion,
+	Ramp
+};
+
+
+
+UENUM()
+enum class EDrawPolyPathWidthMode
+{
+	Interactive,
+	Constant
+};
+
+UENUM()
+enum class EDrawPolyPathHeightMode
+{
+	Interactive,
+	Constant
+};
+
+
+
+UCLASS()
+class MESHMODELINGTOOLSEXP_API UDrawPolyPathProperties : public UInteractiveToolPropertySet
+{
+	GENERATED_BODY()
+
+public:
+	UPROPERTY(EditAnywhere, Category = Shape)
+	EDrawPolyPathOutputMode OutputType = EDrawPolyPathOutputMode::Extrusion;
+
+	UPROPERTY(EditAnywhere, Category = Shape)
+	EDrawPolyPathWidthMode WidthMode = EDrawPolyPathWidthMode::Interactive;
+
+	UPROPERTY(EditAnywhere, Category = Shape, meta = (UIMin = "0.0001", UIMax = "1000", ClampMin = "0", ClampMax = "999999"))
+	float Width = 10.0f;
+
+	UPROPERTY(EditAnywhere, Category = Shape, meta = (
+		EditCondition = "OutputType != EDrawPolyPathOutputMode::Ribbon", EditConditionHides))
+	EDrawPolyPathHeightMode HeightMode = EDrawPolyPathHeightMode::Interactive;
+
+	UPROPERTY(EditAnywhere, Category = Shape, meta = (
+		EditCondition = "OutputType != EDrawPolyPathOutputMode::Ribbon", 
+		EditConditionHides, UIMin = "-1000", UIMax = "1000", ClampMin = "-10000", ClampMax = "10000"))
+	float Height = 10.0f;
+
+	UPROPERTY(EditAnywhere, Category = Shape, meta = (
+		EditCondition = "OutputType == EDrawPolyPathOutputMode::Ramp", EditConditionHides,
+		UIMin = "0.01", UIMax = "1.0", ClampMin = "0", ClampMax = "100.0"))
+	float RampStartRatio = 0.05f;
+
+
+	UPROPERTY(EditAnywhere, Category = Snapping)
+	bool bSnapToWorldGrid = true;
+};
+
+
+
+UENUM()
+enum class EDrawPolyPathExtrudeDirection
+{
+	SelectionNormal,
+	WorldX,
+	WorldY,
+	WorldZ,
+	LocalX,
+	LocalY,
+	LocalZ
+};
+
+
+UCLASS()
+class MESHMODELINGTOOLSEXP_API UDrawPolyPathExtrudeProperties : public UInteractiveToolPropertySet
+{
+	GENERATED_BODY()
+
+public:
+	UPROPERTY(EditAnywhere, Category = Extrude)
+	EDrawPolyPathExtrudeDirection Direction = EDrawPolyPathExtrudeDirection::SelectionNormal;
+
+
+};
+
+
+
+/**
+ *
+ */
+UCLASS()
+class MESHMODELINGTOOLSEXP_API UDrawPolyPathTool : public UInteractiveTool, public IClickBehaviorTarget, public IHoverBehaviorTarget
+{
+	GENERATED_BODY()
+public:
+	virtual void RegisterActions(FInteractiveToolActionSet& ActionSet) override;
+
+	virtual void SetWorld(UWorld* World);
+
+	virtual void Setup() override;
+	virtual void Shutdown(EToolShutdownType ShutdownType) override;
+
+	virtual void OnTick(float DeltaTime) override;
+	virtual void Render(IToolsContextRenderAPI* RenderAPI) override;
+
+	virtual bool HasCancel() const override { return false; }
+	virtual bool HasAccept() const override { return false; }
+	virtual bool CanAccept() const override { return false; }
+
+	virtual bool HitTest(const FRay& Ray, FHitResult& OutHit);
+
+	// IClickBehaviorTarget API
+	virtual FInputRayHit IsHitByClick(const FInputDeviceRay& ClickPos) override;
+	virtual void OnClicked(const FInputDeviceRay& ClickPos) override;
+
+	// IHoverBehaviorTarget API
+	virtual FInputRayHit BeginHoverSequenceHitTest(const FInputDeviceRay& PressPos) override;
+	virtual void OnBeginHover(const FInputDeviceRay& DevicePos) override {}
+	virtual bool OnUpdateHover(const FInputDeviceRay& DevicePos) override;
+	virtual void OnEndHover() override {}
+
+protected:
+	UWorld* TargetWorld;
+
+	/** Property set for type of output object (StaticMesh, Volume, etc) */
+	UPROPERTY()
+	TObjectPtr<UCreateMeshObjectTypeProperties> OutputTypeProperties;
+
+	UPROPERTY()
+	TObjectPtr<UDrawPolyPathProperties> TransformProps;
+
+	UPROPERTY()
+	TObjectPtr<UDrawPolyPathExtrudeProperties> ExtrudeProperties;
+
+	UPROPERTY()
+	TObjectPtr<UNewMeshMaterialProperties> MaterialProperties;
+
+protected:
+	// camera state at last render
+	UE::Geometry::FTransform3d WorldTransform;
+	FViewCameraState CameraState;
+
+	// drawing plane and gizmo
+
+	UPROPERTY()
+	TObjectPtr<UConstructionPlaneMechanic> PlaneMechanic = nullptr;
+
+	UE::Geometry::FFrame3d DrawPlaneWorld;
+
+	bool CanUpdateDrawPlane() const;
+
+	// UV Scale factor to apply to texturing on any new geometry (e.g. new faces added by extrude)
+	float UVScaleFactor = 1.0f;
+
+	TArray<UE::Geometry::FFrame3d> CurPathPoints;
+	TArray<double> OffsetScaleFactors;
+	TArray<double> ArcLengths;
+	TArray<FVector3d> CurPolyLine;
+	double CurPathLength;
+	double CurOffsetDistance;
+	double CurHeight;
+	bool bPathIsClosed = false;		// If true, CurPathPoints are assumed to define a closed path
+
+	UPROPERTY()
+	TObjectPtr<UPolyEditPreviewMesh> EditPreview;
+
+	UPROPERTY()
+	TObjectPtr<UPlaneDistanceFromHitMechanic> ExtrudeHeightMechanic = nullptr;
+	UPROPERTY()
+	TObjectPtr<USpatialCurveDistanceMechanic> CurveDistMechanic = nullptr;
+	UPROPERTY()
+	TObjectPtr<UCollectSurfacePathMechanic> SurfacePathMechanic = nullptr;
+
+	void InitializeNewSurfacePath();
+	void UpdateSurfacePathPlane();
+	void OnCompleteSurfacePath();
+	void OnCompleteOffsetDistance();
+	void OnCompleteExtrudeHeight();
+
+	void BeginInteractiveOffsetDistance();
+	void BeginConstantOffsetDistance();
+	void UpdatePathPreview();
+	void BeginInteractiveExtrudeHeight();
+	void UpdateExtrudePreview();
+	void InitializePreviewMesh();
+	void ClearPreview();
+	void GeneratePathMesh(UE::Geometry::FDynamicMesh3& Mesh);
+	void GenerateExtrudeMesh(UE::Geometry::FDynamicMesh3& PathMesh);
+	void GenerateRampMesh(FDynamicMesh3& PathMesh);
+	void EmitNewObject(EDrawPolyPathOutputMode OutputMode);
+
+	// user feedback messages
+	void ShowStartupMessage();
+	void ShowExtrudeMessage();
+	void ShowOffsetMessage();
+
+	friend class FDrawPolyPathStateChange;
+	int32 CurrentCurveTimestamp = 1;
+	void UndoCurrentOperation();
+	bool CheckInCurve(int32 Timestamp) const { return CurrentCurveTimestamp == Timestamp; }
+};
+
+
+
+
+class MESHMODELINGTOOLSEXP_API FDrawPolyPathStateChange : public FToolCommandChange
+{
+public:
+	bool bHaveDoneUndo = false;
+	int32 CurveTimestamp = 0;
+	FDrawPolyPathStateChange(int32 CurveTimestampIn)
+	{
+		CurveTimestamp = CurveTimestampIn;
+	}
+	virtual void Apply(UObject* Object) override {}
+	virtual void Revert(UObject* Object) override;
+	virtual bool HasExpired(UObject* Object) const override;
+	virtual FString ToString() const override;
+};
diff --git a/Engine/Plugins/Experimental/MeshModelingToolsetExp/Source/MeshModelingToolsExp/Public/DynamicMeshSculptTool.h b/Engine/Plugins/Experimental/MeshModelingToolsetExp/Source/MeshModelingToolsExp/Public/DynamicMeshSculptTool.h
new file mode 100644
--- /dev/null
+++ b/Engine/Plugins/Experimental/MeshModelingToolsetExp/Source/MeshModelingToolsExp/Public/DynamicMeshSculptTool.h
@@ -0,0 +1,498 @@
+// Copyright Epic Games, Inc. All Rights Reserved.
+
+#pragma once
+
+#include "CoreMinimal.h"
+#include "HAL/Platform.h"
+#include "SubRegionRemesher.h"
+#include "UObject/NoExportTypes.h"
+#include "BaseTools/MeshSurfacePointTool.h"
+#include "Components/OctreeDynamicMeshComponent.h"
+#include "DynamicMesh/DynamicMeshAABBTree3.h"
+#include "DynamicMesh/DynamicMeshOctree3.h"
+#include "DynamicMesh/MeshNormals.h"
+#include "BaseTools/BaseBrushTool.h"
+#include "ToolDataVisualizer.h"
+#include "Changes/ValueWatcher.h"
+#include "BaseGizmos/BrushStampIndicator.h"
+#include "Properties/MeshMaterialProperties.h"
+#include "Properties/RemeshProperties.h"
+#include "TransformTypes.h"
+#include "Sculpting/MeshSculptToolBase.h"
+#include "Async/Async.h"
+#include "Util/UniqueIndexSet.h"
+#include "DynamicMeshSculptTool.generated.h"
+
+class UTransformGizmo;
+class UTransformProxy;
+class UMaterialInstanceDynamic;
+
+class FMeshVertexChangeBuilder;
+class UPreviewMesh;
+PREDECLARE_GEOMETRY(class FDynamicMeshChangeTracker);
+PREDECLARE_GEOMETRY(class FSubRegionRemesher);
+PREDECLARE_USE_GEOMETRY_CLASS(FDynamicMesh3);
+class FPersistentStampRemesher;
+
+
+/** Mesh Sculpting Brush Types */
+UENUM()
+enum class EDynamicMeshSculptBrushType : uint8
+{
+	/** Move vertices parallel to the view plane  */
+	Move UMETA(DisplayName = "Move"),
+
+	/** Grab Brush, fall-off alters the influence of the grab */
+	PullKelvin UMETA(DisplayName = "Kelvin Grab"),
+
+	/** Grab Brush that may generate cusps, fall-off alters the influence of the grab */
+	PullSharpKelvin UMETA(DisplayName = "Sharp Kelvin Grab"),
+
+	/** Smooth mesh vertices  */
+	Smooth UMETA(DisplayName = "Smooth"),
+
+	/** Displace vertices along the average surface normal (Ctrl to invert) */
+	Offset UMETA(DisplayName = "Sculpt (Normal)"),
+
+	/** Displace vertices towards the camera viewpoint (Ctrl to invert) */
+	SculptView UMETA(DisplayName = "Sculpt (Viewpoint)"),
+
+	/** Displaces vertices along the average surface normal to a maximum height based on the brush size (Ctrl to invert) */
+	SculptMax UMETA(DisplayName = "Sculpt Max"),
+
+	/** Displace vertices along their vertex normals */
+	Inflate UMETA(DisplayName = "Inflate"),
+
+	/** Scale Brush will inflate or pinch radially from the center of the brush */
+	ScaleKelvin UMETA(DisplayName = "Kelvin Scale"),
+
+	/** Move vertices towards the center of the brush (Ctrl to push away)*/
+	Pinch UMETA(DisplayName = "Pinch"),
+
+	/** Twist Brush moves vertices in the plane perpendicular to the local mesh normal */
+	TwistKelvin UMETA(DisplayName = "Kelvin Twist"),
+
+	/** Move vertices towards the average plane of the brush stamp region */
+	Flatten UMETA(DisplayName = "Flatten"),
+
+	/** Move vertices towards a plane defined by the initial brush position  */
+	Plane UMETA(DisplayName = "Plane (Normal)"),
+
+	/** Move vertices towards a view-facing plane defined at the initial brush position */
+	PlaneViewAligned UMETA(DisplayName = "Plane (Viewpoint)"),
+
+	/** Move vertices towards a fixed plane in world space, positioned with a 3D gizmo */
+	FixedPlane UMETA(DisplayName = "FixedPlane"),
+
+	/** Remesh the brushed region but do not otherwise deform it */
+	Resample UMETA(DisplayName = "Resample"),
+
+	LastValue UMETA(Hidden)
+
+};
+
+/**
+ * Tool Builder
+ */
+UCLASS()
+class MESHMODELINGTOOLSEXP_API UDynamicMeshSculptToolBuilder : public UMeshSurfacePointToolBuilder
+{
+	GENERATED_BODY()
+
+public:
+	bool bEnableRemeshing;
+
+	UDynamicMeshSculptToolBuilder()
+	{
+		bEnableRemeshing = false;
+	}
+
+	virtual UMeshSurfacePointTool* CreateNewTool(const FToolBuilderState& SceneState) const override;
+};
+
+
+
+
+UCLASS()
+class MESHMODELINGTOOLSEXP_API UBrushSculptProperties : public UInteractiveToolPropertySet
+{
+	GENERATED_BODY()
+
+public:
+	/* This is a dupe of the bool in the tool class.  I needed it here so it could be checked as an EditCondition */
+	UPROPERTY(meta = (TransientToolProperty))
+	bool bIsRemeshingEnabled = false;
+
+	/** Primary Brush Mode */
+	UPROPERTY(EditAnywhere, Category = Sculpting, meta = (DisplayName = "Brush Type"))
+	EDynamicMeshSculptBrushType PrimaryBrushType = EDynamicMeshSculptBrushType::Move;
+
+	/** Strength of the Primary Brush */
+	UPROPERTY(EditAnywhere, Category = Sculpting, meta = (DisplayName = "Strength", UIMin = "0.0", UIMax = "1.0", ClampMin = "0.0", ClampMax = "1.0", EditCondition = "PrimaryBrushType != EDynamicMeshSculptBrushType::Pull"))
+	float PrimaryBrushSpeed = 0.5;
+
+	/** If true, try to preserve the shape of the UV/3D mapping. This will limit Smoothing and Remeshing in some cases. */
+	UPROPERTY(EditAnywhere, Category = Sculpting)
+	bool bPreserveUVFlow = false;
+
+	/** When Freeze Target is toggled on, the Brush Target Surface will be Frozen in its current state, until toggled off. Brush strokes will be applied relative to the Target Surface, for applicable Brushes */
+	UPROPERTY(EditAnywhere, Category = Sculpting, meta = (EditCondition = "PrimaryBrushType == EDynamicMeshSculptBrushType::Sculpt || PrimaryBrushType == EDynamicMeshSculptBrushType::SculptMax || PrimaryBrushType == EDynamicMeshSculptBrushType::SculptView || PrimaryBrushType == EDynamicMeshSculptBrushType::Pinch || PrimaryBrushType == EDynamicMeshSculptBrushType::Resample" ))
+	bool bFreezeTarget = false;
+
+	/** Strength of Shift-to-Smooth Brushing and Smoothing Brush */
+	UPROPERTY(EditAnywhere, Category = Smoothing, meta = (DisplayName = "Smoothing Strength", UIMin = "0.0", UIMax = "1.0", ClampMin = "0.0", ClampMax = "1.0"))
+	float SmoothBrushSpeed = 0.25;
+
+	/** If enabled, Remeshing is limited during Smoothing to avoid wiping out higher-density triangle areas */
+	UPROPERTY(EditAnywhere, Category = Smoothing, meta = (DisplayName = "Preserve Tri Density", EditConditionHides, HideEditConditionToggle, EditCondition = "bIsRemeshingEnabled"))
+	bool bDetailPreservingSmooth = true;
+};
+
+
+UCLASS()
+class MESHMODELINGTOOLSEXP_API UDynamicSculptToolActions : public UInteractiveToolPropertySet
+{
+	GENERATED_BODY()
+public:
+	TWeakObjectPtr<UDynamicMeshSculptTool> ParentTool;
+
+	void Initialize(UDynamicMeshSculptTool* ParentToolIn) { ParentTool = ParentToolIn; }
+
+	UFUNCTION(CallInEditor, Category = MeshEdits)
+	void DiscardAttributes();
+};
+
+
+
+
+UCLASS()
+class MESHMODELINGTOOLSEXP_API UBrushRemeshProperties : public URemeshProperties
+{
+	GENERATED_BODY()
+
+public:
+	/** Toggle remeshing on/off */
+	UPROPERTY(EditAnywhere, Category = Remeshing, meta = (DisplayPriority = 1))
+	bool bEnableRemeshing = true;
+
+	// Note that if you change range here, you must also update UDynamicMeshSculptTool::ConfigureRemesher!
+	/** Desired size of triangles after Remeshing, relative to average initial triangle size. Larger value results in larger triangles. */
+	UPROPERTY(EditAnywhere, Category = Remeshing, meta = (DisplayName = "Relative Tri Size", UIMin = "-5", UIMax = "5", ClampMin = "-5", ClampMax = "5", DisplayPriority = 2))
+	int TriangleSize = 0;
+
+	/** Control the amount of simplification during sculpting. Higher values will avoid wiping out fine details on the mesh. */
+	UPROPERTY(EditAnywhere, Category = Remeshing, meta = (UIMin = "0", UIMax = "5", ClampMin = "0", ClampMax = "5", DisplayPriority = 3))
+	int PreserveDetail = 0;
+
+	UPROPERTY(EditAnywhere, Category = Remeshing, AdvancedDisplay)
+	int Iterations = 5;
+};
+
+
+
+UCLASS()
+class MESHMODELINGTOOLSEXP_API UFixedPlaneBrushProperties : public UInteractiveToolPropertySet
+{
+	GENERATED_BODY()
+
+public:
+	UPROPERTY( meta = (TransientToolProperty) )
+	bool bPropertySetEnabled = true;
+
+	/** Toggle whether Work Plane Positioning Gizmo is visible */
+	UPROPERTY(EditAnywhere, Category = TargetPlane, meta = (HideEditConditionToggle, EditCondition = "bPropertySetEnabled == true"))
+	bool bShowGizmo = true;
+
+	UPROPERTY(EditAnywhere, Category = TargetPlane, AdvancedDisplay, meta = (HideEditConditionToggle, EditCondition = "bPropertySetEnabled == true"))
+	FVector Position = FVector::ZeroVector;
+
+	UPROPERTY(EditAnywhere, Category = TargetPlane, AdvancedDisplay, meta = (HideEditConditionToggle, EditCondition = "bPropertySetEnabled == true"))
+	FQuat Rotation = FQuat::Identity;
+};
+
+
+
+
+/**
+ * Dynamic Mesh Sculpt Tool Class
+ */
+UCLASS()
+class MESHMODELINGTOOLSEXP_API UDynamicMeshSculptTool : public UMeshSurfacePointTool
+{
+	GENERATED_BODY()
+
+public:
+	using FFrame3d = UE::Geometry::FFrame3d;
+	using FRay3d = UE::Geometry::FRay3d;
+
+	UDynamicMeshSculptTool();
+
+	virtual void SetWorld(UWorld* World);
+	virtual void RegisterActions(FInteractiveToolActionSet& ActionSet) override;
+
+	virtual void Setup() override;
+	virtual void Shutdown(EToolShutdownType ShutdownType) override;
+
+	virtual void OnTick(float DeltaTime) override;
+	virtual void Render(IToolsContextRenderAPI* RenderAPI) override;
+
+	virtual bool HasCancel() const override { return true; }
+	virtual bool HasAccept() const override { return true; }
+
+	virtual bool HitTest(const FRay& Ray, FHitResult& OutHit) override;
+
+	virtual void OnBeginDrag(const FRay& Ray) override;
+	virtual void OnUpdateDrag(const FRay& Ray) override;
+	virtual void OnEndDrag(const FRay& Ray) override;
+
+	virtual FInputRayHit BeginHoverSequenceHitTest(const FInputDeviceRay& PressPos) override;
+	virtual bool OnUpdateHover(const FInputDeviceRay& DevicePos) override;
+
+	virtual void SetEnableRemeshing(bool bEnable) { bEnableRemeshing = bEnable; }
+	virtual bool GetEnableRemeshing() const { return bEnableRemeshing; }
+
+	virtual void DiscardAttributes();
+
+	virtual void OnPropertyModified(UObject* PropertySet, FProperty* Property) override;
+
+public:
+	/** Properties that control brush size/etc*/
+	UPROPERTY()
+	TObjectPtr<USculptBrushProperties> BrushProperties;
+
+	/** Properties that control sculpting*/
+	UPROPERTY()
+	TObjectPtr<UBrushSculptProperties> SculptProperties;
+
+	UPROPERTY()
+	TObjectPtr<USculptMaxBrushProperties> SculptMaxBrushProperties;
+	
+	UPROPERTY()
+	TObjectPtr<UKelvinBrushProperties> KelvinBrushProperties;
+
+	/** Properties that control dynamic remeshing */
+	UPROPERTY()
+	TObjectPtr<UBrushRemeshProperties> RemeshProperties;
+
+	UPROPERTY()
+	TObjectPtr<UFixedPlaneBrushProperties> GizmoProperties;
+
+	UPROPERTY()
+	TObjectPtr<UMeshEditingViewProperties> ViewProperties;
+
+	UPROPERTY()
+	TObjectPtr<UDynamicSculptToolActions> SculptToolActions;
+
+public:
+	virtual void IncreaseBrushRadiusAction();
+	virtual void DecreaseBrushRadiusAction();
+	virtual void IncreaseBrushRadiusSmallStepAction();
+	virtual void DecreaseBrushRadiusSmallStepAction();
+	virtual void IncreaseBrushFalloffAction();
+	virtual void DecreaseBrushFalloffAction();
+
+	virtual void IncreaseBrushSpeedAction();
+	virtual void DecreaseBrushSpeedAction();
+
+	virtual void NextHistoryBrushModeAction();
+	virtual void PreviousHistoryBrushModeAction();
+
+private:
+	UWorld* TargetWorld;		// required to spawn UPreviewMesh/etc
+	FViewCameraState CameraState;
+
+	UPROPERTY()
+	TObjectPtr<UBrushStampIndicator> BrushIndicator;
+
+	UPROPERTY()
+	TObjectPtr<UMaterialInstanceDynamic> BrushIndicatorMaterial;
+
+	UPROPERTY()
+	TObjectPtr<UPreviewMesh> BrushIndicatorMesh;
+
+	UPROPERTY()
+	TObjectPtr<UOctreeDynamicMeshComponent> DynamicMeshComponent;
+
+	UPROPERTY()
+	TObjectPtr<UMaterialInstanceDynamic> ActiveOverrideMaterial;
+
+	UE::Geometry::FTransform3d InitialTargetTransform;
+	UE::Geometry::FTransform3d CurTargetTransform;
+
+	// realtime visualization
+	void OnDynamicMeshComponentChanged();
+	FDelegateHandle OnDynamicMeshComponentChangedHandle;
+
+	TValueWatcher<bool> ShowWireframeWatcher;
+	TValueWatcher<EMeshEditingMaterialModes> MaterialModeWatcher;
+	TValueWatcher<TWeakObjectPtr<UMaterialInterface>> CustomMaterialWatcher;
+	TValueWatcher<bool> FlatShadingWatcher;
+	TValueWatcher<FLinearColor> ColorWatcher;
+	TValueWatcher<FLinearColor> TransparentColorWatcher;
+	TValueWatcher<double> OpacityWatcher;
+	TValueWatcher<bool> TwoSidedWatcher;
+	TValueWatcher<UTexture2D*> ImageWatcher;
+	TValueWatcher<EDynamicMeshSculptBrushType> BrushTypeWatcher;
+	TValueWatcher<FVector> GizmoPositionWatcher;
+	TValueWatcher<FQuat> GizmoRotationWatcher;
+	void UpdateMaterialMode(EMeshEditingMaterialModes NewMode);
+	void UpdateFlatShadingSetting(bool bNewValue);
+	void UpdateColorSetting(FLinearColor NewColor);
+	void UpdateOpacitySetting(double Opacity);
+	void UpdateTwoSidedSetting(bool bOn);
+	void UpdateCustomMaterial(TWeakObjectPtr<UMaterialInterface> NewMaterial);
+	void UpdateImageSetting(UTexture2D* NewImage);
+	void UpdateBrushType(EDynamicMeshSculptBrushType BrushType);
+	void UpdateGizmoFromProperties();
+
+	UE::Geometry::FInterval1d BrushRelativeSizeRange;
+	double CurrentBrushRadius;
+	void CalculateBrushRadius();
+
+	bool bEnableRemeshing;
+	double InitialEdgeLength;
+	void ScheduleRemeshPass();
+	void ConfigureRemesher(UE::Geometry::FSubRegionRemesher& Remesher);
+	void InitializeRemesherROI(UE::Geometry::FSubRegionRemesher& Remesher);
+
+	TSharedPtr<FPersistentStampRemesher> ActiveRemesher;
+	void InitializeActiveRemesher();
+	void PrecomputeRemesherROI();
+	void RemeshROIPass_ActiveRemesher(bool bHasPrecomputedROI);
+
+
+	bool bInDrag;
+
+	UE::Geometry::FFrame3d ActiveDragPlane;
+	FVector3d LastHitPosWorld;
+	FVector3d BrushStartCenterWorld;
+	FVector3d BrushStartNormalWorld;
+	FVector3d LastBrushPosLocal;
+	FVector3d LastBrushPosWorld;
+	FVector3d LastBrushPosNormalWorld;
+	FVector3d LastSmoothBrushPosLocal;
+	int32 LastBrushTriangleID = -1;
+
+	TArray<int> UpdateROITriBuffer;
+	UE::Geometry::FUniqueIndexSet VertexROIBuilder;
+	TArray<int> VertexROI;
+	UE::Geometry::FUniqueIndexSet TriangleROIBuilder;
+	TSet<int> TriangleROI;
+	//TSet<int> TriangleROI;
+	void UpdateROI(const FVector3d& BrushPos);
+
+	bool bRemeshPending;
+	bool bNormalUpdatePending;
+	
+	bool bTargetDirty;
+	TFuture<void> PendingTargetUpdate;
+
+	bool bSmoothing;
+	bool bInvert;
+	float ActivePressure = 1.0f;
+
+	bool bHaveRemeshed;
+
+	bool bStampPending;
+	FRay PendingStampRay;
+	int StampTimestamp = 0;
+	EDynamicMeshSculptBrushType LastStampType = EDynamicMeshSculptBrushType::LastValue;
+	EDynamicMeshSculptBrushType PendingStampType = LastStampType;
+	void ApplyStamp(const FRay& WorldRay);
+
+	FDynamicMesh3 BrushTargetMesh;
+	UE::Geometry::FDynamicMeshAABBTree3 BrushTargetMeshSpatial;
+	UE::Geometry::FMeshNormals BrushTargetNormals;
+	bool bCachedFreezeTarget = false;
+	void UpdateTarget();
+	bool GetTargetMeshNearest(const FVector3d& Position, double SearchRadius, FVector3d& TargetPosOut, FVector3d& TargetNormalOut);
+
+	int FindHitSculptMeshTriangle(const FRay3d& LocalRay);
+	int FindHitTargetMeshTriangle(const FRay3d& LocalRay);
+	bool IsHitTriangleBackFacing(int32 TriangleID, const FDynamicMesh3* QueryMesh);
+
+	bool UpdateBrushPosition(const FRay& WorldRay);
+	bool UpdateBrushPositionOnActivePlane(const FRay& WorldRay);
+	bool UpdateBrushPositionOnTargetMesh(const FRay& WorldRay, bool bFallbackToViewPlane);
+	bool UpdateBrushPositionOnSculptMesh(const FRay& WorldRay, bool bFallbackToViewPlane);
+	void AlignBrushToView();
+
+	bool ApplySmoothBrush(const FRay& WorldRay);
+	bool ApplyMoveBrush(const FRay& WorldRay);
+	bool ApplyOffsetBrush(const FRay& WorldRay, bool bUseViewDirection);
+	bool ApplySculptMaxBrush(const FRay& WorldRay);
+	bool ApplyPinchBrush(const FRay& WorldRay);
+	bool ApplyInflateBrush(const FRay& WorldRay);
+	bool ApplyPlaneBrush(const FRay& WorldRay);
+	bool ApplyFixedPlaneBrush(const FRay& WorldRay);
+	bool ApplyFlattenBrush(const FRay& WorldRay);
+	bool ApplyResampleBrush(const FRay& WorldRay);
+	bool ApplyPullKelvinBrush(const FRay& WorldRay);
+	bool ApplyPullSharpKelvinBrush(const FRay& WorldRay);
+	bool ApplyTwistKelvinBrush(const FRay& WorldRay);
+	bool ApplyScaleKelvinBrush(const FRay& WorldRay);
+
+	double SculptMaxFixedHeight = -1.0;
+
+	double CalculateBrushFalloff(double Distance);
+	TArray<FVector3d> ROIPositionBuffer;
+	void SyncMeshWithPositionBuffer(FDynamicMesh3* Mesh);
+
+	UE::Geometry::FFrame3d ActiveFixedBrushPlane;
+	UE::Geometry::FFrame3d ComputeROIBrushPlane(const FVector3d& BrushCenter, bool bIgnoreDepth, bool bViewAligned);
+
+	TArray<int> NormalsBuffer;
+	TArray<bool> NormalsVertexFlags;
+	void RecalculateNormals_PerVertex(const TSet<int32>& Triangles);
+	void RecalculateNormals_Overlay(const TSet<int32>& Triangles);
+
+	bool bHaveMeshBoundaries;
+	bool bHaveUVSeams;
+	bool bHaveNormalSeams;
+	TSet<int32> RemeshRemovedTriangles;
+	TSet<int32> RemeshFinalTriangleROI;
+	void PrecomputeRemeshInfo();
+	void RemeshROIPass();
+
+	FMeshVertexChangeBuilder* ActiveVertexChange = nullptr;
+	UE::Geometry::FDynamicMeshChangeTracker* ActiveMeshChange = nullptr;
+	void BeginChange(bool bIsVertexChange);
+	void EndChange();
+	void SaveActiveROI();
+
+	double EstimateIntialSafeTargetLength(const FDynamicMesh3& Mesh, int MinTargetTriCount);
+
+	TArray<EDynamicMeshSculptBrushType> BrushTypeHistory;
+	int BrushTypeHistoryIndex = 0;
+
+	UPreviewMesh* MakeDefaultSphereMesh(UObject* Parent, UWorld* World, int Resolution = 32);
+
+	//
+	// support for gizmo in FixedPlane mode
+	//
+
+	// plane gizmo
+	UPROPERTY()
+	TObjectPtr<UTransformGizmo> PlaneTransformGizmo;
+
+	UPROPERTY()
+	TObjectPtr<UTransformProxy> PlaneTransformProxy;
+
+	void PlaneTransformChanged(UTransformProxy* Proxy, FTransform Transform);
+
+	enum class EPendingWorkPlaneUpdate
+	{
+		NoUpdatePending,
+		MoveToHitPositionNormal,
+		MoveToHitPosition,
+		MoveToHitPositionViewAligned
+	};
+	EPendingWorkPlaneUpdate PendingWorkPlaneUpdate;
+	void SetFixedSculptPlaneFromWorldPos(const FVector& Position, const FVector& Normal, EPendingWorkPlaneUpdate UpdateType);
+	void UpdateFixedSculptPlanePosition(const FVector& Position);
+	void UpdateFixedSculptPlaneRotation(const FQuat& Rotation);
+	void UpdateFixedPlaneGizmoVisibility(bool bVisible);
+};
+
+
+
diff --git a/Engine/Plugins/Experimental/MeshModelingToolsetExp/Source/MeshModelingToolsExp/Public/EditNormalsTool.h b/Engine/Plugins/Experimental/MeshModelingToolsetExp/Source/MeshModelingToolsExp/Public/EditNormalsTool.h
new file mode 100644
--- /dev/null
+++ b/Engine/Plugins/Experimental/MeshModelingToolsetExp/Source/MeshModelingToolsExp/Public/EditNormalsTool.h
@@ -0,0 +1,176 @@
+// Copyright Epic Games, Inc. All Rights Reserved.
+
+#pragma once
+
+#include "CoreMinimal.h"
+#include "UObject/NoExportTypes.h"
+#include "MultiSelectionTool.h"
+#include "InteractiveToolBuilder.h"
+#include "MeshOpPreviewHelpers.h"
+#include "CleaningOps/EditNormalsOp.h"
+#include "DynamicMesh/DynamicMesh3.h"
+#include "BaseTools/SingleClickTool.h"
+#include "EditNormalsTool.generated.h"
+
+
+// predeclarations
+struct FMeshDescription;
+class UDynamicMeshComponent;
+
+/**
+ *
+ */
+UCLASS()
+class MESHMODELINGTOOLSEXP_API UEditNormalsToolBuilder : public UInteractiveToolBuilder
+{
+	GENERATED_BODY()
+
+public:
+	virtual bool CanBuildTool(const FToolBuilderState& SceneState) const override;
+	virtual UInteractiveTool* BuildTool(const FToolBuilderState& SceneState) const override;
+
+protected:
+	virtual const FToolTargetTypeRequirements& GetTargetRequirements() const override;
+};
+
+
+
+
+
+/**
+ * Standard properties
+ */
+UCLASS()
+class MESHMODELINGTOOLSEXP_API UEditNormalsToolProperties : public UInteractiveToolPropertySet
+{
+	GENERATED_BODY()
+
+public:
+	UEditNormalsToolProperties();
+
+	bool WillTopologyChange()
+	{
+		return bFixInconsistentNormals || bInvertNormals || SplitNormalMethod != ESplitNormalMethod::UseExistingTopology;
+	}
+
+	/** Recompute all mesh normals */
+	UPROPERTY(EditAnywhere, Category = NormalsCalculation, meta = (EditCondition = "SplitNormalMethod == ESplitNormalMethod::UseExistingTopology"))
+	bool bRecomputeNormals;
+
+	/** Choose the method for computing vertex normals */
+	UPROPERTY(EditAnywhere, Category = NormalsCalculation)
+	ENormalCalculationMethod NormalCalculationMethod;
+
+	/** For meshes with inconsistent triangle orientations/normals, flip as needed to make the normals consistent */
+	UPROPERTY(EditAnywhere, Category = NormalsCalculation)
+	bool bFixInconsistentNormals;
+
+	/** Invert (flip) all mesh normals and associated triangle orientations */
+	UPROPERTY(EditAnywhere, Category = NormalsCalculation)
+	bool bInvertNormals;
+
+	/** Control whether and how the topology of the normals is recomputed, e.g. to create sharp edges where face normals change by a large amount or where face group IDs change.  Normals will always be recomputed unless SplitNormal Method is UseExistingTopology. */
+	UPROPERTY(EditAnywhere, Category = NormalsTopology)
+	ESplitNormalMethod SplitNormalMethod;
+
+	/** Threshold on angle of change in face normals across an edge, above which we create a sharp edge if bSplitNormals is true */
+	UPROPERTY(EditAnywhere, Category = NormalsTopology, meta = (UIMin = "0.0", UIMax = "180.0", ClampMin = "0.0", ClampMax = "180.0", EditCondition = "SplitNormalMethod == ESplitNormalMethod::FaceNormalThreshold"))
+	float SharpEdgeAngleThreshold;
+
+	/** Assign separate normals at 'sharp' vertices -- for example, at the tip of a cone */
+	UPROPERTY(EditAnywhere, Category = NormalsTopology, meta = (EditCondition = "SplitNormalMethod == ESplitNormalMethod::FaceNormalThreshold"))
+	bool bAllowSharpVertices;
+};
+
+
+
+
+/**
+ * Advanced properties
+ */
+UCLASS()
+class MESHMODELINGTOOLSEXP_API UEditNormalsAdvancedProperties : public UInteractiveToolPropertySet
+{
+	GENERATED_BODY()
+
+public:
+	UEditNormalsAdvancedProperties();
+};
+
+
+/**
+ * Factory with enough info to spawn the background-thread Operator to do a chunk of work for the tool
+ *  stores a pointer to the tool and enough info to know which specific operator it should spawn
+ */
+UCLASS()
+class MESHMODELINGTOOLSEXP_API UEditNormalsOperatorFactory : public UObject, public UE::Geometry::IDynamicMeshOperatorFactory
+{
+	GENERATED_BODY()
+
+public:
+	// IDynamicMeshOperatorFactory API
+	virtual TUniquePtr<UE::Geometry::FDynamicMeshOperator> MakeNewOperator() override;
+
+	UPROPERTY()
+	TObjectPtr<UEditNormalsTool> Tool;
+
+	int ComponentIndex;
+
+};
+
+/**
+ * Simple Mesh Normal Updating Tool
+ */
+UCLASS()
+class MESHMODELINGTOOLSEXP_API UEditNormalsTool : public UMultiSelectionTool
+{
+	GENERATED_BODY()
+
+public:
+
+	friend UEditNormalsOperatorFactory;
+
+	UEditNormalsTool();
+
+	virtual void Setup() override;
+	virtual void Shutdown(EToolShutdownType ShutdownType) override;
+
+	virtual void SetWorld(UWorld* World);
+
+	virtual void OnTick(float DeltaTime) override;
+	virtual void Render(IToolsContextRenderAPI* RenderAPI) override;
+
+	virtual bool HasCancel() const override { return true; }
+	virtual bool HasAccept() const override;
+	virtual bool CanAccept() const override;
+
+#if WITH_EDITOR
+	virtual void PostEditChangeProperty(FPropertyChangedEvent &PropertyChangedEvent) override;
+#endif
+
+	virtual void OnPropertyModified(UObject* PropertySet, FProperty* Property) override;
+
+protected:
+
+	UPROPERTY()
+	TObjectPtr<UEditNormalsToolProperties> BasicProperties;
+
+	UPROPERTY()
+	TObjectPtr<UEditNormalsAdvancedProperties> AdvancedProperties;
+
+
+	UPROPERTY()
+	TArray<TObjectPtr<UMeshOpPreviewWithBackgroundCompute>> Previews;
+
+
+protected:
+	TArray<TSharedPtr<UE::Geometry::FDynamicMesh3, ESPMode::ThreadSafe>> OriginalDynamicMeshes;
+
+	UWorld* TargetWorld;
+
+	FViewCameraState CameraState;
+
+	void UpdateNumPreviews();
+
+	void GenerateAsset(const TArray<FDynamicMeshOpResult>& Results);
+};
