MeshModelingToolset: Split plugin into non-experimental and experimental

diff --git a/Engine/Plugins/Experimental/MeshModelingToolsetExp/Source/MeshModelingToolsExp/Public/EditPivotTool.h b/Engine/Plugins/Experimental/MeshModelingToolsetExp/Source/MeshModelingToolsExp/Public/EditPivotTool.h
new file mode 100644
--- /dev/null
+++ b/Engine/Plugins/Experimental/MeshModelingToolsetExp/Source/MeshModelingToolsExp/Public/EditPivotTool.h
@@ -0,0 +1,222 @@
+// Copyright Epic Games, Inc. All Rights Reserved.
+
+#pragma once
+
+#include "CoreMinimal.h"
+#include "MultiSelectionTool.h"
+#include "InteractiveToolBuilder.h"
+#include "BaseBehaviors/BehaviorTargetInterfaces.h"
+#include "Changes/TransformChange.h"
+#include "FrameTypes.h"
+#include "BoxTypes.h"
+#include "EditPivotTool.generated.h"
+
+class UDragAlignmentMechanic;
+class UBaseAxisTranslationGizmo;
+class UAxisAngleGizmo;
+class UTransformGizmo;
+class UTransformProxy;
+
+/**
+ *
+ */
+UCLASS()
+class MESHMODELINGTOOLSEXP_API UEditPivotToolBuilder : public UInteractiveToolBuilder
+{
+	GENERATED_BODY()
+
+public:
+	virtual bool CanBuildTool(const FToolBuilderState& SceneState) const override;
+	virtual UInteractiveTool* BuildTool(const FToolBuilderState& SceneState) const override;
+
+protected:
+	virtual const FToolTargetTypeRequirements& GetTargetRequirements() const override;
+};
+
+
+
+
+/** Snap-Drag Rotation Mode */
+UENUM()
+enum class EEditPivotSnapDragRotationMode : uint8
+{
+	/** Snap-Drag only translates, ignoring Normals */
+	Ignore = 0 UMETA(DisplayName = "Ignore"),
+
+	/** Snap-Drag aligns the Source and Target Normals to point in the same direction */
+	Align = 1 UMETA(DisplayName = "Align"),
+
+	/** Snap-Drag aligns the Source Normal to the opposite of the Target Normal direction */
+	AlignFlipped = 2 UMETA(DisplayName = "Align Flipped"),
+
+	LastValue UMETA(Hidden)
+};
+
+
+
+/**
+ * Standard properties of the Transform Meshes operation
+ */
+UCLASS()
+class MESHMODELINGTOOLSEXP_API UEditPivotToolProperties : public UInteractiveToolPropertySet
+{
+	GENERATED_BODY()
+public:
+
+	/** When enabled, click-drag to reposition the Pivot */
+	UPROPERTY(EditAnywhere, Category = Options)
+	bool bEnableSnapDragging = false;
+
+	/** When Snap-Dragging, align source and target normals */
+	UPROPERTY(EditAnywhere, Category = Options, meta = (EditCondition = "bEnableSnapDragging == true"))
+	EEditPivotSnapDragRotationMode RotationMode = EEditPivotSnapDragRotationMode::Align;
+};
+
+
+USTRUCT()
+struct FEditPivotTarget
+{
+	GENERATED_BODY()
+
+	UPROPERTY()
+	TObjectPtr<UTransformProxy> TransformProxy = nullptr;
+
+	UPROPERTY()
+	TObjectPtr<UTransformGizmo> TransformGizmo = nullptr;
+};
+
+
+
+
+UENUM()
+enum class EEditPivotToolActions
+{
+	NoAction,
+
+	Center,
+	Bottom,
+	Top,
+	Left,
+	Right,
+	Front,
+	Back,
+	WorldOrigin
+};
+
+
+
+UCLASS()
+class MESHMODELINGTOOLSEXP_API UEditPivotToolActionPropertySet : public UInteractiveToolPropertySet
+{
+	GENERATED_BODY()
+
+public:
+	TWeakObjectPtr<UEditPivotTool> ParentTool;
+
+	void Initialize(UEditPivotTool* ParentToolIn) { ParentTool = ParentToolIn; }
+	void PostAction(EEditPivotToolActions Action);
+
+	/** Use the World-Space Bounding Box of the target object, instead of the Object-space Bounding Box */
+	UPROPERTY(EditAnywhere, Category = BoxPositions)
+	bool bUseWorldBox = false;
+
+	UFUNCTION(CallInEditor, Category = BoxPositions, meta = (DisplayPriority = 1))
+	void Center() { PostAction(EEditPivotToolActions::Center); }
+
+	UFUNCTION(CallInEditor, Category = BoxPositions, meta = (DisplayPriority = 2))
+	void Bottom() { PostAction(EEditPivotToolActions::Bottom ); }
+
+	UFUNCTION(CallInEditor, Category = BoxPositions, meta = (DisplayPriority = 2))
+	void Top() { PostAction(EEditPivotToolActions::Top); }
+
+	UFUNCTION(CallInEditor, Category = BoxPositions, meta = (DisplayPriority = 3))
+	void Left() { PostAction(EEditPivotToolActions::Left); }
+
+	UFUNCTION(CallInEditor, Category = BoxPositions, meta = (DisplayPriority = 3))
+	void Right() { PostAction(EEditPivotToolActions::Right); }
+
+	UFUNCTION(CallInEditor, Category = BoxPositions, meta = (DisplayPriority = 4))
+	void Front() { PostAction(EEditPivotToolActions::Front); }
+
+	UFUNCTION(CallInEditor, Category = BoxPositions, meta = (DisplayPriority = 4))
+	void Back() { PostAction(EEditPivotToolActions::Back); }
+
+	UFUNCTION(CallInEditor, Category = BoxPositions, meta = (DisplayPriority = 5))
+	void WorldOrigin() { PostAction(EEditPivotToolActions::WorldOrigin); }
+};
+
+
+
+/**
+ *
+ */
+UCLASS()
+class MESHMODELINGTOOLSEXP_API UEditPivotTool : public UMultiSelectionTool, public IClickDragBehaviorTarget
+{
+	GENERATED_BODY()
+
+public:
+	UEditPivotTool();
+
+	virtual void RegisterActions(FInteractiveToolActionSet& ActionSet) override;
+
+	virtual void SetWorld(UWorld* World, UInteractiveGizmoManager* GizmoManager);
+
+	virtual void Setup() override;
+	virtual void Shutdown(EToolShutdownType ShutdownType) override;
+
+	virtual void OnTick(float DeltaTime) override;
+	virtual void Render(IToolsContextRenderAPI* RenderAPI) override;
+
+	virtual bool HasCancel() const override { return true; }
+	virtual bool HasAccept() const override { return true; }
+
+	virtual void OnPropertyModified(UObject* PropertySet, FProperty* Property) override;
+
+
+	// ICLickDragBehaviorTarget interface
+	virtual FInputRayHit CanBeginClickDragSequence(const FInputDeviceRay& PressPos) override;
+	virtual void OnClickPress(const FInputDeviceRay& PressPos) override;
+	virtual void OnClickDrag(const FInputDeviceRay& DragPos) override;
+	virtual void OnClickRelease(const FInputDeviceRay& ReleasePos) override;
+	virtual void OnTerminateDragSequence() override;
+
+public:
+	UPROPERTY()
+	TObjectPtr<UEditPivotToolProperties> TransformProps;
+
+	UPROPERTY()
+	TObjectPtr<UEditPivotToolActionPropertySet> EditPivotActions;
+
+	virtual void RequestAction(EEditPivotToolActions ActionType);
+
+protected:
+	UWorld* TargetWorld;
+	UInteractiveGizmoManager* GizmoManager;
+
+	TArray<int> MapToFirstOccurrences;
+
+	UE::Geometry::FTransform3d Transform;
+	UE::Geometry::FAxisAlignedBox3d ObjectBounds;
+	UE::Geometry::FAxisAlignedBox3d WorldBounds;
+	void Precompute();
+
+	UPROPERTY()
+	TArray<FEditPivotTarget> ActiveGizmos;
+
+	UPROPERTY()
+	TObjectPtr<UDragAlignmentMechanic> DragAlignmentMechanic = nullptr;
+
+	void UpdateSetPivotModes(bool bEnableSetPivot);
+	void SetActiveGizmos_Single(bool bLocalRotations);
+	void ResetActiveGizmos();
+
+	FTransform StartDragTransform;
+
+	EEditPivotToolActions PendingAction;
+	virtual void ApplyAction(EEditPivotToolActions ActionType);
+	virtual void SetPivotToBoxPoint(EEditPivotToolActions ActionPoint);
+	virtual void SetPivotToWorldOrigin();
+
+	void UpdateAssets(const UE::Geometry::FFrame3d& NewPivotWorldFrame);
+};
diff --git a/Engine/Plugins/Experimental/MeshModelingToolsetExp/Source/MeshModelingToolsExp/Public/EditUVIslandsTool.h b/Engine/Plugins/Experimental/MeshModelingToolsetExp/Source/MeshModelingToolsExp/Public/EditUVIslandsTool.h
new file mode 100644
--- /dev/null
+++ b/Engine/Plugins/Experimental/MeshModelingToolsetExp/Source/MeshModelingToolsExp/Public/EditUVIslandsTool.h
@@ -0,0 +1,160 @@
+// Copyright Epic Games, Inc. All Rights Reserved.
+
+#pragma once
+
+#include "CoreMinimal.h"
+#include "BaseTools/MeshSurfacePointTool.h"
+#include "Properties/MeshMaterialProperties.h"
+#include "Components/DynamicMeshComponent.h"
+#include "DynamicMesh/DynamicMeshAABBTree3.h"
+#include "ToolDataVisualizer.h"
+#include "Changes/MeshVertexChange.h"
+#include "GroupTopology.h"
+#include "Selection/GroupTopologySelector.h"
+#include "ModelingOperators/Public/ModelingTaskTypes.h"
+#include "Transforms/MultiTransformer.h"
+#include "Selection/PolygonSelectionMechanic.h"
+#include "EditUVIslandsTool.generated.h"
+
+class FMeshVertexChangeBuilder;
+using UE::Geometry::FDynamicMeshUVOverlay;
+
+/**
+ * ToolBuilder
+ */
+UCLASS()
+class MESHMODELINGTOOLSEXP_API UEditUVIslandsToolBuilder : public UMeshSurfacePointToolBuilder
+{
+	GENERATED_BODY()
+public:
+	virtual UMeshSurfacePointTool* CreateNewTool(const FToolBuilderState& SceneState) const override;
+};
+
+
+
+
+class MESHMODELINGTOOLSEXP_API FUVGroupTopology : public FGroupTopology
+{
+public:
+	TArray<int32> TriIslandGroups;
+	const FDynamicMeshUVOverlay* UVOverlay;
+
+	FUVGroupTopology() {}
+	FUVGroupTopology(const UE::Geometry::FDynamicMesh3* Mesh, uint32 UVLayerIndex, bool bAutoBuild = false);
+
+	void CalculateIslandGroups();
+
+	virtual int GetGroupID(int32 TriangleID) const override
+	{
+		return TriIslandGroups[TriangleID];
+	}
+
+	UE::Geometry::FFrame3d GetIslandFrame(int32 GroupID, UE::Geometry::FDynamicMeshAABBTree3& AABBTree);
+};
+
+
+/**
+ *
+ */
+UCLASS()
+class MESHMODELINGTOOLSEXP_API UEditUVIslandsTool : public UMeshSurfacePointTool
+{
+	GENERATED_BODY()
+
+public:
+	UEditUVIslandsTool();
+
+	virtual void RegisterActions(FInteractiveToolActionSet& ActionSet) override;
+
+	virtual void Setup() override;
+	virtual void Shutdown(EToolShutdownType ShutdownType) override;
+
+	virtual void OnTick(float DeltaTime) override;
+	virtual void Render(IToolsContextRenderAPI* RenderAPI) override;
+
+	virtual bool HasCancel() const override { return true; }
+	virtual bool HasAccept() const override { return true; }
+
+	// UMeshSurfacePointTool API
+	virtual bool HitTest(const FRay& Ray, FHitResult& OutHit) override;
+	virtual void OnBeginDrag(const FRay& Ray) override;
+	virtual void OnUpdateDrag(const FRay& Ray) override;
+	virtual void OnEndDrag(const FRay& Ray) override;
+	virtual bool OnUpdateHover(const FInputDeviceRay& DevicePos) override;
+	virtual void OnEndHover() override;
+
+	// IClickDragBehaviorTarget API
+	virtual FInputRayHit CanBeginClickDragSequence(const FInputDeviceRay& PressPos) override;
+
+
+public:
+	UPROPERTY()
+	TObjectPtr<UExistingMeshMaterialProperties> MaterialSettings = nullptr;
+
+	UPROPERTY()
+	TObjectPtr<UMaterialInstanceDynamic> CheckerMaterial = nullptr;
+
+protected:
+	UPROPERTY()
+	TObjectPtr<UDynamicMeshComponent> DynamicMeshComponent;
+
+	UPROPERTY()
+	TObjectPtr<UPolygonSelectionMechanic> SelectionMechanic;
+
+	bool bSelectionStateDirty = false;
+	void OnSelectionModifiedEvent();
+
+	UPROPERTY()
+	TObjectPtr<UMultiTransformer> MultiTransformer = nullptr;
+
+	void OnMultiTransformerTransformBegin();
+	void OnMultiTransformerTransformUpdate();
+	void OnMultiTransformerTransformEnd();
+
+	// realtime visualization
+	void OnDynamicMeshComponentChanged();
+	FDelegateHandle OnDynamicMeshComponentChangedHandle;
+
+	// camera state at last render
+	UE::Geometry::FTransform3d WorldTransform;
+	FViewCameraState CameraState;
+
+	// True for the duration of UI click+drag
+	bool bInDrag;
+
+	double UVTranslateScale;
+	UE::Geometry::FFrame3d InitialGizmoFrame;
+	FVector3d InitialGizmoScale;
+	void ComputeUpdate_Gizmo();
+
+	FUVGroupTopology Topology;
+	void PrecomputeTopology();
+
+	FDynamicMeshAABBTree3 MeshSpatial;
+	FDynamicMeshAABBTree3& GetSpatial();
+	bool bSpatialDirty;
+
+
+	//
+	// data for current drag
+	//
+	struct FEditIsland
+	{
+		UE::Geometry::FFrame3d LocalFrame;
+		TArray<int32> Triangles;
+		TArray<int32> UVs;
+		UE::Geometry::FAxisAlignedBox2d UVBounds;
+		FVector2d UVOrigin;
+		TArray<FVector2f> InitialPositions;
+	};
+	TArray<FEditIsland> ActiveIslands;
+	void UpdateUVTransformFromSelection(const FGroupTopologySelection& Selection);
+
+	FMeshVertexChangeBuilder* ActiveVertexChange;
+	void BeginChange();
+	void EndChange();
+	void UpdateChangeFromROI(bool bFinal);
+
+	void OnMaterialSettingsChanged();
+};
+
diff --git a/Engine/Plugins/Experimental/MeshModelingToolsetExp/Source/MeshModelingToolsExp/Public/FindPolygonsAlgorithm.h b/Engine/Plugins/Experimental/MeshModelingToolsetExp/Source/MeshModelingToolsExp/Public/FindPolygonsAlgorithm.h
new file mode 100644
--- /dev/null
+++ b/Engine/Plugins/Experimental/MeshModelingToolsetExp/Source/MeshModelingToolsExp/Public/FindPolygonsAlgorithm.h
@@ -0,0 +1,54 @@
+// Copyright Epic Games, Inc. All Rights Reserved.
+
+#pragma once
+
+#include "CoreMinimal.h"
+#include "DynamicMesh/DynamicMesh3.h"
+
+using UE::Geometry::FDynamicMesh3;
+
+class FFindPolygonsAlgorithm
+{
+public:
+
+	FFindPolygonsAlgorithm() {}
+	FFindPolygonsAlgorithm(FDynamicMesh3* MeshIn);
+
+	FDynamicMesh3* Mesh = nullptr;
+	TArray<TArray<int>> FoundPolygons;
+	TArray<int> PolygonTags;
+	TArray<FVector3d> PolygonNormals;
+
+	// if > 1, groups with a triangle count smaller than this will be merged with a neighbouring group
+	int32 MinGroupSize = 0;
+
+	TArray<int> PolygonEdges;
+
+	bool FindPolygonsFromFaceNormals(double DotTolerance = 0.0001);
+	bool FindPolygonsFromUVIslands();
+	bool FindPolygonsFromConnectedTris();
+
+	enum class EWeightingType
+	{
+		None,
+		NormalDeviation
+	};
+
+	/**
+	 * Incrementally compute approximate-geodesic-based furthest-point sampling of the mesh until NumPoints
+	 * samples have been found, then compute local geodesic patches (eg approximate surface voronoi diagaram).
+	 * Optionally weight geodesic-distance computation, which will produce different patch shapes.
+	 */
+	bool FindPolygonsFromFurthestPointSampling(int32 NumPoints, EWeightingType WeightingType, FVector3d WeightingCoeffs = FVector3d::One());
+
+
+	bool FindPolygonEdges();
+
+
+protected:
+
+	void PostProcessPolygons(bool bApplyMerging);
+	void OptimizePolygons();
+
+	void SetGroupsFromPolygons();
+};
\ No newline at end of file
diff --git a/Engine/Plugins/Experimental/MeshModelingToolsetExp/Source/MeshModelingToolsExp/Public/HoleFillTool.h b/Engine/Plugins/Experimental/MeshModelingToolsetExp/Source/MeshModelingToolsExp/Public/HoleFillTool.h
new file mode 100644
--- /dev/null
+++ b/Engine/Plugins/Experimental/MeshModelingToolsetExp/Source/MeshModelingToolsExp/Public/HoleFillTool.h
@@ -0,0 +1,324 @@
+// Copyright Epic Games, Inc. All Rights Reserved.
+
+// HoleFillTool: Fill one or more boundary loops on a selected mesh. Several hole-filling methods are available.
+
+#pragma once
+
+#include "CoreMinimal.h"
+#include "InteractiveToolBuilder.h"
+#include "SingleSelectionTool.h"
+#include "CleaningOps/HoleFillOp.h"
+#include "BaseTools/MeshSurfacePointTool.h"
+#include "BaseBehaviors/BehaviorTargetInterfaces.h"
+#include "DynamicMesh/DynamicMeshAABBTree3.h"
+#include "GroupTopology.h"
+#include "BaseTools/SingleSelectionMeshEditingTool.h"
+#include "HoleFillTool.generated.h"
+
+class UPolygonSelectionMechanic;
+class UDynamicMeshReplacementChangeTarget;
+class UMeshOpPreviewWithBackgroundCompute;
+class UDynamicMeshComponent;
+struct FDynamicMeshOpResult;
+class UHoleFillTool;
+
+/*
+ * Tool builder
+ */
+
+UCLASS()
+class MESHMODELINGTOOLSEXP_API UHoleFillToolBuilder : public USingleSelectionMeshEditingToolBuilder
+{
+	GENERATED_BODY()
+
+public:
+	virtual USingleSelectionMeshEditingTool* CreateNewTool(const FToolBuilderState& SceneState) const override;
+};
+
+/*
+ * Properties. This class reflects the parameters in FSmoothFillOptions, but is decorated to allow use in the UI system.
+ */
+
+UCLASS()
+class MESHMODELINGTOOLSEXP_API USmoothHoleFillProperties : public UInteractiveToolPropertySet
+{
+	GENERATED_BODY()
+
+public:
+
+	/** Allow smoothing and remeshing of triangles outside of the fill region */
+	UPROPERTY(EditAnywhere, Category = SmoothHoleFillOptions)
+	bool bConstrainToHoleInterior;
+
+	/** Number of vertex rings outside of the fill region to allow remeshing */
+	UPROPERTY(EditAnywhere, Category = SmoothHoleFillOptions, 
+		meta = (UIMin = "0", ClampMin = "0", EditCondition = "!bConstrainToHoleInterior"))
+	int RemeshingExteriorRegionWidth;
+
+	/** Number of vertex rings outside of the fill region to perform smoothing */
+	UPROPERTY(EditAnywhere, Category = SmoothHoleFillOptions, meta = (UIMin = "0", ClampMin = "0"))
+	int SmoothingExteriorRegionWidth;
+
+	/** Number of vertex rings away from the fill region boundary to constrain smoothing */
+	UPROPERTY(EditAnywhere, Category = SmoothHoleFillOptions, meta = (UIMin = "0", ClampMin = "0"))
+	int SmoothingInteriorRegionWidth;
+
+	/** Desired Smoothness. This is not a linear quantity, but larger numbers produce smoother results */
+	UPROPERTY(EditAnywhere, Category = SmoothHoleFillOptions, meta = (UIMin = "0.0", UIMax = "1.0", ClampMin = "0.0", ClampMax = "100.0"))
+	float InteriorSmoothness;
+
+	/** Relative triangle density of fill region */
+	UPROPERTY(EditAnywhere, Category = SmoothHoleFillOptions, meta = (UIMin = "0.001", UIMax = "10.0", ClampMin = "0.001", ClampMax = "10.0"))
+	double FillDensityScalar = 1.0;
+
+	/** 
+	 * Whether to project to the original mesh during post-smooth remeshing. This can be expensive on large meshes with 
+	 * many holes. 
+	 */
+	UPROPERTY(EditAnywhere, Category = SmoothHoleFillOptions)
+	bool bProjectDuringRemesh = false;
+
+
+	// Set default property values
+	USmoothHoleFillProperties()
+	{
+		// Create a default FSmoothFillOptions and populate this class with its defaults.
+		UE::Geometry::FSmoothFillOptions DefaultOptionsObject;
+		bConstrainToHoleInterior = DefaultOptionsObject.bConstrainToHoleInterior;
+		RemeshingExteriorRegionWidth = DefaultOptionsObject.RemeshingExteriorRegionWidth;
+		SmoothingExteriorRegionWidth = DefaultOptionsObject.SmoothingExteriorRegionWidth;
+		SmoothingInteriorRegionWidth = DefaultOptionsObject.SmoothingInteriorRegionWidth;
+		FillDensityScalar = DefaultOptionsObject.FillDensityScalar;
+		InteriorSmoothness = DefaultOptionsObject.InteriorSmoothness;
+		bProjectDuringRemesh = DefaultOptionsObject.bProjectDuringRemesh;
+	}
+
+	UE::Geometry::FSmoothFillOptions ToSmoothFillOptions() const
+	{
+		UE::Geometry::FSmoothFillOptions Options;
+		Options.bConstrainToHoleInterior = bConstrainToHoleInterior;
+		Options.RemeshingExteriorRegionWidth = RemeshingExteriorRegionWidth;
+		Options.SmoothingExteriorRegionWidth = SmoothingExteriorRegionWidth;
+		Options.SmoothingInteriorRegionWidth = SmoothingInteriorRegionWidth;
+		Options.InteriorSmoothness = InteriorSmoothness;
+		Options.FillDensityScalar = FillDensityScalar;
+		Options.bProjectDuringRemesh = bProjectDuringRemesh;
+		return Options;
+	}
+};
+
+UCLASS()
+class MESHMODELINGTOOLSEXP_API UHoleFillToolProperties : public UInteractiveToolPropertySet
+{
+	GENERATED_BODY()
+
+public:
+	UPROPERTY(EditAnywhere, Category = Options)
+	EHoleFillOpFillType FillType = EHoleFillOpFillType::Minimal;
+
+	/** Clean up triangles that have no neighbors */
+	UPROPERTY(EditAnywhere, Category = Options)
+	bool bRemoveIsolatedTriangles = false;
+
+	/** Identify and quickly fill single-triangle holes */
+	UPROPERTY(EditAnywhere, Category = Options)
+	bool bQuickFillSmallHoles = false;
+
+};
+
+
+UENUM()
+enum class EHoleFillToolActions
+{
+	NoAction,
+	SelectAll,
+	ClearSelection
+};
+
+UCLASS()
+class MESHMODELINGTOOLSEXP_API UHoleFillToolActions : public UInteractiveToolPropertySet
+{
+GENERATED_BODY()
+
+	TWeakObjectPtr<UHoleFillTool> ParentTool;
+
+public:
+
+	void Initialize(UHoleFillTool* ParentToolIn)
+	{
+		ParentTool = ParentToolIn;
+	}
+
+	void PostAction(EHoleFillToolActions Action);
+
+	UFUNCTION(CallInEditor, Category = SelectionEdits, meta = (DisplayName = "Select All", DisplayPriority = 1))
+		void SelectAll()
+	{
+		PostAction(EHoleFillToolActions::SelectAll);
+	}
+
+	UFUNCTION(CallInEditor, Category = SelectionEdits, meta = (DisplayName = "Clear", DisplayPriority = 1))
+		void Clear()
+	{
+		PostAction(EHoleFillToolActions::ClearSelection);
+	}
+};
+
+
+UCLASS()
+class MESHMODELINGTOOLSEXP_API UHoleFillStatisticsProperties : public UInteractiveToolPropertySet
+{
+	GENERATED_BODY()
+public:
+
+	UPROPERTY(VisibleAnywhere, Category = HoleFillStatistics)
+	FString InitialHoles;
+
+	UPROPERTY(VisibleAnywhere, Category = HoleFillStatistics)
+	FString SelectedHoles;
+
+	UPROPERTY(VisibleAnywhere, Category = HoleFillStatistics)
+	FString SuccessfulFills;
+
+	UPROPERTY(VisibleAnywhere, Category = HoleFillStatistics)
+	FString FailedFills;
+
+	UPROPERTY(VisibleAnywhere, Category = HoleFillStatistics)
+	FString RemainingHoles;
+
+	void Initialize(const UHoleFillTool& HoleFillTool);
+
+	void Update(const UHoleFillTool& HoleFillTool, const UE::Geometry::FHoleFillOp& HoleFillOp);
+};
+
+
+/*
+ * Operator factory
+ */
+
+UCLASS()
+class MESHMODELINGTOOLSEXP_API UHoleFillOperatorFactory : public UObject, public UE::Geometry::IDynamicMeshOperatorFactory
+{
+	GENERATED_BODY()
+
+public:
+
+	TUniquePtr<UE::Geometry::FDynamicMeshOperator> MakeNewOperator() override;
+
+	UPROPERTY()
+	TObjectPtr<UHoleFillTool> FillTool;
+};
+
+/*
+ * Tool
+ * Inherit from IClickBehaviorTarget so we can click on boundary loops.
+ */
+
+UCLASS()
+class MESHMODELINGTOOLSEXP_API UHoleFillTool : public USingleSelectionMeshEditingTool
+{
+	GENERATED_BODY()
+
+public:
+
+	// UMeshSurfacePointTool
+	void Setup() override;
+	void OnTick(float DeltaTime) override;
+	void OnPropertyModified(UObject* PropertySet, FProperty* Property) override;
+	bool HasCancel() const override { return true; }
+	bool HasAccept() const override { return true; }
+	bool CanAccept() const override;
+	void Shutdown(EToolShutdownType ShutdownType) override;
+
+	void OnSelectionModified();
+
+	virtual void RequestAction(EHoleFillToolActions Action);
+
+protected:
+
+	friend UHoleFillOperatorFactory;
+	friend UHoleFillToolBuilder;
+	friend UHoleFillStatisticsProperties;
+
+	UPROPERTY()
+	TObjectPtr<USmoothHoleFillProperties> SmoothHoleFillProperties = nullptr;
+
+	UPROPERTY()
+	TObjectPtr<UHoleFillToolProperties> Properties = nullptr;
+
+	UPROPERTY()
+	TObjectPtr<UHoleFillToolActions> Actions = nullptr;
+
+	UPROPERTY()
+	TObjectPtr<UHoleFillStatisticsProperties> Statistics = nullptr;
+
+	UPROPERTY()
+	TObjectPtr<UMeshOpPreviewWithBackgroundCompute> Preview = nullptr;
+
+	UPROPERTY()
+	TObjectPtr<UPolygonSelectionMechanic> SelectionMechanic = nullptr;
+
+	// Input mesh. Ownership shared with Op.
+	TSharedPtr<UE::Geometry::FDynamicMesh3, ESPMode::ThreadSafe> OriginalMesh;
+
+	// UV Scale factor is cached based on the bounding box of the mesh before any fills are performed
+	float MeshUVScaleFactor = 0.0f;
+
+	// Used for hit querying
+	UE::Geometry::FDynamicMeshAABBTree3 MeshSpatial;
+
+	TSet<int32> NewTriangleIDs;
+
+	// Create the Preview object
+	void SetupPreview();
+
+	// Invalidate background compute result (some input changed)
+	void InvalidatePreviewResult();
+
+	bool bHavePendingAction = false;
+	EHoleFillToolActions PendingAction;
+	virtual void ApplyAction(EHoleFillToolActions ActionType);
+	void SelectAll();
+	void ClearSelection();
+
+	// Object used to get boundary loop information. All triangles return the same GroupID, so the only boundaries that 
+	// are returned are the the actual mesh boundaries.
+	// TODO: It seems like overkill to use a FGroupTopology subclass when we don't actually care about groups.
+	class FBasicTopology : public UE::Geometry::FGroupTopology
+	{
+	public:
+		FBasicTopology(const UE::Geometry::FDynamicMesh3* Mesh, bool bAutoBuild) :
+			FGroupTopology(Mesh, false)
+		{
+			if (bAutoBuild)
+			{
+				// Virtual func resolution doesn't work in constructors. Though we're not currently
+				// overriding RebuildTopology, let's do the proper thing in case we get copy-pasted
+				// somewhere where we do.
+				RebuildTopology();
+			}
+		}
+
+		int GetGroupID(int TriangleID) const override
+		{
+			return Mesh->IsTriangle(TriangleID) ? 1 : 0;
+		}
+	};
+	TUniquePtr<FBasicTopology> Topology;
+
+	struct FSelectedBoundaryLoop
+	{
+		int32 EdgeTopoID;
+		TArray<int32> EdgeIDs;
+	};
+	TArray<FSelectedBoundaryLoop> ActiveBoundaryLoopSelection;
+	void UpdateActiveBoundaryLoopSelection();
+
+	// Just call the SelectionMechanism's Render function
+	void Render(IToolsContextRenderAPI* RenderAPI) override;
+
+	// Populate an array of Edge Loops to be processed by an FHoleFillOp. Returns the edge loops currently selected
+	// by this tool.
+	void GetLoopsToFill(TArray<UE::Geometry::FEdgeLoop>& OutLoops) const;
+
+};
diff --git a/Engine/Plugins/Experimental/MeshModelingToolsetExp/Source/MeshModelingToolsExp/Public/MeshAttributePaintTool.h b/Engine/Plugins/Experimental/MeshModelingToolsetExp/Source/MeshModelingToolsExp/Public/MeshAttributePaintTool.h
new file mode 100644
--- /dev/null
+++ b/Engine/Plugins/Experimental/MeshModelingToolsetExp/Source/MeshModelingToolsExp/Public/MeshAttributePaintTool.h
@@ -0,0 +1,298 @@
+// Copyright Epic Games, Inc. All Rights Reserved.
+
+#pragma once
+
+#include "DynamicMeshBrushTool.h"
+#include "Changes/ValueWatcher.h"
+#include "Changes/IndexedAttributeChange.h"
+#include "DynamicMesh/DynamicVerticesOctree3.h"
+#include "MeshDescription.h"
+#include "MeshAttributePaintTool.generated.h"
+
+
+struct FMeshDescription;
+
+
+/**
+ * Maps float values to linear color ramp.
+ */
+class FFloatAttributeColorMapper
+{
+public:
+	virtual ~FFloatAttributeColorMapper() {}
+
+	FLinearColor LowColor = FLinearColor(0.9f, 0.9f, 0.9f, 1.0f);
+	FLinearColor HighColor = FLinearColor(1.0f, 0.0f, 0.0f, 1.0f);
+
+	virtual FLinearColor ToColor(float Value)
+	{
+		float t = FMath::Clamp(Value, 0.0f, 1.0f);
+		return FLinearColor(
+			FMathf::Lerp(LowColor.R, HighColor.R, t),
+			FMathf::Lerp(LowColor.G, HighColor.G, t),
+			FMathf::Lerp(LowColor.B, HighColor.B, t),
+			1.0f);
+	}
+
+	template<typename VectorType>
+	VectorType ToColor(float Value)
+	{
+		FLinearColor Color = ToColor(Value);
+		return VectorType(Color.R, Color.G, Color.B);
+	}
+};
+
+
+/**
+ * Abstract interface to a single-channel indexed floating-point attribute
+ */
+class IMeshVertexAttributeAdapter
+{
+public:
+	virtual ~IMeshVertexAttributeAdapter() {}
+
+	virtual int ElementNum() const = 0;
+	virtual float GetValue(int32 Index) const = 0;
+	virtual void SetValue(int32 Index, float Value) = 0;
+	virtual UE::Geometry::FInterval1f GetValueRange() = 0;
+};
+
+
+
+/**
+ * Abstract interface to a set of single-channel indexed floating-point attributes
+ */
+class IMeshVertexAttributeSource
+{
+public:
+	virtual ~IMeshVertexAttributeSource() {}
+
+	virtual TArray<FName> GetAttributeList() = 0;
+	virtual TUniquePtr<IMeshVertexAttributeAdapter> GetAttribute(FName AttributeName) = 0;
+	/** @return number of indices in each attribute */
+	virtual int32 GetAttributeElementNum() = 0;
+};
+
+
+
+
+
+
+/**
+ * Tool Builder for Attribute Paint Tool
+ */
+UCLASS()
+class MESHMODELINGTOOLSEXP_API UMeshAttributePaintToolBuilder : public UMeshSurfacePointToolBuilder
+{
+	GENERATED_BODY()
+public:
+	/** Optional color map customization */
+	TUniqueFunction<TUniquePtr<FFloatAttributeColorMapper>()> ColorMapFactory;
+
+	virtual UMeshSurfacePointTool* CreateNewTool(const FToolBuilderState& SceneState) const override;
+};
+
+
+
+UENUM()
+enum class EBrushActionMode
+{
+	Paint,
+	FloodFill
+};
+
+
+/**
+ * Selected-Attribute settings Attribute Paint Tool
+ */
+UCLASS()
+class MESHMODELINGTOOLSEXP_API UMeshAttributePaintBrushOperationProperties : public UInteractiveToolPropertySet
+{
+	GENERATED_BODY()
+public:
+	UPROPERTY(EditAnywhere, Category = Attribute)
+	EBrushActionMode BrushAction = EBrushActionMode::Paint;
+};
+
+
+
+
+
+
+/**
+ * Selected-Attribute settings Attribute Paint Tool
+ */
+UCLASS()
+class MESHMODELINGTOOLSEXP_API UMeshAttributePaintToolProperties : public UInteractiveToolPropertySet
+{
+	GENERATED_BODY()
+public:
+	UPROPERTY(VisibleAnywhere, Category = Attribute)
+	TArray<FName> Attributes;
+
+	UPROPERTY(EditAnywhere, Category = Attribute, meta = (ArrayClamp = "Attributes"))
+	int SelectedAttribute = 0;
+
+	UPROPERTY(VisibleAnywhere, Category = Attribute)
+	FString AttributeName;
+};
+
+
+
+
+
+
+UENUM()
+enum class EMeshAttributePaintToolActions
+{
+	NoAction
+};
+
+
+
+UCLASS()
+class MESHMODELINGTOOLSEXP_API UMeshAttributePaintEditActions : public UInteractiveToolPropertySet
+{
+	GENERATED_BODY()
+
+public:
+	TWeakObjectPtr<UMeshAttributePaintTool> ParentTool;
+
+	void Initialize(UMeshAttributePaintTool* ParentToolIn) { ParentTool = ParentToolIn; }
+
+	void PostAction(EMeshAttributePaintToolActions Action);
+};
+
+
+
+/**
+ * FCommandChange for color map changes
+ */
+class MESHMODELINGTOOLSEXP_API FMeshAttributePaintChange : public TCustomIndexedValuesChange<float, int32>
+{
+public:
+	virtual FString ToString() const override
+	{
+		return FString(TEXT("Paint Attribute"));
+	}
+};
+
+
+
+/**
+ * UMeshAttributePaintTool paints single-channel float attributes on a MeshDescription.
+ * 
+ */
+UCLASS()
+class MESHMODELINGTOOLSEXP_API UMeshAttributePaintTool : public UDynamicMeshBrushTool
+{
+	GENERATED_BODY()
+
+public:
+	virtual void SetWorld(UWorld* World);
+
+	virtual void RegisterActions(FInteractiveToolActionSet& ActionSet) override;
+
+	virtual void Setup() override;
+	virtual void OnTick(float DeltaTime) override;
+
+	virtual bool HasCancel() const override { return true; }
+	virtual bool HasAccept() const override { return true; }
+
+	// UBaseBrushTool overrides
+	virtual void OnBeginDrag(const FRay& Ray) override;
+	virtual void OnUpdateDrag(const FRay& Ray) override;
+	virtual void OnEndDrag(const FRay& Ray) override;
+	virtual bool OnUpdateHover(const FInputDeviceRay& DevicePos) override;
+
+	virtual void RequestAction(EMeshAttributePaintToolActions ActionType);
+
+	virtual void SetColorMap(TUniquePtr<FFloatAttributeColorMapper> ColorMap);
+
+protected:
+	virtual void ApplyStamp(const FBrushStampData& Stamp);
+
+
+	struct FStampActionData
+	{
+		TArray<int32> ROIVertices;
+		TArray<float> ROIBefore;
+		TArray<float> ROIAfter;
+	};
+
+
+	virtual void ApplyStamp_Paint(const FBrushStampData& Stamp, FStampActionData& ActionData);
+	virtual void ApplyStamp_FloodFill(const FBrushStampData& Stamp, FStampActionData& ActionData);
+
+	virtual void OnShutdown(EToolShutdownType ShutdownType) override;
+	
+
+protected:
+	UPROPERTY()
+	TObjectPtr<UMeshAttributePaintBrushOperationProperties> BrushActionProps;
+
+	UPROPERTY()
+	TObjectPtr<UMeshAttributePaintToolProperties> AttribProps;
+
+	TValueWatcher<int32> SelectedAttributeWatcher;
+
+	//UPROPERTY()
+	//UMeshAttributePaintEditActions* AttributeEditActions;
+
+protected:
+	UWorld* TargetWorld;
+
+	bool bInRemoveStroke = false;
+	bool bInSmoothStroke = false;
+	FBrushStampData StartStamp;
+	FBrushStampData LastStamp;
+	bool bStampPending;
+
+	TUniquePtr<FMeshDescription> EditedMesh;
+
+	double CalculateBrushFalloff(double Distance);
+	UE::Geometry::TDynamicVerticesOctree3<FDynamicMesh3> VerticesOctree;
+	TArray<int> PreviewBrushROI;
+	void CalculateVertexROI(const FBrushStampData& Stamp, TArray<int>& VertexROI);
+
+	TUniquePtr<FFloatAttributeColorMapper> ColorMapper;
+	TUniquePtr<IMeshVertexAttributeSource> AttributeSource;
+
+	struct FAttributeData
+	{
+		FName Name;
+		TUniquePtr<IMeshVertexAttributeAdapter> Attribute;
+		TArray<float> CurrentValues;
+		TArray<float> InitialValues;
+	};
+	TArray<FAttributeData> Attributes;
+	int32 AttributeBufferCount;
+	int32 CurrentAttributeIndex;
+	UE::Geometry::FInterval1f CurrentValueRange;
+
+	// actions
+
+	bool bHavePendingAction = false;
+	EMeshAttributePaintToolActions PendingAction;
+	virtual void ApplyAction(EMeshAttributePaintToolActions ActionType);
+
+	bool bVisibleAttributeValid = false;
+	int32 PendingNewSelectedIndex = -1;
+	
+	void InitializeAttributes();
+	void StoreCurrentAttribute();
+	void UpdateVisibleAttribute();
+	void UpdateSelectedAttribute(int32 NewSelectedIndex);
+
+	TUniquePtr<TIndexedValuesChangeBuilder<float, FMeshAttributePaintChange>> ActiveChangeBuilder;
+	void BeginChange();
+	TUniquePtr<FMeshAttributePaintChange> EndChange();
+	void ExternalUpdateValues(int32 AttribIndex, const TArray<int32>& VertexIndices, const TArray<float>& NewValues);
+};
+
+
+
+
+
+
+
diff --git a/Engine/Plugins/Experimental/MeshModelingToolsetExp/Source/MeshModelingToolsExp/Public/MeshInspectorTool.h b/Engine/Plugins/Experimental/MeshModelingToolsetExp/Source/MeshModelingToolsExp/Public/MeshInspectorTool.h
new file mode 100644
--- /dev/null
+++ b/Engine/Plugins/Experimental/MeshModelingToolsetExp/Source/MeshModelingToolsExp/Public/MeshInspectorTool.h
@@ -0,0 +1,148 @@
+// Copyright Epic Games, Inc. All Rights Reserved.
+
+#pragma once
+
+#include "CoreMinimal.h"
+#include "UObject/NoExportTypes.h"
+#include "SingleSelectionTool.h"
+#include "InteractiveToolBuilder.h"
+#include "Drawing/LineSetComponent.h"
+#include "DynamicMesh/DynamicMesh3.h"
+#include "DynamicMesh/DynamicMeshAABBTree3.h"
+#include "Properties/MeshMaterialProperties.h"
+#include "PreviewMesh.h"
+#include "BaseTools/SingleSelectionMeshEditingTool.h"
+#include "MeshInspectorTool.generated.h"
+
+
+// predeclarations
+struct FMeshDescription;
+class UMaterialInstanceDynamic;
+
+/**
+ *
+ */
+UCLASS()
+class MESHMODELINGTOOLSEXP_API UMeshInspectorToolBuilder : public USingleSelectionMeshEditingToolBuilder
+{
+	GENERATED_BODY()
+
+public:
+	virtual USingleSelectionMeshEditingTool* CreateNewTool(const FToolBuilderState& SceneState) const override;
+};
+
+
+
+UCLASS()
+class MESHMODELINGTOOLSEXP_API UMeshInspectorProperties : public UInteractiveToolPropertySet
+{
+	GENERATED_BODY()
+
+public:
+	/** Toggle visibility of all mesh edges */
+	UPROPERTY(EditAnywhere, Category = Options)
+	bool bWireframe = true;
+
+	/** Toggle visibility of open boundary edges */
+	UPROPERTY(EditAnywhere, Category = Options)
+	bool bBoundaryEdges = true;
+
+	/** Toggle visibility of bowtie vertices */
+	UPROPERTY(EditAnywhere, Category = Options)
+	bool bBowtieVertices = true;
+
+	/** Toggle visibility of polygon borders */
+	UPROPERTY(EditAnywhere, Category = Options)
+	bool bPolygonBorders = false;
+
+	/** Toggle visibility of UV seam edges */
+	UPROPERTY(EditAnywhere, Category = Options)
+	bool bUVSeams = false;
+
+	/** Toggle visibility of UV bowtie vertices */
+	UPROPERTY(EditAnywhere, Category = Options)
+	bool bUVBowties = false;
+
+	/** Toggle visibility of Normal seam edges */
+	UPROPERTY(EditAnywhere, Category = Options)
+	bool bNormalSeams = false;
+
+	/** Toggle visibility of normal vectors */
+	UPROPERTY(EditAnywhere, Category = Options)
+	bool bNormalVectors = false;
+
+	/** Toggle visibility of tangent vectors */
+	UPROPERTY(EditAnywhere, Category = Options)
+	bool bTangentVectors = false;
+
+	/** Length of line segments representing normal vectors */
+	UPROPERTY(EditAnywhere, Category = Options, meta = (EditCondition = "bNormalVectors", 
+		UIMin="0", UIMax="400", ClampMin = "0", ClampMax = "1000000000.0"))
+	float NormalLength = 5.0f;
+
+	/** Length of line segments representing tangent vectors */
+	UPROPERTY(EditAnywhere, Category = Options, meta = (EditCondition = "bTangentVectors", 
+		UIMin = "0", UIMax = "400", ClampMin = "0", ClampMax = "1000000000.0"))
+	float TangentLength = 5.0f;
+};
+
+/**
+ * Mesh Inspector Tool for visualizing mesh information
+ */
+UCLASS()
+class MESHMODELINGTOOLSEXP_API UMeshInspectorTool : public USingleSelectionMeshEditingTool
+{
+	GENERATED_BODY()
+
+public:
+	UMeshInspectorTool();
+
+	virtual void RegisterActions(FInteractiveToolActionSet& ActionSet) override;
+
+	virtual void Setup() override;
+	virtual void Shutdown(EToolShutdownType ShutdownType) override;
+
+	virtual void Render(IToolsContextRenderAPI* RenderAPI) override;
+
+	virtual bool HasCancel() const override { return false; }
+	virtual bool HasAccept() const override;
+	virtual bool CanAccept() const override;
+
+	virtual void OnPropertyModified(UObject* PropertySet, FProperty* Property) override;
+
+public:
+
+	virtual void IncreaseLineWidthAction();
+	virtual void DecreaseLineWidthAction();
+
+protected:
+
+	UPROPERTY()
+	TObjectPtr<UMeshInspectorProperties> Settings = nullptr;
+
+	UPROPERTY()
+	TObjectPtr<UExistingMeshMaterialProperties> MaterialSettings = nullptr;
+
+
+	float LineWidthMultiplier = 1.0f;
+
+protected:
+	UPROPERTY()
+	TObjectPtr<UPreviewMesh> PreviewMesh;
+
+	UPROPERTY()
+	TObjectPtr<ULineSetComponent> DrawnLineSet;
+
+	UPROPERTY()
+	TObjectPtr<UMaterialInterface> DefaultMaterial = nullptr;
+
+	TArray<int> BoundaryEdges;
+	TArray<int> BoundaryBowties;
+	TArray<int> UVSeamEdges;
+	TArray<int> UVBowties;
+	TArray<int> NormalSeamEdges;
+	TArray<int> GroupBoundaryEdges;
+
+	void UpdateVisualization();
+	void Precompute();
+};
diff --git a/Engine/Plugins/Experimental/MeshModelingToolsetExp/Source/MeshModelingToolsExp/Public/MeshSelectionTool.h b/Engine/Plugins/Experimental/MeshModelingToolsetExp/Source/MeshModelingToolsExp/Public/MeshSelectionTool.h
new file mode 100644
--- /dev/null
+++ b/Engine/Plugins/Experimental/MeshModelingToolsetExp/Source/MeshModelingToolsExp/Public/MeshSelectionTool.h
@@ -0,0 +1,434 @@
+// Copyright Epic Games, Inc. All Rights Reserved.
+
+#pragma once
+
+#include "CoreMinimal.h"
+#include "DynamicMeshBrushTool.h"
+#include "InteractiveToolQueryInterfaces.h"
+#include "SelectionSet.h"
+#include "Changes/MeshSelectionChange.h"
+#include "DynamicMesh/DynamicMeshOctree3.h"
+#include "Polygroups/PolygroupSet.h"
+#include "MeshSelectionTool.generated.h"
+
+class UMeshStatisticsProperties;
+class UMeshElementsVisualizer;
+class UMeshUVChannelProperties;
+class UPolygroupLayersProperties;
+
+/**
+ *
+ */
+UCLASS()
+class MESHMODELINGTOOLSEXP_API UMeshSelectionToolBuilder : public UMeshSurfacePointToolBuilder
+{
+	GENERATED_BODY()
+
+public:
+	virtual UMeshSurfacePointTool* CreateNewTool(const FToolBuilderState& SceneState) const override;
+};
+
+
+
+
+UENUM()
+enum class EMeshSelectionToolActions
+{
+	NoAction,
+
+	SelectAll,
+	ClearSelection,
+	InvertSelection,
+	GrowSelection,
+	ShrinkSelection,
+	ExpandToConnected,
+
+	SelectLargestComponentByTriCount,
+	SelectLargestComponentByArea,
+	OptimizeSelection,
+
+	DeleteSelected,
+	DisconnectSelected,
+	SeparateSelected,
+	FlipSelected,
+	CreateGroup,
+
+	CycleSelectionMode,
+	CycleViewMode
+};
+
+
+
+UCLASS()
+class MESHMODELINGTOOLSEXP_API UMeshSelectionToolActionPropertySet : public UInteractiveToolPropertySet
+{
+	GENERATED_BODY()
+
+public:
+	TWeakObjectPtr<UMeshSelectionTool> ParentTool;
+
+	void Initialize(UMeshSelectionTool* ParentToolIn) { ParentTool = ParentToolIn; }
+
+	void PostAction(EMeshSelectionToolActions Action);
+};
+
+
+
+
+UCLASS()
+class MESHMODELINGTOOLSEXP_API UMeshSelectionEditActions : public UMeshSelectionToolActionPropertySet
+{
+	GENERATED_BODY()
+
+public:
+	/** Clear the active triangle selection */
+	UFUNCTION(CallInEditor, Category = SelectionEdits, meta = (DisplayPriority = 0))
+	void Clear()
+	{
+		PostAction(EMeshSelectionToolActions::ClearSelection);
+	}
+
+	/** Select all triangles in the mesh */
+	UFUNCTION(CallInEditor, Category = SelectionEdits, meta = (DisplayPriority = 1))
+	void SelectAll()
+	{
+		PostAction(EMeshSelectionToolActions::SelectAll);
+	}
+
+	/** Invert the active triangle selection */
+	UFUNCTION(CallInEditor, Category = SelectionEdits, meta = (DisplayPriority = 2))
+	void Invert()
+	{
+		PostAction(EMeshSelectionToolActions::InvertSelection);
+	}
+
+	/** Grow the active triangle selection to include any triangles touching a vertex on the selection boundary */
+	UFUNCTION(CallInEditor, Category = SelectionEdits, meta = (DisplayPriority = 3))
+	void Grow()
+	{
+		PostAction(EMeshSelectionToolActions::GrowSelection);
+	}
+
+	/** Shrink the active triangle selection by removing any triangles touching a vertex on the selection boundary */
+	UFUNCTION(CallInEditor, Category = SelectionEdits, meta = ( DisplayPriority = 4))
+	void Shrink()
+	{
+		PostAction(EMeshSelectionToolActions::ShrinkSelection);
+	}
+
+	/** Grow the active selection to include any triangle connected via shared edges (ie flood-fill) */
+	UFUNCTION(CallInEditor, Category = SelectionEdits, meta = (DisplayPriority = 5))
+	void FloodFill()
+	{
+		PostAction(EMeshSelectionToolActions::ExpandToConnected);
+	}
+
+	/** Select the largest connected mesh component by triangle count */
+	UFUNCTION(CallInEditor, Category = SelectionEdits, meta = (DisplayPriority = 6))
+	void LargestAreaPart()
+	{
+		PostAction(EMeshSelectionToolActions::SelectLargestComponentByTriCount);
+	}
+
+	/** Select the largest connected mesh component by mesh area */
+	UFUNCTION(CallInEditor, Category = SelectionEdits, meta = (DisplayPriority = 7))
+	void LargestTriCountPart()
+	{
+		PostAction(EMeshSelectionToolActions::SelectLargestComponentByArea);
+	}
+
+	/** Optimize the selection border */
+	UFUNCTION(CallInEditor, Category = SelectionEdits, meta = (DisplayPriority = 8))
+	void OptimizeBorder()
+	{
+		PostAction(EMeshSelectionToolActions::OptimizeSelection);
+	}
+
+};
+
+
+
+
+UCLASS()
+class MESHMODELINGTOOLSEXP_API UMeshSelectionMeshEditActions : public UMeshSelectionToolActionPropertySet
+{
+	GENERATED_BODY()
+
+public:
+	/** Delete the selected triangles */
+	UFUNCTION(CallInEditor, Category = MeshEdits, meta = (DisplayName = "Delete", DisplayPriority = 1))
+	void Delete()
+	{
+		PostAction(EMeshSelectionToolActions::DeleteSelected);
+	}
+
+	/** Disconnected the selected triangles from their neighbours, to create mesh boundaries along the selection borders */
+	UFUNCTION(CallInEditor, Category = MeshEdits, meta = (DisplayName = "Disconnect", DisplayPriority = 3))
+	void Disconnect() 
+	{
+		PostAction(EMeshSelectionToolActions::DisconnectSelected);
+	}
+
+	/** Flip the normals of the selected triangles. This will create hard normals at selection borders. */
+	UFUNCTION(CallInEditor, Category = MeshEdits, meta = (DisplayName = "Flip Normals", DisplayPriority = 4))
+	void FlipNormals() 
+	{
+		PostAction(EMeshSelectionToolActions::FlipSelected);
+	}
+
+	/** Assign a new unique Polygroup index to the selected triangles */
+	UFUNCTION(CallInEditor, Category = MeshEdits, meta = (DisplayName = "Create Polygroup", DisplayPriority = 5))
+	void CreatePolygroup()
+	{
+		PostAction(EMeshSelectionToolActions::CreateGroup);
+	}
+
+	/** Delete the selected triangles from the active Mesh Object and create a new Mesh containing those triangles */
+	UFUNCTION(CallInEditor, Category = MeshEdits, meta = (DisplayName = "Separate", DisplayPriority = 10))
+	void Separate() 
+	{
+		PostAction(EMeshSelectionToolActions::SeparateSelected);
+	}
+
+};
+
+
+
+
+
+
+UENUM()
+enum class EMeshSelectionToolPrimaryMode
+{
+	/** Select all triangles inside the brush area */
+	Brush,
+
+	/** Select all triangles inside the brush volume */
+	VolumetricBrush,
+
+	/** Select all triangles inside brush with normal within angular tolerance of hit triangle */
+	AngleFiltered,
+
+	/** Select all triangles inside brush that are visible from current view */
+	Visible,
+
+	/** Select all triangles connected to any triangle inside the brush */
+	AllConnected,
+
+	/** Select all triangles in groups connected to any triangle inside the brush */
+	AllInGroup,
+
+	/** Select all triangles with same material as hit triangle */
+	ByMaterial,
+
+	/** Select all triangles in same UV island as hit triangle */
+	ByUVIsland,
+
+	/** Select all triangles with normal within angular tolerance of hit triangle */
+	AllWithinAngle
+};
+
+
+
+UENUM()
+enum class EMeshFacesColorMode
+{
+	/** Display original mesh materials */
+	None,
+	/** Color mesh triangles by PolyGroup Color */
+	ByGroup,
+	/** Color mesh triangles by Material ID */
+	ByMaterialID,
+	/** Color mesh triangles by UV Island */
+	ByUVIsland
+};
+
+
+UCLASS()
+class MESHMODELINGTOOLSEXP_API UMeshSelectionToolProperties : public UInteractiveToolPropertySet
+{
+	GENERATED_BODY()
+public:
+
+	/** The Selection Mode defines the behavior of the selection brush */
+	UPROPERTY(EditAnywhere, Category = Selection)
+	EMeshSelectionToolPrimaryMode SelectionMode = EMeshSelectionToolPrimaryMode::Brush;
+
+	/** Angle in Degrees used for Angle-based Selection Modes */
+	UPROPERTY(EditAnywhere, Category = Selection, meta = (UIMin = "0.0", UIMax = "90.0") )
+	float AngleTolerance = 1.0;
+
+	/** Allow the brush to hit back-facing parts of the surface  */
+	UPROPERTY(EditAnywhere, Category = Selection)
+	bool bHitBackFaces = true;
+
+	/** Toggle drawing of highlight points on/off */
+	UPROPERTY(EditAnywhere, Category = Selection)
+	bool bShowPoints = false;
+
+	/** Color each triangle based on the selected mesh attribute */
+	UPROPERTY(EditAnywhere, Category = Selection)
+	EMeshFacesColorMode FaceColorMode = EMeshFacesColorMode::None;
+};
+
+
+
+/**
+ *
+ */
+UCLASS()
+class MESHMODELINGTOOLSEXP_API UMeshSelectionTool : public UDynamicMeshBrushTool, public IInteractiveToolNestedAcceptCancelAPI
+{
+	GENERATED_BODY()
+
+public:
+	UMeshSelectionTool();
+
+	virtual void SetWorld(UWorld* World);
+
+	virtual void RegisterActions(FInteractiveToolActionSet& ActionSet) override;
+
+	virtual void Setup() override;
+
+	virtual void OnTick(float DeltaTime) override;
+	virtual void Render(IToolsContextRenderAPI* RenderAPI) override;
+
+	virtual bool HasCancel() const override { return true; }
+	virtual bool HasAccept() const override { return true; }
+	virtual bool CanAccept() const override { return Super::CanAccept() && bHaveModifiedMesh; }
+
+	// UBaseBrushTool overrides
+	virtual bool HitTest(const FRay& Ray, FHitResult& OutHit) override;
+	virtual void OnBeginDrag(const FRay& Ray) override;
+	virtual void OnUpdateDrag(const FRay& Ray) override;
+	virtual void OnEndDrag(const FRay& Ray) override;
+	virtual bool OnUpdateHover(const FInputDeviceRay& DevicePos) override;
+
+	// IInteractiveToolCameraFocusAPI implementation
+	virtual FBox GetWorldSpaceFocusBox() override;
+
+	// IInteractiveToolCancelAPI
+	virtual bool SupportsNestedCancelCommand() override;
+	virtual bool CanCurrentlyNestedCancel() override;
+	virtual bool ExecuteNestedCancelCommand() override;
+
+public:
+
+	virtual void RequestAction(EMeshSelectionToolActions ActionType);
+
+	UPROPERTY()
+	TObjectPtr<UMeshSelectionToolProperties> SelectionProps;
+
+	UPROPERTY()
+	TObjectPtr<UMeshSelectionEditActions> SelectionActions;
+
+	UPROPERTY()
+	TObjectPtr<UMeshSelectionToolActionPropertySet> EditActions;
+
+	UPROPERTY()
+	TObjectPtr<UMeshStatisticsProperties> MeshStatisticsProperties;
+
+	UPROPERTY()
+	TObjectPtr<UMeshElementsVisualizer> MeshElementsDisplay;
+
+	UPROPERTY()
+	TObjectPtr<UMeshUVChannelProperties> UVChannelProperties = nullptr;
+
+	UPROPERTY()
+	TObjectPtr<UPolygroupLayersProperties> PolygroupLayerProperties = nullptr;
+
+protected:
+	virtual UMeshSelectionToolActionPropertySet* CreateEditActions();
+	virtual void AddSubclassPropertySets() {}
+
+protected:
+
+	virtual void ApplyStamp(const FBrushStampData& Stamp);
+
+	virtual void OnShutdown(EToolShutdownType ShutdownType) override;
+
+protected:
+	UPROPERTY()
+	TObjectPtr<UMeshSelectionSet> Selection;
+
+	UPROPERTY()
+	TArray<TObjectPtr<AActor>> SpawnedActors;
+
+	UWorld* TargetWorld;
+
+	// note: ideally this octree would be part of PreviewMesh!
+	TUniquePtr<UE::Geometry::FDynamicMeshOctree3> Octree;
+	bool bOctreeValid = false;
+	TUniquePtr<UE::Geometry::FDynamicMeshOctree3>& GetOctree();
+
+	EMeshSelectionElementType SelectionType = EMeshSelectionElementType::Face;
+
+	bool bInRemoveStroke = false;
+	FBrushStampData StartStamp;
+	FBrushStampData LastStamp;
+	bool bStampPending;
+
+	void UpdateFaceSelection(const FBrushStampData& Stamp, const TArray<int>& BrushROI);
+
+
+	// temp
+	TArray<int> IndexBuf;
+	TArray<int32> TemporaryBuffer;
+	TSet<int32> TemporarySet;
+	void CalculateVertexROI(const FBrushStampData& Stamp, TArray<int>& VertexROI);
+	void CalculateTriangleROI(const FBrushStampData& Stamp, TArray<int>& TriangleROI);
+	TArray<int> PreviewBrushROI;
+	TBitArray<> SelectedVertices;
+	TBitArray<> SelectedTriangles;
+	void OnExternalSelectionChange();
+
+	void OnRegionHighlightUpdated(const TArray<int32>& Triangles);
+	void OnRegionHighlightUpdated(const TSet<int32>& Triangles);
+	void OnRegionHighlightUpdated();
+	void UpdateVisualization(bool bSelectionModified);
+	bool bFullMeshInvalidationPending = false;
+	bool bColorsUpdatePending = false;
+	FColor GetCurrentFaceColor(const FDynamicMesh3* Mesh, int TriangleID);
+	void CacheUVIslandIDs();
+	TArray<int> TriangleToUVIsland;
+
+	// selection change
+	FMeshSelectionChangeBuilder* ActiveSelectionChange = nullptr;
+	void BeginChange(bool bAdding);
+	TUniquePtr<FToolCommandChange> EndChange();
+	void CancelChange();
+
+
+	// actions
+
+	bool bHavePendingAction = false;
+	EMeshSelectionToolActions PendingAction;
+	virtual void ApplyAction(EMeshSelectionToolActions ActionType);
+
+	void SelectAll();
+	void ClearSelection();
+	void InvertSelection();
+	void GrowShrinkSelection(bool bGrow);
+	void ExpandToConnected();
+
+	void SelectLargestComponent(bool bWeightByArea);
+	void OptimizeSelection();
+
+	void DeleteSelectedTriangles();
+	void DisconnectSelectedTriangles(); // disconnects edges between selected and not-selected triangles; keeps all triangles in the same mesh
+	void SeparateSelectedTriangles(); // separates out selected triangles to a new mesh, removing them from the current mesh
+	void FlipSelectedTriangles();
+	void AssignNewGroupToSelectedTriangles();
+
+	TSharedPtr<UE::Geometry::FPolygroupSet, ESPMode::ThreadSafe> ActiveGroupSet;
+	void OnSelectedGroupLayerChanged();
+	void UpdateActiveGroupLayer();
+
+	// if true, mesh has been edited
+	bool bHaveModifiedMesh = false;
+
+	virtual void ApplyShutdownAction(EToolShutdownType ShutdownType);
+};
+
+
+
diff --git a/Engine/Plugins/Experimental/MeshModelingToolsetExp/Source/MeshModelingToolsExp/Public/MeshVertexSculptTool.h b/Engine/Plugins/Experimental/MeshModelingToolsetExp/Source/MeshModelingToolsExp/Public/MeshVertexSculptTool.h
new file mode 100644
--- /dev/null
+++ b/Engine/Plugins/Experimental/MeshModelingToolsetExp/Source/MeshModelingToolsExp/Public/MeshVertexSculptTool.h
@@ -0,0 +1,313 @@
+// Copyright Epic Games, Inc. All Rights Reserved.
+
+#pragma once
+
+#include "CoreMinimal.h"
+#include "HAL/Platform.h"
+#include "UObject/NoExportTypes.h"
+#include "Components/DynamicMeshComponent.h"
+#include "DynamicMesh/DynamicMeshAABBTree3.h"
+#include "DynamicMesh/DynamicMeshOctree3.h"
+#include "DynamicMesh/MeshNormals.h"
+#include "BaseTools/BaseBrushTool.h"
+#include "TransformTypes.h"
+#include "Sculpting/MeshSculptToolBase.h"
+#include "Sculpting/MeshBrushOpBase.h"
+#include "Image/ImageBuilder.h"
+#include "Util/UniqueIndexSet.h"
+#include "Polygroups/PolygroupSet.h"
+#include "MeshVertexSculptTool.generated.h"
+
+class UTransformGizmo;
+class UTransformProxy;
+class UMaterialInstanceDynamic;
+class FMeshVertexChangeBuilder;
+class UPreviewMesh;
+
+
+
+
+/**
+ * Tool Builder
+ */
+UCLASS()
+class MESHMODELINGTOOLSEXP_API UMeshVertexSculptToolBuilder : public UMeshSurfacePointToolBuilder
+{
+	GENERATED_BODY()
+
+public:
+	virtual UMeshSurfacePointTool* CreateNewTool(const FToolBuilderState& SceneState) const override;
+};
+
+
+
+
+
+/** Mesh Sculpting Brush Types */
+UENUM()
+enum class EMeshVertexSculptBrushType : uint8
+{
+	/** Move vertices parallel to the view plane  */
+	Move UMETA(DisplayName = "Move"),
+
+	/** Grab Brush, fall-off alters the influence of the grab */
+	PullKelvin UMETA(DisplayName = "Kelvin Grab"),
+
+	/** Grab Brush that may generate cusps, fall-off alters the influence of the grab */
+	PullSharpKelvin UMETA(DisplayName = "Sharp Kelvin Grab"),
+
+	/** Smooth mesh vertices  */
+	Smooth UMETA(DisplayName = "Smooth"),
+
+	/** Smooth mesh vertices but only in direction of normal (Ctrl to invert) */
+	SmoothFill UMETA(DisplayName = "SmoothFill"),
+
+	/** Displace vertices along the average surface normal (Ctrl to invert) */
+	Offset UMETA(DisplayName = "Sculpt (Normal)"),
+
+	/** Displace vertices towards the camera viewpoint (Ctrl to invert) */
+	SculptView UMETA(DisplayName = "Sculpt (Viewpoint)"),
+
+	/** Displaces vertices along the average surface normal to a maximum height based on the brush size (Ctrl to invert) */
+	SculptMax UMETA(DisplayName = "Sculpt Max"),
+
+	/** Displace vertices along their vertex normals */
+	Inflate UMETA(DisplayName = "Inflate"),
+
+	/** Scale Brush will inflate or pinch radially from the center of the brush */
+	ScaleKelvin UMETA(DisplayName = "Kelvin Scale"),
+
+	/** Move vertices towards the center of the brush (Ctrl to push away)*/
+	Pinch UMETA(DisplayName = "Pinch"),
+
+	/** Twist Brush moves vertices in the plane perpendicular to the local mesh normal */
+	TwistKelvin UMETA(DisplayName = "Kelvin Twist"),
+
+	/** Move vertices towards the average plane of the brush stamp region */
+	Flatten UMETA(DisplayName = "Flatten"),
+
+	/** Move vertices towards a plane defined by the initial brush position  */
+	Plane UMETA(DisplayName = "Plane (Normal)"),
+
+	/** Move vertices towards a view-facing plane defined at the initial brush position */
+	PlaneViewAligned UMETA(DisplayName = "Plane (Viewpoint)"),
+
+	/** Move vertices towards a fixed plane in world space, positioned with a 3D gizmo */
+	FixedPlane UMETA(DisplayName = "FixedPlane"),
+
+	LastValue UMETA(Hidden)
+
+};
+
+
+
+/** Brush Triangle Filter Type */
+UENUM()
+enum class EMeshVertexSculptBrushFilterType : uint8
+{
+	/** Do not filter brush area */
+	None,
+	/** Only apply brush to triangles in the same connected mesh component/island */
+	Component,
+	/** Only apply brush to triangles with the same PolyGroup */
+	PolyGroup
+};
+
+
+
+UCLASS()
+class MESHMODELINGTOOLSEXP_API UVertexBrushSculptProperties : public UInteractiveToolPropertySet
+{
+	GENERATED_BODY()
+
+public:
+	/** Primary Brush Mode */
+	UPROPERTY(EditAnywhere, Category = Sculpting, meta = (DisplayName = "Brush Type"))
+	EMeshVertexSculptBrushType PrimaryBrushType = EMeshVertexSculptBrushType::Offset;
+
+	/** Primary Brush Falloff Type, multiplied by Alpha Mask where applicable */
+	UPROPERTY(EditAnywhere, Category = Sculpting, meta = (DisplayName = "Falloff Shape"))
+	EMeshSculptFalloffType PrimaryFalloffType = EMeshSculptFalloffType::Smooth;
+
+	/** Filter applied to Stamp Region Triangles, based on first Stroke Stamp */
+	UPROPERTY(EditAnywhere, Category = Sculpting, meta = (DisplayName = "Region Filter"))
+	EMeshVertexSculptBrushFilterType BrushFilter = EMeshVertexSculptBrushFilterType::None;
+
+	/** When Freeze Target is toggled on, the Brush Target Surface will be Frozen in its current state, until toggled off. Brush strokes will be applied relative to the Target Surface, for applicable Brushes */
+	UPROPERTY(EditAnywhere, Category = Sculpting, meta = (EditCondition = "PrimaryBrushType == EMeshVertexSculptBrushType::Offset || PrimaryBrushType == EMeshVertexSculptBrushType::SculptMax || PrimaryBrushType == EMeshVertexSculptBrushType::SculptView || PrimaryBrushType == EMeshVertexSculptBrushType::Pinch || PrimaryBrushType == EMeshVertexSculptBrushType::Resample" ))
+	bool bFreezeTarget = false;
+
+	/** When enabled, instead of Mesh Smoothing, the Shift-Smooth modifier will "erase" the displacement relative to the Brush Target Surface */
+	//UPROPERTY(EditAnywhere, Category = Sculpting, meta = (DisplayName = "Shift-Smooth Erases", EditCondition = "bFreezeTarget == true"))
+	UPROPERTY()
+	bool bSmoothErases = false;
+};
+
+
+
+/**
+ * Tool Properties for a brush alpha mask
+ */
+UCLASS()
+class MESHMODELINGTOOLSEXP_API UVertexBrushAlphaProperties : public UInteractiveToolPropertySet
+{
+	GENERATED_BODY()
+
+public:
+	/** Alpha mask applied to brush stamp. Red channel is used. */
+	UPROPERTY(EditAnywhere, Category = Alpha, meta = (DisplayName = "Alpha Mask"))
+	TObjectPtr<UTexture2D> Alpha = nullptr;
+
+	/** Alpha is rotated by this angle, inside the brush stamp frame (vertically aligned) */
+	UPROPERTY(EditAnywhere, Category = Alpha, meta = (DisplayName = "Rotation", UIMin = "-180.0", UIMax = "180.0", ClampMin = "-360.0", ClampMax = "360.0"))
+	float RotationAngle = 0.0;
+
+	/** If true, a random angle in +/- RandomRange is added to Rotation angle for each stamp */
+	UPROPERTY(EditAnywhere, Category = Alpha, AdvancedDisplay)
+	bool bRandomize = false;
+
+	/** Bounds of random generation (positive and negative) for randomized stamps */
+	UPROPERTY(EditAnywhere, Category = Alpha, AdvancedDisplay, meta = (UIMin = "0.0", UIMax = "180.0"))
+	float RandomRange = 180.0;
+};
+
+
+
+
+
+/**
+ * Mesh Vertex Sculpt Tool Class
+ */
+UCLASS()
+class MESHMODELINGTOOLSEXP_API UMeshVertexSculptTool : public UMeshSculptToolBase
+{
+	GENERATED_BODY()
+public:
+	virtual void Setup() override;
+	virtual void Shutdown(EToolShutdownType ShutdownType) override;
+
+	virtual void OnTick(float DeltaTime) override;
+
+	virtual bool HasCancel() const override { return true; }
+	virtual bool HasAccept() const override { return true; }
+
+	virtual bool OnUpdateHover(const FInputDeviceRay& DevicePos) override;
+
+	virtual void OnPropertyModified(UObject* PropertySet, FProperty* Property) override;
+
+public:
+
+	/** Properties that control sculpting*/
+	UPROPERTY()
+	TObjectPtr<UVertexBrushSculptProperties> SculptProperties;
+
+	UPROPERTY()
+	TObjectPtr<UVertexBrushAlphaProperties> AlphaProperties;
+
+	UPROPERTY()
+	TObjectPtr<UTexture2D> BrushAlpha;
+
+public:
+	virtual void IncreaseBrushSpeedAction() override;
+	virtual void DecreaseBrushSpeedAction() override;
+
+	virtual void UpdateBrushAlpha(UTexture2D* NewAlpha);
+
+protected:
+	// UMeshSculptToolBase API
+	virtual UBaseDynamicMeshComponent* GetSculptMeshComponent() { return DynamicMeshComponent; }
+	virtual FDynamicMesh3* GetBaseMesh() { return &BaseMesh; }
+	virtual const FDynamicMesh3* GetBaseMesh() const { return &BaseMesh; }
+
+	virtual int32 FindHitSculptMeshTriangle(const FRay3d& LocalRay) override;
+	virtual int32 FindHitTargetMeshTriangle(const FRay3d& LocalRay) override;
+	bool IsHitTriangleBackFacing(int32 TriangleID, const FDynamicMesh3* QueryMesh);
+
+	virtual void OnBeginStroke(const FRay& WorldRay) override;
+	virtual void OnEndStroke() override;
+	// end UMeshSculptToolBase API
+
+protected:
+	UPROPERTY()
+	TObjectPtr<UDynamicMeshComponent> DynamicMeshComponent;
+
+	// realtime visualization
+	void OnDynamicMeshComponentChanged(UDynamicMeshComponent* Component, const FMeshVertexChange* Change, bool bRevert);
+	FDelegateHandle OnDynamicMeshComponentChangedHandle;
+
+	void UpdateBrushType(EMeshVertexSculptBrushType BrushType);
+
+	TUniquePtr<UE::Geometry::FPolygroupSet> ActiveGroupSet;
+	TArray<int32> TriangleComponentIDs;
+
+	int32 InitialStrokeTriangleID = -1;
+
+	TSet<int32> AccumulatedTriangleROI;
+	bool bUndoUpdatePending = false;
+	TFuture<bool> UndoNormalsFuture;
+	TFuture<bool> UndoUpdateOctreeFuture;
+	TFuture<bool> UndoUpdateBaseMeshFuture;
+	TArray<int> NormalsBuffer;
+	void WaitForPendingUndoRedo();
+
+	TArray<uint32> OctreeUpdateTempBuffer;
+	TArray<bool> OctreeUpdateTempFlagBuffer;
+	TFuture<void> StampUpdateOctreeFuture;
+	bool bStampUpdatePending = false;
+	void WaitForPendingStampUpdate();
+
+	TArray<int> RangeQueryTriBuffer;
+	UE::Geometry::FUniqueIndexSet VertexROIBuilder;
+	UE::Geometry::FUniqueIndexSet TriangleROIBuilder;
+	TArray<UE::Geometry::FIndex3i> TriangleROIInBuf;
+	TArray<int> VertexROI;
+	TArray<int> TriangleROIArray;
+	void UpdateROI(const FVector3d& BrushPos);
+
+	UE::Geometry::FUniqueIndexSet NormalsROIBuilder;
+	TArray<std::atomic<bool>> NormalsFlags;		// set of per-vertex or per-element-id flags that indicate
+												// whether normal needs recompute. Fast to do it this way
+												// than to use a TSet or UniqueIndexSet...
+
+	bool bTargetDirty;
+
+	EMeshVertexSculptBrushType PendingStampType = EMeshVertexSculptBrushType::Smooth;
+
+	bool UpdateStampPosition(const FRay& WorldRay);
+	TFuture<void> ApplyStamp();
+
+	FRandomStream StampRandomStream;
+
+	FDynamicMesh3 BaseMesh;
+	UE::Geometry::FDynamicMeshOctree3 BaseMeshSpatial;
+	TArray<int32> BaseMeshIndexBuffer;
+	bool bCachedFreezeTarget = false;
+	void UpdateBaseMesh(const TSet<int32>* TriangleROI = nullptr);
+	bool GetBaseMeshNearest(int32 VertexID, const FVector3d& Position, double SearchRadius, FVector3d& TargetPosOut, FVector3d& TargetNormalOut);
+	TFunction<bool(int32, const FVector3d&, double MaxDist, FVector3d&, FVector3d&)> BaseMeshQueryFunc;
+
+	UE::Geometry::FDynamicMeshOctree3 Octree;
+
+	bool UpdateBrushPosition(const FRay& WorldRay);
+
+	double SculptMaxFixedHeight = -1.0;
+
+	bool bHaveBrushAlpha = false;
+	UE::Geometry::TImageBuilder<FVector4f> BrushAlphaValues;
+	UE::Geometry::FImageDimensions BrushAlphaDimensions;
+	double SampleBrushAlpha(const FSculptBrushStamp& Stamp, const FVector3d& Position) const;
+
+	TArray<FVector3d> ROIPositionBuffer;
+	TArray<FVector3d> ROIPrevPositionBuffer;
+
+	FMeshVertexChangeBuilder* ActiveVertexChange = nullptr;
+	void BeginChange();
+	void EndChange();
+
+
+protected:
+	virtual bool ShowWorkPlane() const override { return SculptProperties->PrimaryBrushType == EMeshVertexSculptBrushType::FixedPlane; }
+};
+
+
+
diff --git a/Engine/Plugins/Experimental/MeshModelingToolsetExp/Source/MeshModelingToolsExp/Public/MirrorTool.h b/Engine/Plugins/Experimental/MeshModelingToolsetExp/Source/MeshModelingToolsExp/Public/MirrorTool.h
new file mode 100644
--- /dev/null
+++ b/Engine/Plugins/Experimental/MeshModelingToolsetExp/Source/MeshModelingToolsExp/Public/MirrorTool.h
@@ -0,0 +1,246 @@
+// Copyright Epic Games, Inc. All Rights Reserved.
+
+#pragma once
+
+#include "CoreMinimal.h"
+
+#include "BaseGizmos/TransformGizmo.h"
+#include "Changes/DynamicMeshChangeTarget.h"
+#include "InteractiveToolBuilder.h"
+#include "Mechanics/ConstructionPlaneMechanic.h"
+#include "MeshOpPreviewHelpers.h"
+#include "MultiSelectionTool.h"
+#include "Selection/SelectClickedAction.h"
+#include "ToolContextInterfaces.h"
+
+#include "MirrorTool.generated.h"
+
+class UTransformGizmo;
+class UTransformProxy;
+
+UCLASS()
+class MESHMODELINGTOOLSEXP_API UMirrorToolBuilder : public UInteractiveToolBuilder
+{
+	GENERATED_BODY()
+public:
+	virtual bool CanBuildTool(const FToolBuilderState& SceneState) const override;
+	virtual UInteractiveTool* BuildTool(const FToolBuilderState& SceneState) const override;
+
+protected:
+	virtual const FToolTargetTypeRequirements& GetTargetRequirements() const override;
+};
+
+UENUM()
+enum class EMirrorSaveMode : uint8
+{
+	/**  Save the results in place of the original assets. */
+	UpdateAssets = 0,
+
+	/** Save the results as new assets. */
+	CreateNewAssets = 1,
+};
+
+UENUM()
+enum class EMirrorOperationMode : uint8
+{
+	/**  Append a mirrored version of the mesh to itself. */
+	MirrorAndAppend = 0,
+
+	/** Mirror the existing mesh. */
+	MirrorExisting = 1,
+};
+
+UENUM()
+enum class EMirrorCtrlClickBehavior : uint8
+{
+	/** Move the mirror plane to clicked location without adjusting its normal. */
+	Reposition = 0,
+
+	/** Move the mirror plane and adjust its normal according to click location. */
+	RepositionAndReorient = 1,
+};
+
+/**
+ * 
+ */
+UCLASS()
+class MESHMODELINGTOOLSEXP_API UMirrorToolProperties : public UInteractiveToolPropertySet
+{
+	GENERATED_BODY()
+
+public:
+
+	/** Mode of operation. */
+	UPROPERTY(EditAnywhere, Category = Options)
+	EMirrorOperationMode OperationMode = EMirrorOperationMode::MirrorAndAppend;
+
+	/** Cut off everything on the back side of the mirror plane before mirroring. */
+	UPROPERTY(EditAnywhere, Category = Options)
+	bool bCropAlongMirrorPlaneFirst = true;
+
+	/** Weld vertices that lie on the mirror plane. Vertices will not be welded if doing so would give an edge more than two faces, or if they are part of a face in the plane. */
+	UPROPERTY(EditAnywhere, Category = Options, meta = (
+		EditCondition = "OperationMode == EMirrorOperationMode::MirrorAndAppend", EditConditionHides))
+	bool bWeldVerticesOnMirrorPlane = true;
+
+	/** Distance (in unscaled mesh space) to allow a point to be from the plane and still consider it "on the mirror plane". */
+	UPROPERTY(EditAnywhere, Category = Options, meta = (
+		EditCondition = "OperationMode == EMirrorOperationMode::MirrorAndAppend && bWeldVerticesOnMirrorPlane", EditConditionHides,
+		UIMin = 0, UIMax = 0.01, ClampMin = 0, ClampMax = 10))
+	double PlaneTolerance = KINDA_SMALL_NUMBER;
+
+	/** When welding, whether to allow bowtie vertices to be created, or to duplicate the vertex. */
+	UPROPERTY(EditAnywhere, Category = Options, AdvancedDisplay, meta = (
+		EditCondition = "bWeldVerticesOnMirrorPlane && OperationMode == EMirrorOperationMode::MirrorAndAppend", EditConditionHides))
+	bool bAllowBowtieVertexCreation = false;
+
+	/** What Ctrl + clicking does to the mirror plane. */
+	UPROPERTY(EditAnywhere, Category = RepositionOptions)
+	EMirrorCtrlClickBehavior CtrlClickBehavior = EMirrorCtrlClickBehavior::Reposition;
+
+	/** If true the "Preset Mirror Directions" buttons only change the plane orientation, not location. */
+	UPROPERTY(EditAnywhere, Category = RepositionOptions, AdvancedDisplay)
+	bool bButtonsOnlyChangeOrientation = false;
+	
+	/** Whether to show the preview. */
+	UPROPERTY(EditAnywhere, Category = Options)
+	bool bShowPreview = true;
+
+	/** How to save the result. */
+	UPROPERTY(EditAnywhere, Category = ToolOutputOptions)
+	EMirrorSaveMode SaveMode = EMirrorSaveMode::UpdateAssets;
+};
+
+
+UCLASS()
+class MESHMODELINGTOOLSEXP_API UMirrorOperatorFactory : public UObject, public UE::Geometry::IDynamicMeshOperatorFactory
+{
+	GENERATED_BODY()
+
+public:
+	// IDynamicMeshOperatorFactory API
+	virtual TUniquePtr<UE::Geometry::FDynamicMeshOperator> MakeNewOperator() override;
+
+	UPROPERTY()
+	TObjectPtr<UMirrorTool> MirrorTool;
+
+	/** Index of the component within MirrorTool->ComponentTargets that this factory creates an operator for. */
+	int ComponentIndex;
+};
+
+UENUM()
+enum class EMirrorToolAction
+{
+	NoAction,
+
+	ShiftToCenter,
+
+	Left,
+	Right,
+	Up,
+	Down,
+	Forward,
+	Backward
+};
+
+UCLASS()
+class MESHMODELINGTOOLSEXP_API UMirrorToolActionPropertySet : public UInteractiveToolPropertySet
+{
+	GENERATED_BODY()
+
+public:
+	TWeakObjectPtr<UMirrorTool> ParentTool;
+
+	void Initialize(UMirrorTool* ParentToolIn) { ParentTool = ParentToolIn; }
+	void PostAction(EMirrorToolAction Action);
+
+	/** Move the mirror plane to center of bounding box without changing its normal. */
+	UFUNCTION(CallInEditor, Category = RepositionPlane)
+	void ShiftToCenter() { PostAction(EMirrorToolAction::ShiftToCenter); }
+
+	/** Move the mirror plane and adjust its normal to mirror entire selection leftward. */
+	UFUNCTION(CallInEditor, Category = PresetMirrorDirections, meta = (DisplayPriority = 1))
+	void Left() { PostAction(EMirrorToolAction::Left); }
+
+	/** Move the mirror plane and adjust its normal to mirror entire selection rightward. */
+	UFUNCTION(CallInEditor, Category = PresetMirrorDirections, meta = (DisplayPriority = 2))
+	void Right() { PostAction(EMirrorToolAction::Right); }
+
+	/** Move the mirror plane and adjust its normal to mirror entire selection upward. */
+	UFUNCTION(CallInEditor, Category = PresetMirrorDirections, meta = (DisplayPriority = 3))
+	void Up() { PostAction(EMirrorToolAction::Up); }
+
+	/** Move the mirror plane and adjust its normal to mirror entire selection downward. */
+	UFUNCTION(CallInEditor, Category = PresetMirrorDirections, meta = (DisplayPriority = 4))
+	void Down() { PostAction(EMirrorToolAction::Down); }
+
+	/** Move the mirror plane and adjust its normal to mirror entire selection forward. */
+	UFUNCTION(CallInEditor, Category = PresetMirrorDirections, meta = (DisplayPriority = 5))
+	void Forward() { PostAction(EMirrorToolAction::Forward); }
+
+	/** Move the mirror plane and adjust its normal to mirror entire selection backward. */
+	UFUNCTION(CallInEditor, Category = PresetMirrorDirections, meta = (DisplayPriority = 6))
+	void Backward() { PostAction(EMirrorToolAction::Backward); }
+};
+
+/** Tool for mirroring one or more meshes across a plane. */
+UCLASS()
+class MESHMODELINGTOOLSEXP_API UMirrorTool : public UMultiSelectionTool, public IModifierToggleBehaviorTarget
+{
+	GENERATED_BODY()
+public:
+
+	friend UMirrorOperatorFactory;
+
+	UMirrorTool();
+
+	virtual void Setup() override;
+	virtual void Shutdown(EToolShutdownType ShutdownType) override;
+
+	virtual void SetWorld(UWorld* World);
+
+	virtual void OnTick(float DeltaTime) override;
+	virtual void Render(IToolsContextRenderAPI* RenderAPI) override;
+
+	virtual bool HasCancel() const override { return true; }
+	virtual bool HasAccept() const override { return true; }
+	virtual bool CanAccept() const override;
+
+	void RequestAction(EMirrorToolAction ActionType);
+
+	virtual void OnPropertyModified(UObject* PropertySet, FProperty* Property) override;
+
+	// IClickSequenceBehaviorTarget implementation
+	virtual void OnUpdateModifierState(int ModifierID, bool bIsOn) override;
+protected:
+
+	UPROPERTY()
+	TObjectPtr<UMirrorToolProperties> Settings = nullptr;
+
+	UPROPERTY()
+	TObjectPtr<UMirrorToolActionPropertySet> ToolActions = nullptr;
+
+	UPROPERTY()
+	TArray<TObjectPtr<UDynamicMeshReplacementChangeTarget>> MeshesToMirror;
+
+	UPROPERTY()
+	TArray<TObjectPtr<UMeshOpPreviewWithBackgroundCompute>> Previews;
+
+	FVector3d MirrorPlaneOrigin = FVector3d::Zero();
+	FVector3d MirrorPlaneNormal = FVector3d::UnitZ();
+
+	UWorld* TargetWorld;
+
+	UPROPERTY()
+	TObjectPtr<UConstructionPlaneMechanic> PlaneMechanic;
+
+	EMirrorToolAction PendingAction;
+	FBox CombinedBounds;
+	void ApplyAction(EMirrorToolAction ActionType);
+
+	void SetupPreviews();
+	void GenerateAsset(const TArray<FDynamicMeshOpResult>& Results);
+
+private:
+	void CheckAndDisplayWarnings();
+};
\ No newline at end of file
diff --git a/Engine/Plugins/Experimental/MeshModelingToolsetExp/Source/MeshModelingToolsExp/Public/OffsetMeshTool.h b/Engine/Plugins/Experimental/MeshModelingToolsetExp/Source/MeshModelingToolsExp/Public/OffsetMeshTool.h
new file mode 100644
--- /dev/null
+++ b/Engine/Plugins/Experimental/MeshModelingToolsetExp/Source/MeshModelingToolsExp/Public/OffsetMeshTool.h
@@ -0,0 +1,154 @@
+// Copyright Epic Games, Inc. All Rights Reserved.
+
+#pragma once
+
+#include "CoreMinimal.h"
+#include "BaseTools/BaseMeshProcessingTool.h"
+#include "PropertySets/WeightMapSetProperties.h"
+#include "OffsetMeshTool.generated.h"
+
+
+UENUM()
+enum class EOffsetMeshToolOffsetType : uint8
+{
+	/** Iterative Offseting with N iterations */
+	Iterative UMETA(DisplayName = "Iterative"),
+
+	/** Implicit Offseting, produces smoother output and does a better job at preserving UVs, but can be very slow on large meshes */
+	Implicit UMETA(DisplayName = "Implicit")
+};
+
+
+
+/** Base properties of Offset */
+UCLASS()
+class MESHMODELINGTOOLSEXP_API UOffsetMeshToolProperties : public UInteractiveToolPropertySet
+{
+	GENERATED_BODY()
+
+public:
+	/** Type of Offseting to apply */
+	UPROPERTY(EditAnywhere, Category = Offset)
+	EOffsetMeshToolOffsetType OffsetType = EOffsetMeshToolOffsetType::Iterative;
+
+	/** Offset Distance in World Units */
+	UPROPERTY(EditAnywhere, Category = Offset, meta = (UIMin = "-100.0", UIMax = "100.0", ClampMin = "-10000.0", ClampMax = "100000.0"))
+	float Distance = 1.0f;
+
+	/** If true, create a thickened shell, instead of only moving the input vertices */
+	UPROPERTY(EditAnywhere, Category = Offset, meta = (EditCondition = "OffsetType == EOffsetMeshToolOffsetType::Iterative") )
+	bool bCreateShell = false;
+};
+
+
+
+UCLASS()
+class MESHMODELINGTOOLSEXP_API UOffsetWeightMapSetProperties : public UWeightMapSetProperties
+{
+	GENERATED_BODY()
+public:
+
+	/** Minimum Offset Distance in World Units, for Weight Map values of zero (clamped to Distance) */
+	UPROPERTY(EditAnywhere, Category = WeightMap, meta = (UIMin = "-100.0", UIMax = "100.0", ClampMin = "-10000.0", ClampMax = "100000.0", DisplayPriority = 5))
+	float MinDistance = 1.0f;
+};
+
+
+/** Properties for Iterative Offseting */
+UCLASS()
+class MESHMODELINGTOOLSEXP_API UIterativeOffsetProperties : public UInteractiveToolPropertySet
+{
+	GENERATED_BODY()
+public:
+
+	/** Number of Offseting iterations */
+	UPROPERTY(EditAnywhere, Category = IterativeOffsetOptions, meta = (UIMin = "1", UIMax = "100", ClampMin = "0", ClampMax = "1000"))
+	int32 Steps = 10;
+
+	/** Control whether the boundary is allowed to move */
+	UPROPERTY(EditAnywhere, Category = IterativeOffsetOptions)
+	bool bOffsetBoundaries = true;
+
+	/** Amount of smoothing applied per Offset step */
+	UPROPERTY(EditAnywhere, Category = IterativeOffsetOptions, meta = (UIMin = "0.0", UIMax = "1.0", ClampMin = "0.0", ClampMax = "1.0"))
+	float SmoothingPerStep = 0.0f;
+
+	/** Reproject smooth vertices onto non-smoothed Offset Surface at each step (expensive but better-preserves uniform distance) */
+	UPROPERTY(EditAnywhere, Category = IterativeOffsetOptions, meta = (UIMin = "0.0", UIMax = "1.0", ClampMin = "0.0", ClampMax = "1.0", EditCondition="SmoothingPerStep > 0"))
+	bool bReprojectSmooth = false;
+};
+
+
+/** Properties for Implicit Offseting */
+UCLASS()
+class MESHMODELINGTOOLSEXP_API UImplicitOffsetProperties : public UInteractiveToolPropertySet
+{
+	GENERATED_BODY()
+
+public:
+	/** How tightly we should constrain the constrained implicit offset to the explicit offset */
+	UPROPERTY(EditAnywhere, Category = ImplicitOffsetingOptions, meta = (UIMin = "0.0", UIMax = "1.0", ClampMin = "0.0", ClampMax = "100.0"))
+	float Smoothness = 0.2f;
+
+	/** If this is false, triangles will be reshaped to be more regular, which will distort UVs */
+	UPROPERTY(EditAnywhere, Category = ImplicitOffsetingOptions)
+	bool bPreserveUVs = true;
+};
+
+
+
+
+
+
+/**
+ * Mesh Offseting Tool
+ */
+UCLASS()
+class MESHMODELINGTOOLSEXP_API UOffsetMeshTool : public UBaseMeshProcessingTool
+{
+	GENERATED_BODY()
+
+public:
+	UOffsetMeshTool();
+
+	virtual void InitializeProperties() override;
+	virtual void OnShutdown(EToolShutdownType ShutdownType) override;
+
+	virtual TUniquePtr<UE::Geometry::FDynamicMeshOperator> MakeNewOperator() override;
+
+	virtual bool RequiresInitialVtxNormals() const { return true; }
+	virtual bool RequiresInitialBoundaryLoops() const { return true; }
+	virtual bool HasMeshTopologyChanged() const override;
+
+	virtual FText GetToolMessageString() const override;
+	virtual FText GetAcceptTransactionName() const override;
+
+protected:
+	UPROPERTY()
+	TObjectPtr<UOffsetMeshToolProperties> OffsetProperties = nullptr;
+
+	UPROPERTY()
+	TObjectPtr<UIterativeOffsetProperties> IterativeProperties = nullptr;
+
+	UPROPERTY()
+	TObjectPtr<UImplicitOffsetProperties> ImplicitProperties = nullptr;
+
+	UPROPERTY()
+	TObjectPtr<UOffsetWeightMapSetProperties> WeightMapProperties = nullptr;
+};
+
+
+
+
+/**
+ *
+ */
+UCLASS()
+class MESHMODELINGTOOLSEXP_API UOffsetMeshToolBuilder : public UBaseMeshProcessingToolBuilder
+{
+	GENERATED_BODY()
+public:
+	virtual UBaseMeshProcessingTool* MakeNewToolInstance(UObject* Outer) const {
+		return NewObject<UOffsetMeshTool>(Outer);
+	}
+};
diff --git a/Engine/Plugins/Experimental/MeshModelingToolsetExp/Source/MeshModelingToolsExp/Public/Physics/CollisionGeometryVisualization.h b/Engine/Plugins/Experimental/MeshModelingToolsetExp/Source/MeshModelingToolsExp/Public/Physics/CollisionGeometryVisualization.h
new file mode 100644
--- /dev/null
+++ b/Engine/Plugins/Experimental/MeshModelingToolsetExp/Source/MeshModelingToolsExp/Public/Physics/CollisionGeometryVisualization.h
@@ -0,0 +1,21 @@
+// Copyright Epic Games, Inc. All Rights Reserved.
+
+#pragma once
+
+#include "Physics/PhysicsDataCollection.h"
+#include "Drawing/PreviewGeometryActor.h"
+
+namespace UE
+{
+	namespace PhysicsTools
+	{
+		/**
+		 * Create line sets in a UPreviewGeometry for all the elements in a Physics Data Collection.
+		 * Spheres and Capsules are drawn as 3-axis wireframes. Convexes are added as wireframes.
+		 */
+		void MESHMODELINGTOOLSEXP_API InitializePreviewGeometryLines(const FPhysicsDataCollection& PhysicsData, UPreviewGeometry* PreviewGeom,
+			const FColor& LineColor, float LineThickness, float DepthBias = 0.0, int32 CircleStepResolution = 16 );
+	}
+}
+
+
diff --git a/Engine/Plugins/Experimental/MeshModelingToolsetExp/Source/MeshModelingToolsExp/Public/Physics/ExtractCollisionGeometryTool.h b/Engine/Plugins/Experimental/MeshModelingToolsetExp/Source/MeshModelingToolsExp/Public/Physics/ExtractCollisionGeometryTool.h
new file mode 100644
--- /dev/null
+++ b/Engine/Plugins/Experimental/MeshModelingToolsetExp/Source/MeshModelingToolsExp/Public/Physics/ExtractCollisionGeometryTool.h
@@ -0,0 +1,116 @@
+// Copyright Epic Games, Inc. All Rights Reserved.
+
+#pragma once
+
+#include "CoreMinimal.h"
+#include "BaseTools/SingleSelectionMeshEditingTool.h"
+#include "DynamicMesh/DynamicMesh3.h"
+#include "Physics/CollisionPropertySets.h"
+#include "PropertySets/CreateMeshObjectTypeProperties.h"
+#include "ExtractCollisionGeometryTool.generated.h"
+
+class UPreviewGeometry;
+class UPreviewMesh;
+
+UCLASS()
+class MESHMODELINGTOOLSEXP_API UExtractCollisionGeometryToolBuilder : public USingleSelectionMeshEditingToolBuilder
+{
+	GENERATED_BODY()
+
+public:
+	virtual USingleSelectionMeshEditingTool* CreateNewTool(const FToolBuilderState& SceneState) const override;
+
+protected:
+	virtual const FToolTargetTypeRequirements& GetTargetRequirements() const override;
+};
+
+
+UENUM()
+enum class EExtractCollisionOutputType : uint8
+{
+	/** Simple Collision shapes (Box, Sphere, Capsule, Convex) */
+	Simple = 0,
+	/** Complex Collision Mesh */
+	Complex = 1
+};
+
+
+UCLASS()
+class MESHMODELINGTOOLSEXP_API UExtractCollisionToolProperties : public UInteractiveToolPropertySet
+{
+	GENERATED_BODY()
+public:
+	/** Type of collision geometry to convert to Mesh */
+	UPROPERTY(EditAnywhere, Category = Options)
+	EExtractCollisionOutputType CollisionType = EExtractCollisionOutputType::Simple;
+
+	/** Whether or not to weld coincident border edges of the Complex Collision Mesh (if possible) */
+	UPROPERTY(EditAnywhere, Category = Options, Meta = (EditCondition = "CollisionType == EExtractCollisionOutputType::Complex"))
+	bool bWeldEdges = true;
+
+	/** Whether or not to generate a seperate Mesh Object for each Simple Collision Shape  */
+	UPROPERTY(EditAnywhere, Category = Options, Meta = (EditCondition = "CollisionType == EExtractCollisionOutputType::Simple"))
+	bool bOutputSeparateMeshes = true;
+
+	/** Show/Hide a preview of the generated mesh (overlaps source mesh) */
+	UPROPERTY(EditAnywhere, Category = Options)
+	bool bShowPreview = false;
+
+	/** Show/Hide input mesh */
+	UPROPERTY(EditAnywhere, Category = Options)
+	bool bShowInputMesh = true;
+};
+
+
+
+/**
+ * Mesh Inspector Tool for visualizing mesh information
+ */
+UCLASS()
+class MESHMODELINGTOOLSEXP_API UExtractCollisionGeometryTool : public USingleSelectionMeshEditingTool
+{
+	GENERATED_BODY()
+public:
+	virtual void Setup() override;
+	virtual void Shutdown(EToolShutdownType ShutdownType) override;
+
+	virtual void OnTick(float DeltaTime) override;
+
+	virtual bool HasCancel() const override { return true; }
+	virtual bool HasAccept() const override { return true; }
+	virtual bool CanAccept() const override;
+
+protected:
+
+	/** Property set for type of output object (StaticMesh, Volume, etc) */
+	UPROPERTY()
+	TObjectPtr<UCreateMeshObjectTypeProperties> OutputTypeProperties;
+
+	UPROPERTY()
+	TObjectPtr<UExtractCollisionToolProperties> Settings;
+
+	UPROPERTY()
+	TObjectPtr<UCollisionGeometryVisualizationProperties> VizSettings = nullptr;
+
+	UPROPERTY()
+	TObjectPtr<UPhysicsObjectToolPropertySet> ObjectProps;
+
+protected:
+	UPROPERTY()
+	TObjectPtr<UPreviewGeometry> PreviewElements;
+
+	UPROPERTY()
+	TObjectPtr<UPreviewMesh> PreviewMesh;
+
+	// these are TSharedPtr because TPimplPtr cannot currently be added to a TArray?
+	TSharedPtr<FPhysicsDataCollection> PhysicsInfo;
+
+	TArray<UE::Geometry::FDynamicMesh3> CurrentMeshParts;
+	UE::Geometry::FDynamicMesh3 CurrentMesh;
+	bool bResultValid = false;
+	void RecalculateMesh_Simple();
+	void RecalculateMesh_Complex();
+
+	bool bVisualizationDirty = false;
+	void UpdateVisualization();
+};
diff --git a/Engine/Plugins/Experimental/MeshModelingToolsetExp/Source/MeshModelingToolsExp/Public/Physics/PhysicsInspectorTool.h b/Engine/Plugins/Experimental/MeshModelingToolsetExp/Source/MeshModelingToolsExp/Public/Physics/PhysicsInspectorTool.h
new file mode 100644
--- /dev/null
+++ b/Engine/Plugins/Experimental/MeshModelingToolsetExp/Source/MeshModelingToolsExp/Public/Physics/PhysicsInspectorTool.h
@@ -0,0 +1,69 @@
+// Copyright Epic Games, Inc. All Rights Reserved.
+
+#pragma once
+
+#include "CoreMinimal.h"
+#include "MultiSelectionTool.h"
+#include "InteractiveToolBuilder.h"
+#include "Physics/CollisionPropertySets.h"
+#include "PhysicsInspectorTool.generated.h"
+
+class UPreviewGeometry;
+
+UCLASS()
+class MESHMODELINGTOOLSEXP_API UPhysicsInspectorToolBuilder : public UInteractiveToolBuilder
+{
+	GENERATED_BODY()
+
+public:
+	virtual bool CanBuildTool(const FToolBuilderState& SceneState) const override;
+	virtual UInteractiveTool* BuildTool(const FToolBuilderState& SceneState) const override;
+
+protected:
+	virtual const FToolTargetTypeRequirements& GetTargetRequirements() const override;
+};
+
+
+
+
+/**
+ * Mesh Inspector Tool for visualizing mesh information
+ */
+UCLASS()
+class MESHMODELINGTOOLSEXP_API UPhysicsInspectorTool : public UMultiSelectionTool
+{
+	GENERATED_BODY()
+public:
+	virtual void Setup() override;
+	virtual void Shutdown(EToolShutdownType ShutdownType) override;
+
+	virtual void OnTick(float DeltaTime) override;
+
+	virtual bool HasCancel() const override { return false; }
+	virtual bool HasAccept() const override { return false; }
+	virtual bool CanAccept() const override { return false; }
+
+protected:
+
+	UPROPERTY()
+	TObjectPtr<UCollisionGeometryVisualizationProperties> VizSettings = nullptr;
+
+	UPROPERTY()
+	TArray<TObjectPtr<UPhysicsObjectToolPropertySet>> ObjectData;
+
+	UPROPERTY()
+	TObjectPtr<UMaterialInterface> LineMaterial = nullptr;
+
+protected:
+	UPROPERTY()
+	TArray<TObjectPtr<UPreviewGeometry>> PreviewElements;
+
+	// these are TSharedPtr because TPimplPtr cannot currently be added to a TArray?
+	TArray<TSharedPtr<FPhysicsDataCollection>> PhysicsInfos;
+
+	void InitializeGeometry(const FPhysicsDataCollection& PhysicsData, UPreviewGeometry* PreviewGeom);
+	void InitializeObjectProperties(const FPhysicsDataCollection& PhysicsData, UPhysicsObjectToolPropertySet* PropSet);
+
+	bool bVisualizationDirty = false;
+	void UpdateVisualization();
+};
diff --git a/Engine/Plugins/Experimental/MeshModelingToolsetExp/Source/MeshModelingToolsExp/Public/Physics/SetCollisionGeometryTool.h b/Engine/Plugins/Experimental/MeshModelingToolsetExp/Source/MeshModelingToolsExp/Public/Physics/SetCollisionGeometryTool.h
new file mode 100644
--- /dev/null
+++ b/Engine/Plugins/Experimental/MeshModelingToolsetExp/Source/MeshModelingToolsExp/Public/Physics/SetCollisionGeometryTool.h
@@ -0,0 +1,238 @@
+// Copyright Epic Games, Inc. All Rights Reserved.
+
+#pragma once
+
+#include "CoreMinimal.h"
+#include "MultiSelectionTool.h"
+#include "InteractiveToolBuilder.h"
+#include "DynamicMesh/DynamicMesh3.h"
+#include "SphereTypes.h"
+#include "OrientedBoxTypes.h"
+#include "CapsuleTypes.h"
+#include "Physics/CollisionPropertySets.h"
+#include "PropertySets/PolygroupLayersProperties.h"
+#include "Polygroups/PolygroupSet.h"
+#include "SetCollisionGeometryTool.generated.h"
+
+class UPreviewGeometry;
+PREDECLARE_GEOMETRY(class FMeshSimpleShapeApproximation)
+PREDECLARE_USE_GEOMETRY_CLASS(FDynamicMesh3);
+
+UCLASS()
+class MESHMODELINGTOOLSEXP_API USetCollisionGeometryToolBuilder : public UInteractiveToolBuilder
+{
+	GENERATED_BODY()
+
+public:
+	virtual bool CanBuildTool(const FToolBuilderState& SceneState) const override;
+	virtual UInteractiveTool* BuildTool(const FToolBuilderState& SceneState) const override;
+
+protected:
+	virtual const FToolTargetTypeRequirements& GetTargetRequirements() const override;
+};
+
+
+
+
+UENUM()
+enum class ESetCollisionGeometryInputMode
+{
+	CombineAll = 0,
+	PerInputObject = 1,
+	PerMeshComponent = 2,
+	PerMeshGroup = 3
+};
+
+
+UENUM()
+enum class ECollisionGeometryType
+{
+	KeepExisting = 0,
+	AlignedBoxes = 1,
+	OrientedBoxes = 2,
+	MinimalSpheres = 3,
+	Capsules = 4,
+	ConvexHulls = 5,
+	SweptHulls = 6,
+	MinVolume = 10,
+
+	None = 11
+};
+
+
+
+UENUM()
+enum class EProjectedHullAxis
+{
+	X = 0,
+	Y = 1,
+	Z = 2,
+	SmallestBoxDimension = 3,
+	SmallestVolume = 4
+};
+
+
+UCLASS()
+class MESHMODELINGTOOLSEXP_API USetCollisionGeometryToolProperties : public UInteractiveToolPropertySet
+{
+	GENERATED_BODY()
+public:
+
+	UPROPERTY(EditAnywhere, Category = Options)
+	ECollisionGeometryType GeometryType = ECollisionGeometryType::AlignedBoxes;
+
+	UPROPERTY(EditAnywhere, Category = Options)
+	ESetCollisionGeometryInputMode InputMode = ESetCollisionGeometryInputMode::PerInputObject;
+
+	UPROPERTY(EditAnywhere, Category = Options)
+	bool bUseWorldSpace = false;
+
+	UPROPERTY(EditAnywhere, Category = Options)
+	bool bRemoveContained = true;
+
+	UPROPERTY(EditAnywhere, Category = Options)
+	bool bEnableMaxCount = true;
+
+	UPROPERTY(EditAnywhere, Category = Options, meta = (UIMin = "1", UIMax = "100", ClampMin = "1", ClampMax = "9999999", EditCondition = "bEnableMaxCount"))
+	int32 MaxCount = 50;
+
+	UPROPERTY(EditAnywhere, Category = Options, AdvancedDisplay)
+	float MinThickness = 0.01;
+
+	UPROPERTY(EditAnywhere, Category = AutoDetect)
+	bool bDetectBoxes = true;
+
+	UPROPERTY(EditAnywhere, Category = AutoDetect)
+	bool bDetectSpheres = true;
+
+	UPROPERTY(EditAnywhere, Category = AutoDetect)
+	bool bDetectCapsules = true;
+
+	UPROPERTY(EditAnywhere, Category = ConvexHulls, meta = (EditConditionHides, EditCondition = "GeometryType == ECollisionGeometryType::ConvexHulls"))
+	bool bSimplifyHulls = true;
+
+	UPROPERTY(EditAnywhere, Category = ConvexHulls, meta = (UIMin = "4", UIMax = "100", ClampMin = "4", ClampMax = "9999999",
+		EditConditionHides, EditCondition = "GeometryType == ECollisionGeometryType::ConvexHulls"))
+	int32 HullTargetFaceCount = 20;
+
+	UPROPERTY(EditAnywhere, Category = SweptHulls, meta = (EditConditionHides, EditCondition = "GeometryType == ECollisionGeometryType::SweptHulls"))
+	bool bSimplifyPolygons = true;
+
+	UPROPERTY(EditAnywhere, Category = SweptHulls, meta = (UIMin = "0", UIMax = "10", ClampMin = "0", ClampMax = "100000",
+		EditConditionHides, EditCondition = "GeometryType == ECollisionGeometryType::SweptHulls"))
+	float HullTolerance = 0.1;
+
+	UPROPERTY(EditAnywhere, Category = SweptHulls, meta = (UIMin = "0", UIMax = "10", ClampMin = "0", ClampMax = "100000",
+		EditConditionHides, EditCondition = "GeometryType == ECollisionGeometryType::SweptHulls"))
+	EProjectedHullAxis SweepAxis = EProjectedHullAxis::SmallestVolume;
+
+	UPROPERTY(EditAnywhere, Category = OutputOptions)
+	bool bAppendToExisting = false;
+
+	UPROPERTY(EditAnywhere, Category = OutputOptions)
+	ECollisionGeometryMode SetCollisionType = ECollisionGeometryMode::SimpleAndComplex;
+};
+
+
+
+
+
+/**
+ * Mesh Inspector Tool for visualizing mesh information
+ */
+UCLASS()
+class MESHMODELINGTOOLSEXP_API USetCollisionGeometryTool : public UMultiSelectionTool
+{
+	GENERATED_BODY()
+public:
+	virtual void Setup() override;
+	virtual void Shutdown(EToolShutdownType ShutdownType) override;
+
+	virtual void OnTick(float DeltaTime) override;
+
+	virtual bool HasCancel() const override { return true; }
+	virtual bool HasAccept() const override { return true; }
+
+protected:
+
+	UPROPERTY()
+	TObjectPtr<USetCollisionGeometryToolProperties> Settings = nullptr;
+
+	UPROPERTY()
+	TObjectPtr<UPolygroupLayersProperties> PolygroupLayerProperties = nullptr;
+
+	UPROPERTY()
+	TObjectPtr<UCollisionGeometryVisualizationProperties> VizSettings = nullptr;
+
+	UPROPERTY()
+	TObjectPtr<UPhysicsObjectToolPropertySet> CollisionProps;
+
+	UPROPERTY()
+	TObjectPtr<UMaterialInterface> LineMaterial = nullptr;
+
+protected:
+	UPROPERTY()
+	TObjectPtr<UPreviewGeometry> PreviewGeom;
+
+	TArray<int32> SourceObjectIndices;
+	bool bSourcesHidden = false;
+
+	TArray<FDynamicMesh3> InitialSourceMeshes;
+
+	void OnInputModeChanged();
+
+	enum class EDetectedCollisionGeometry
+	{
+		None,
+		Sphere = 2,
+		Box = 4,
+		Capsule = 8,
+		Convex = 16
+	};
+
+	struct FSourceMesh
+	{
+		FDynamicMesh3 Mesh;
+
+		EDetectedCollisionGeometry DetectedType = EDetectedCollisionGeometry::None;
+
+		UE::Geometry::FSphere3d DetectedSphere;
+		UE::Geometry::FOrientedBox3d DetectedBox;
+		UE::Geometry::FCapsule3d DetectedCapsule;
+	};
+	
+	bool bInputMeshesValid = false;
+	TArray<TSharedPtr<FDynamicMesh3, ESPMode::ThreadSafe>> InputMeshes;
+	TArray<TSharedPtr<FDynamicMesh3, ESPMode::ThreadSafe>> CombinedInputMeshes;
+	TArray<TSharedPtr<FDynamicMesh3, ESPMode::ThreadSafe>> SeparatedInputMeshes;
+	TArray<TSharedPtr<FDynamicMesh3, ESPMode::ThreadSafe>> PerGroupInputMeshes;
+
+	TSharedPtr<UE::Geometry::FMeshSimpleShapeApproximation, ESPMode::ThreadSafe> InputMeshesApproximator;
+	TSharedPtr<UE::Geometry::FMeshSimpleShapeApproximation, ESPMode::ThreadSafe> CombinedInputMeshesApproximator;
+	TSharedPtr<UE::Geometry::FMeshSimpleShapeApproximation, ESPMode::ThreadSafe> SeparatedMeshesApproximator;
+	TSharedPtr<UE::Geometry::FMeshSimpleShapeApproximation, ESPMode::ThreadSafe> PerGroupMeshesApproximator;
+
+	void PrecomputeInputMeshes();
+	void InitializeDerivedMeshSet(
+		const TArray<TSharedPtr<FDynamicMesh3, ESPMode::ThreadSafe>>& FromInputMeshes,
+		TArray<TSharedPtr<FDynamicMesh3, ESPMode::ThreadSafe>>& ToMeshes,
+		TFunctionRef<bool(const FDynamicMesh3*, int32, int32)> TrisConnectedPredicate);
+	TSharedPtr<UE::Geometry::FMeshSimpleShapeApproximation, ESPMode::ThreadSafe>& GetApproximator(ESetCollisionGeometryInputMode MeshSetMode);
+
+	TUniquePtr<UE::Geometry::FPolygroupSet> ActiveGroupSet;
+	void OnSelectedGroupLayerChanged();
+	void UpdateActiveGroupLayer();
+
+	FTransform OrigTargetTransform;
+	FVector TargetScale3D;
+
+	bool bResultValid = false;
+	TSharedPtr<FPhysicsDataCollection, ESPMode::ThreadSafe> InitialCollision;
+	TSharedPtr<FPhysicsDataCollection, ESPMode::ThreadSafe> GeneratedCollision;
+
+	void UpdateGeneratedCollision();
+//	TSharedPtr<FPhysicsDataCollection> GenerateCollision_MinVolume();
+
+	bool bVisualizationDirty = false;
+	void UpdateVisualization();
+};
diff --git a/Engine/Plugins/Experimental/MeshModelingToolsetExp/Source/MeshModelingToolsExp/Public/PlaneCutTool.h b/Engine/Plugins/Experimental/MeshModelingToolsetExp/Source/MeshModelingToolsExp/Public/PlaneCutTool.h
new file mode 100644
--- /dev/null
+++ b/Engine/Plugins/Experimental/MeshModelingToolsetExp/Source/MeshModelingToolsExp/Public/PlaneCutTool.h
@@ -0,0 +1,213 @@
+// Copyright Epic Games, Inc. All Rights Reserved.
+
+#pragma once
+
+#include "CoreMinimal.h"
+#include "InteractiveGizmo.h"
+#include "MultiSelectionTool.h"
+#include "InteractiveToolBuilder.h"
+#include "MeshOpPreviewHelpers.h"
+#include "DynamicMesh/DynamicMesh3.h"
+#include "Changes/DynamicMeshChangeTarget.h"
+#include "BaseTools/SingleClickTool.h"
+#include "Mechanics/ConstructionPlaneMechanic.h"
+#include "PlaneCutTool.generated.h"
+
+
+// predeclarations
+class UTransformGizmo;
+class UTransformProxy;
+
+
+/**
+ *
+ */
+UCLASS()
+class MESHMODELINGTOOLSEXP_API UPlaneCutToolBuilder : public UInteractiveToolBuilder
+{
+	GENERATED_BODY()
+
+public:
+	virtual bool CanBuildTool(const FToolBuilderState& SceneState) const override;
+	virtual UInteractiveTool* BuildTool(const FToolBuilderState& SceneState) const override;
+
+protected:
+	virtual const FToolTargetTypeRequirements& GetTargetRequirements() const override;
+};
+
+
+
+/**
+* Properties controlling how changes are baked out to static meshes on tool accept
+*/
+UCLASS()
+class MESHMODELINGTOOLSEXP_API UAcceptOutputProperties : public UInteractiveToolPropertySet
+{
+	GENERATED_BODY()
+
+public:
+
+	/** If true, meshes cut into multiple pieces will be saved as separate assets on 'accept'. */
+	UPROPERTY(EditAnywhere, Category = ToolOutputOptions)
+	bool bExportSeparatedPiecesAsNewMeshAssets = true;
+};
+
+
+
+
+
+
+/**
+ * Standard properties of the plane cut operation
+ */
+UCLASS()
+class MESHMODELINGTOOLSEXP_API UPlaneCutToolProperties : public UInteractiveToolPropertySet
+{
+	GENERATED_BODY()
+
+public:
+	UPlaneCutToolProperties()
+	{}
+
+	/** If true, both halves of the cut are computed */
+	UPROPERTY(EditAnywhere, Category = Options)
+	bool bKeepBothHalves = false;
+
+	/** If keeping both halves, separate the two pieces by this amount */
+	UPROPERTY(EditAnywhere, Category = Options, meta = (EditCondition = "bKeepBothHalves == true", UIMin = "0", ClampMin = "0") )
+	float SpacingBetweenHalves = 0;
+
+	UPROPERTY(EditAnywhere, Category = Options)
+	bool bShowPreview = true;
+
+	/** If true, the cut surface is filled with simple planar hole fill surface(s) */
+	UPROPERTY(EditAnywhere, Category = Options)
+	bool bFillCutHole = true;
+
+	/** If true, will attempt to fill cut holes even if they're ill-formed (e.g. because they connect to pre-existing holes in the geometry) */
+	UPROPERTY(EditAnywhere, Category = Options, AdvancedDisplay, meta = (EditCondition = "bFillCutHole"))
+	bool bFillSpans = false;
+};
+
+
+
+UENUM()
+enum class EPlaneCutToolActions
+{
+	NoAction,
+	Cut,
+	FlipPlane
+};
+
+
+
+
+UCLASS()
+class MESHMODELINGTOOLSEXP_API UPlaneCutOperatorFactory : public UObject, public UE::Geometry::IDynamicMeshOperatorFactory
+{
+	GENERATED_BODY()
+
+public:
+	// IDynamicMeshOperatorFactory API
+	virtual TUniquePtr<UE::Geometry::FDynamicMeshOperator> MakeNewOperator() override;
+
+	UPROPERTY()
+	TObjectPtr<UPlaneCutTool> CutTool;
+
+	int ComponentIndex;
+};
+
+/**
+ * Simple Mesh Plane Cutting Tool
+ */
+UCLASS()
+class MESHMODELINGTOOLSEXP_API UPlaneCutTool : public UMultiSelectionTool
+{
+	GENERATED_BODY()
+
+public:
+
+	friend UPlaneCutOperatorFactory;
+
+	UPlaneCutTool();
+
+	virtual void Setup() override;
+	virtual void Shutdown(EToolShutdownType ShutdownType) override;
+
+	virtual void SetWorld(UWorld* World);
+
+	virtual void RegisterActions(FInteractiveToolActionSet& ActionSet) override;
+
+	virtual void OnTick(float DeltaTime) override;
+	virtual void Render(IToolsContextRenderAPI* RenderAPI) override;
+
+	virtual bool HasCancel() const override { return true; }
+	virtual bool HasAccept() const override { return true; }
+	virtual bool CanAccept() const override;
+
+#if WITH_EDITOR
+	virtual void PostEditChangeProperty(FPropertyChangedEvent &PropertyChangedEvent) override;
+#endif
+
+	virtual void OnPropertyModified(UObject* PropertySet, FProperty* Property) override;
+
+protected:
+
+	UPROPERTY()
+	TObjectPtr<UPlaneCutToolProperties> BasicProperties;
+
+	UPROPERTY()
+	TObjectPtr<UAcceptOutputProperties> AcceptProperties;
+
+	UPROPERTY()
+	TArray<TObjectPtr<UMeshOpPreviewWithBackgroundCompute>> Previews;
+
+
+	/// Action buttons.
+	/// Note these set a flag to call the action later (in OnTick)
+	/// Otherwise, the actions in undo history will end up being generically named by an outer UI handler transaction
+
+	/** Cut with the current plane without exiting the tool (Hotkey: T)*/
+	UFUNCTION(CallInEditor, Category = Actions, meta = (DisplayName = "Cut"))
+	void Cut()
+	{
+		PendingAction = EPlaneCutToolActions::Cut;
+	}
+
+	/** Flip the cutting plane (Hotkey: R) */
+	UFUNCTION(CallInEditor, Category = Actions, meta = (DisplayName = "Flip Plane"))
+	void FlipPlane()
+	{
+		PendingAction = EPlaneCutToolActions::FlipPlane;
+	}
+
+protected:
+
+	UPROPERTY()
+	TArray<TObjectPtr<UDynamicMeshReplacementChangeTarget>> MeshesToCut;
+
+	UPROPERTY()
+	TObjectPtr<UConstructionPlaneMechanic> PlaneMechanic = nullptr;
+
+	// Cutting plane
+	UE::Geometry::FFrame3d CutPlaneWorld;
+
+	// UV Scale factor is cached based on the bounding box of the mesh before any cuts are performed, so you don't get inconsistent UVs if you multi-cut the object to smaller sizes
+	TArray<float> MeshUVScaleFactor;
+
+	UWorld* TargetWorld;
+
+	FViewCameraState CameraState;
+
+	EPlaneCutToolActions PendingAction = EPlaneCutToolActions::NoAction;
+
+
+	void DoCut();
+	void DoFlipPlane();
+
+	void SetupPreviews();
+
+	void InvalidatePreviews();
+
+	void GenerateAsset(const TArray<FDynamicMeshOpResult>& Results);
+};
diff --git a/Engine/Plugins/Experimental/MeshModelingToolsetExp/Source/MeshModelingToolsExp/Public/Properties/MeshAnalysisProperties.h b/Engine/Plugins/Experimental/MeshModelingToolsetExp/Source/MeshModelingToolsExp/Public/Properties/MeshAnalysisProperties.h
new file mode 100644
--- /dev/null
+++ b/Engine/Plugins/Experimental/MeshModelingToolsetExp/Source/MeshModelingToolsExp/Public/Properties/MeshAnalysisProperties.h
@@ -0,0 +1,30 @@
+// Copyright Epic Games, Inc. All Rights Reserved.
+
+#pragma once
+
+#include "CoreMinimal.h"
+#include "InteractiveToolBuilder.h"
+#include "GeometryBase.h"
+#include "MeshAnalysisProperties.generated.h"
+
+
+// predeclarations
+PREDECLARE_USE_GEOMETRY_CLASS(FDynamicMesh3);
+
+
+
+UCLASS()
+class MESHMODELINGTOOLSEXP_API UMeshAnalysisProperties : public UInteractiveToolPropertySet
+{
+	GENERATED_BODY()
+public:
+
+	UPROPERTY(VisibleAnywhere, Category = MeshAnalysis)
+	FString SurfaceArea;
+
+	UPROPERTY(VisibleAnywhere, Category = MeshAnalysis)
+	FString Volume;
+
+	void Update(const FDynamicMesh3& Mesh, const FTransform& Transform);
+};
+
