MeshModelingToolset: Split plugin into non-experimental and experimental

diff --git a/Engine/Plugins/Experimental/MeshModelingToolsetExp/Source/MeshModelingToolsExp/Public/Properties/MeshMaterialProperties.h b/Engine/Plugins/Experimental/MeshModelingToolsetExp/Source/MeshModelingToolsExp/Public/Properties/MeshMaterialProperties.h
new file mode 100644
--- /dev/null
+++ b/Engine/Plugins/Experimental/MeshModelingToolsetExp/Source/MeshModelingToolsExp/Public/Properties/MeshMaterialProperties.h
@@ -0,0 +1,148 @@
+// Copyright Epic Games, Inc. All Rights Reserved.
+
+#pragma once
+
+#include "CoreMinimal.h"
+#include "Engine/Classes/Engine/Texture2D.h"
+#include "InteractiveToolBuilder.h"
+#include "MeshMaterialProperties.generated.h"
+
+
+// predeclarations
+class UMaterialInterface;
+class UMaterialInstanceDynamic;
+
+
+// Standard material property settings for tools that generate new meshes
+UCLASS()
+class MESHMODELINGTOOLSEXP_API UNewMeshMaterialProperties : public UInteractiveToolPropertySet
+{
+	GENERATED_BODY()
+public:
+	UNewMeshMaterialProperties();
+
+	/** Material for new mesh*/
+	UPROPERTY(EditAnywhere, NonTransactional, Category = Material, meta = (DisplayName = "Material"))
+	TWeakObjectPtr<UMaterialInterface> Material;
+
+	/** Scale factor for generated UVs */
+	UPROPERTY(EditAnywhere, Category = Material, meta = (DisplayName = "UV Scale", HideEditConditionToggle, EditConditionHides, EditCondition = "bShowExtendedOptions"))
+	float UVScale = 1.0;
+
+	/** If set, UV scales will be relative to world space so different objects created with the same UV scale should have the same average texel size */
+	UPROPERTY(EditAnywhere, Category = Material, meta = (DisplayName = "UV Scale Relative to World Space", HideEditConditionToggle, EditConditionHides, EditCondition = "bShowExtendedOptions"))
+	bool bWorldSpaceUVScale = false;
+
+	/** Overlay wireframe on preview */
+	UPROPERTY(EditAnywhere, Category = Material, meta = (DisplayName = "Show Wireframe", HideEditConditionToggle, EditConditionHides, EditCondition = "bShowExtendedOptions"))
+	bool bWireframe = false;
+
+	// controls visibility of UV/etc properties
+	UPROPERTY()
+	bool bShowExtendedOptions = true;
+};
+
+
+
+/** Standard material modes for tools that need to set custom materials for visualization */
+UENUM()
+enum class ESetMeshMaterialMode : uint8
+{
+	/** Input material */
+	KeepOriginal,
+
+	/** Checkerboard material */
+	Checkerboard,
+
+	/** Override material */
+	Override
+};
+
+// Standard material property settings for tools that visualize materials on existing meshes (e.g. to help show UVs)
+UCLASS()
+class MESHMODELINGTOOLSEXP_API UExistingMeshMaterialProperties : public UInteractiveToolPropertySet
+{
+	GENERATED_BODY()
+public:
+
+	/** Material that will be used on the mesh */
+	UPROPERTY(EditAnywhere, Category = MaterialPreview)
+	ESetMeshMaterialMode MaterialMode = ESetMeshMaterialMode::KeepOriginal;
+
+	UPROPERTY(EditAnywhere, Category = MaterialPreview, meta = (UIMin = "1.0", UIMax = "40.0", ClampMin = "0.01", ClampMax = "1000.0", EditConditionHides, EditCondition = "MaterialMode == ESetMeshMaterialMode::Checkerboard"))
+	float CheckerDensity = 20.0f;
+
+	UPROPERTY(EditAnywhere, Category = MaterialPreview, meta = (DisplayName = "Show UV Channel", EditConditionHides, EditCondition = "MaterialMode == ESetMeshMaterialMode::Checkerboard"))
+	int UVChannel = 0;
+
+	UPROPERTY(EditAnywhere, Category = MaterialPreview, meta = (EditCondition = "MaterialMode == ESetMeshMaterialMode::Override"))
+	TObjectPtr<UMaterialInterface> OverrideMaterial = nullptr;
+
+	UPROPERTY(meta = (TransientToolProperty))
+	TObjectPtr<UMaterialInstanceDynamic> CheckerMaterial = nullptr;
+
+	// Needs custom restore in order to call setup
+	virtual void RestoreProperties(UInteractiveTool* RestoreToTool, const FString& CacheIdentifier = TEXT("")) override;
+
+	void Setup();
+
+	void UpdateMaterials();
+	UMaterialInterface* GetActiveOverrideMaterial() const;
+};
+
+UENUM()
+enum class EMeshEditingMaterialModes
+{
+	ExistingMaterial,
+	Diffuse,
+	Grey,
+	Soft,
+	Transparent,
+	TangentNormal,
+	VertexColor,
+	CustomImage,
+	Custom
+};
+
+
+UCLASS()
+class MESHMODELINGTOOLSEXP_API UMeshEditingViewProperties : public UInteractiveToolPropertySet
+{
+	GENERATED_BODY()
+
+public:
+	/** Toggle drawing of wireframe overlay on/off [Alt+W] */
+	UPROPERTY(EditAnywhere, Category = Rendering)
+	bool bShowWireframe = false;
+
+	/** Set which material to use on object */
+	UPROPERTY(EditAnywhere, Category = Rendering)
+	EMeshEditingMaterialModes MaterialMode = EMeshEditingMaterialModes::Diffuse;
+
+	/** Toggle flat shading on/off */
+	UPROPERTY(EditAnywhere, Category = Rendering, meta = (EditConditionHides, 
+		EditCondition = "MaterialMode != EMeshEditingMaterialModes::ExistingMaterial && MaterialMode != EMeshEditingMaterialModes::Transparent && MaterialMode != EMeshEditingMaterialModes::Custom") )
+	bool bFlatShading = true;
+
+	/** Main Color of Material */
+	UPROPERTY(EditAnywhere, Category = Rendering, meta = (EditConditionHides, EditCondition = "MaterialMode == EMeshEditingMaterialModes::Diffuse"))
+	FLinearColor Color = FLinearColor(0.4f, 0.4f, 0.4f);
+
+	/** Image used in Image-Based Material */
+	UPROPERTY(EditAnywhere, Category = Rendering, meta = (EditConditionHides, EditCondition = "MaterialMode == EMeshEditingMaterialModes::CustomImage", TransientToolProperty) )
+	TObjectPtr<UTexture2D> Image;
+
+	//~ Could have used the same property as Color, above, but the user may want different saved values for the two
+	UPROPERTY(EditAnywhere, Category = Rendering, AdvancedDisplay, meta = (EditConditionHides, EditCondition = "MaterialMode == EMeshEditingMaterialModes::Transparent", DisplayName = "Color"))
+	FLinearColor TransparentMaterialColor = FLinearColor(0.0606, 0.309, 0.842);
+
+	UPROPERTY(EditAnywhere, Category = Rendering, meta = (EditConditionHides, EditCondition = "MaterialMode == EMeshEditingMaterialModes::Transparent", ClampMin = "0", ClampMax = "1.0"))
+	double Opacity = 0.65;
+
+	/** Although a two-sided transparent material causes rendering issues with overlapping faces, it is still frequently useful to see the shape when sculpting around other objects. */
+	UPROPERTY(EditAnywhere, Category = Rendering, meta = (EditConditionHides, EditCondition = "MaterialMode == EMeshEditingMaterialModes::Transparent"))
+	bool bTwoSided = true;
+
+	UPROPERTY(EditAnywhere, Category = Rendering, meta = (EditConditionHides, EditCondition = "MaterialMode == EMeshEditingMaterialModes::Custom"))
+	TWeakObjectPtr<UMaterialInterface> CustomMaterial;
+};
diff --git a/Engine/Plugins/Experimental/MeshModelingToolsetExp/Source/MeshModelingToolsExp/Public/Properties/MeshStatisticsProperties.h b/Engine/Plugins/Experimental/MeshModelingToolsetExp/Source/MeshModelingToolsExp/Public/Properties/MeshStatisticsProperties.h
new file mode 100644
--- /dev/null
+++ b/Engine/Plugins/Experimental/MeshModelingToolsetExp/Source/MeshModelingToolsExp/Public/Properties/MeshStatisticsProperties.h
@@ -0,0 +1,32 @@
+// Copyright Epic Games, Inc. All Rights Reserved.
+
+#pragma once
+
+#include "CoreMinimal.h"
+#include "InteractiveToolBuilder.h"
+#include "GeometryBase.h"
+#include "MeshStatisticsProperties.generated.h"
+
+
+// predeclarations
+PREDECLARE_USE_GEOMETRY_CLASS(FDynamicMesh3);
+
+
+UCLASS()
+class MESHMODELINGTOOLSEXP_API UMeshStatisticsProperties : public UInteractiveToolPropertySet
+{
+	GENERATED_BODY()
+public:
+
+	UPROPERTY(VisibleAnywhere, Category = MeshStatistics)
+	FString Mesh;
+
+	UPROPERTY(VisibleAnywhere, Category = MeshStatistics, AdvancedDisplay)
+	FString UV;
+
+	UPROPERTY(VisibleAnywhere, Category = MeshStatistics, AdvancedDisplay)
+	FString Attributes;
+
+	void Update(const FDynamicMesh3& Mesh);
+};
+
diff --git a/Engine/Plugins/Experimental/MeshModelingToolsetExp/Source/MeshModelingToolsExp/Public/Properties/RevolveProperties.h b/Engine/Plugins/Experimental/MeshModelingToolsetExp/Source/MeshModelingToolsExp/Public/Properties/RevolveProperties.h
new file mode 100644
--- /dev/null
+++ b/Engine/Plugins/Experimental/MeshModelingToolsetExp/Source/MeshModelingToolsExp/Public/Properties/RevolveProperties.h
@@ -0,0 +1,204 @@
+// Copyright Epic Games, Inc. All Rights Reserved.
+
+#pragma once
+
+#include "CoreMinimal.h"
+
+#include "InteractiveTool.h" //UInteractiveToolPropertySet
+#include "VectorTypes.h"
+
+#include "RevolveProperties.generated.h"
+
+class UNewMeshMaterialProperties;
+PREDECLARE_GEOMETRY(class FCurveSweepOp);
+
+UENUM()
+enum class ERevolvePropertiesCapFillMode : uint8
+{
+	/** No cap. */
+	None,
+	/** Cap is triangulated to maximize the minimal angle in the triangles (if they were to be
+	   projected onto a best-fit plane). */
+	Delaunay,
+	/** Cap is triangualted using a standard ear clipping approach. This could result in some
+	   very thin triangles. */
+	EarClipping,
+	/** A vertex is placed in the center and a fan is created to the boundary. This is nice if
+	   the cross section is convex, but creates invalid geometry if it isn't. */
+	CenterFan
+};
+
+
+UENUM()
+enum class ERevolvePropertiesPolygroupMode : uint8
+{
+	/** One polygroup for body of output mesh */
+	Single,
+	/** One polygroup per generated quad/triangle. */
+	PerFace,
+	/** Groups will be arranged in strips running in the profile curve direction, one per revolution step. */
+	PerStep,
+	/** Groups will be arranged in strips running along in the revolution direction according to profile curve. */
+	AccordingToProfileCurve
+};
+
+UENUM()
+enum class ERevolvePropertiesQuadSplit : uint8
+{
+	/** Quads will always be split the same way relative to an unrolled mesh, regardless of quad shape. */
+	Uniform,
+	
+	/** Quads will be split such that the shortest diagonal is connected. */
+	ShortestDiagonal
+};
+
+/**
+ * Common properties for revolving a polyline to create a mesh.
+ */
+UCLASS()
+class MESHMODELINGTOOLSEXP_API URevolveProperties : public UInteractiveToolPropertySet
+{
+	GENERATED_BODY()
+public:
+
+	URevolveProperties()
+	{
+		// We want the revolution degrees to be clamped to 360 when there is no offset along axis, and extendable
+		// beyond 360 when there is one (to make springs, etc). There's not currently a way to have a conditional
+		// clamp value, so instead, we swap the variables that we make visible to the user (RevolutionDegrees and
+		// ClampedRevolutionDegrees). However, we have to do some work to keep the values consistent so that they
+		// do not jump when they are swapped.
+		ClampedRevolutionDegreesWatcherIndex = WatchProperty(ClampedRevolutionDegrees,
+			[this](const double& NewClamped) 
+			{ 
+				RevolutionDegrees = NewClamped; 
+				SilentUpdateWatcherAtIndex(RevolutionDegreesWatcherIndex); //avoid triggering that watcher
+			});
+
+		RevolutionDegreesWatcherIndex = WatchProperty(RevolutionDegrees,
+			[this](const double& NewNonClamped)
+			{
+				ClampedRevolutionDegrees = FMath::Min(RevolutionDegrees, 360.0);
+				SilentUpdateWatcherAtIndex(ClampedRevolutionDegreesWatcherIndex); //avoid triggering that watcher
+			});
+	}
+
+	/** Revolution extent. Clamped to a max of 360 when not offsetting along axis.*/
+	UPROPERTY(EditAnywhere, Category = RevolveSettings, meta = (UIMin = "0", UIMax = "360", ClampMin = "0", ClampMax = "360", 
+		EditCondition = "AlongAxisOffsetPerDegree == 0", EditConditionHides, DisplayName = "Revolution Degrees"))
+	double ClampedRevolutionDegrees = 360;
+
+	/** Revolution extent. */
+	UPROPERTY(EditAnywhere, Category = RevolveSettings, meta = (UIMin = "0", UIMax = "3600", ClampMin = "0", ClampMax = "360000",
+		EditCondition = "AlongAxisOffsetPerDegree != 0", EditConditionHides))
+	double RevolutionDegrees = 360;
+
+	/** When true, the number of steps can be specified explicitly. When false, the number of steps is adjusted automatically. */
+	UPROPERTY(EditAnywhere, Category = RevolveSettings)
+	bool bExplicitSteps = false;
+
+	/** Number of steps to take while revolving. */
+	UPROPERTY(EditAnywhere, Category = RevolveSettings, meta = (UIMin = "1", ClampMin = "1", UIMax = "100", ClampMax = "5000",
+		Display = "Steps", EditCondition = "bExplicitSteps", EditConditionHides))
+	int NumExplicitSteps = 24;
+
+	/** The revolution is split into a number of steps such that each step moves the revolution no more than this number of degrees.  */
+	UPROPERTY(EditAnywhere, Category = RevolveSettings, meta = (UIMin = "1", ClampMin = "1", UIMax = "120", ClampMax = "180",
+		EditCondition = "!bExplicitSteps", EditConditionHides))
+	double MaxDegreesPerStep = 15;
+
+	/** How far to move each step along the axis (per degree). Used to create spirals. */
+	UPROPERTY(EditAnywhere, Category = RevolveSettings, meta = (UIMin = "-1", UIMax = "1", ClampMin = "-100000", ClampMax = "100000"))
+	double AlongAxisOffsetPerDegree = 0;
+
+	/** The angle by which to shift the profile curve around the axis before beginning the revolve */
+	UPROPERTY(EditAnywhere, Category = RevolveSettings, meta = (UIMin = "-360", UIMax = "360", ClampMin = "-36000", ClampMax = "36000"), AdvancedDisplay)
+	double RevolutionDegreesOffset = 0;
+
+	/** By default, revolution is done counterclockwise if looking down the revolution axis. This reverses the direction.*/
+	UPROPERTY(EditAnywhere, Category = RevolveSettings)
+	bool bReverseRevolutionDirection = false;
+
+	/** Flips the mesh inside out. */
+	UPROPERTY(EditAnywhere, Category = RevolveSettings)
+	bool bFlipMesh = false;
+
+	/** If true, then rather than revolving the profile directly, it is interpreted as the midpoint cross section of
+	 the first rotation step. Useful, for instance, for using the tool to create square columns. */
+	UPROPERTY(EditAnywhere, Category = RevolveSettings, AdvancedDisplay)
+	bool bProfileIsCrossSectionOfSide = false;
+
+	/** Determines grouping of generated triangles into polygroups. 
+	   Caps (if present) will always be separate groups. */
+	UPROPERTY(EditAnywhere, Category = RevolveSettings, AdvancedDisplay)
+	ERevolvePropertiesPolygroupMode PolygroupMode = ERevolvePropertiesPolygroupMode::PerFace;
+
+	/** Determines how any generated quads are split into triangles. */
+	UPROPERTY(EditAnywhere, Category = RevolveSettings, AdvancedDisplay)
+	ERevolvePropertiesQuadSplit QuadSplitMode = ERevolvePropertiesQuadSplit::ShortestDiagonal;
+	
+	/** When quads are generated using "shortest" diagonal, this biases the diagonal length comparison
+	 to prefer one slightly in the case of similar diagonals (for example, a value of 0.01 allows a
+	 1% difference in lengths before the triangulation is flipped). Helps symmetric quads be uniformly
+	 triangulated. */
+	UPROPERTY(EditAnywhere, Category = RevolveSettings, AdvancedDisplay, meta = (ClampMin = "0.0",  ClampMax = "2.0", 
+		EditCondition = "QuadSplitMode == ERevolvePropertiesQuadSplit::ShortestDiagonal", EditConditionHides))
+	double DiagonalProportionTolerance = 0.01;
+
+	/** Determines how caps are created if the revolution is partial. Not relevant if the
+	  revolution is full and welded. */
+	UPROPERTY(EditAnywhere, Category = RevolveSettings, AdvancedDisplay, meta = (
+		EditCondition = "AlongAxisOffsetPerDegree != 0 || RevolutionDegrees != 360", EditConditionHides))
+	ERevolvePropertiesCapFillMode CapFillMode = ERevolvePropertiesCapFillMode::Delaunay;
+
+	/** If true, the ends of a fully revolved profile are welded together, rather than duplicating
+	  vertices at the seam. Not relevant if the revolution is not full. */
+	UPROPERTY(EditAnywhere, Category = RevolveSettings, AdvancedDisplay, meta = (
+		EditCondition = "AlongAxisOffsetPerDegree == 0 && RevolutionDegrees == 360", EditConditionHides))
+	bool bWeldFullRevolution = true;
+
+	/** If true, vertices sufficiently close to the axis will not be replicated, instead reusing
+	  the same vertex for any adjacent triangles. */
+	UPROPERTY(EditAnywhere, Category = RevolveSettings, AdvancedDisplay)
+	bool bWeldVertsOnAxis = true;
+
+	/** If welding vertices on the axis, the distance that a vertex can be from the axis and still be welded */
+	UPROPERTY(EditAnywhere, Category = RevolveSettings, AdvancedDisplay, meta = (ClampMin = "0.0", ClampMax = "20.0", 
+		EditCondition = "bWeldVertsOnAxis", EditConditionHides))
+	double AxisWeldTolerance = 0.1;
+
+	/** If true, normals are not averaged or shared between triangles with sufficient angle difference. */
+	UPROPERTY(EditAnywhere, Category = RevolveSettings)
+	bool bSharpNormals = false;
+
+	/** When using sharp normals, the degree difference to accept between adjacent triangle normals to allow them to share
+	 normals at their vertices. */
+	UPROPERTY(EditAnywhere, Category = RevolveSettings, meta = (ClampMin = "0.0", ClampMax = "90.0", 
+		EditCondition = "bSharpNormals", EditConditionHides))
+	double SharpNormalAngleTolerance = 0.1;
+
+	/** If true, UV coordinates will be flipped in the V direction. */
+	UPROPERTY(EditAnywhere, Category = RevolveSettings, AdvancedDisplay)
+	bool bFlipVs = false;
+
+	/* If true, UV layout is not affected by segments of the profile curve that 
+	  do not result in any triangles (i.e., when both ends of the segment are welded
+	  due to being on the revolution axis).*/
+	UPROPERTY(EditAnywhere, Category = RevolveSettings, AdvancedDisplay)
+	bool bUVsSkipFullyWeldedEdges = true;
+
+
+	/**
+	 * Sets most of the settings for a FCurveSweepOp except for the profile curve itself. Should be called
+	 * after setting the profile curve, as the function reverses it if necessary.
+	 *
+	 * CurveSweepOpOut.ProfileCurve and CurveSweepOpOut.bProfileCurveIsClosed must be initialized in advance.
+	 */
+	void ApplyToCurveSweepOp(const UNewMeshMaterialProperties& MaterialProperties,
+		const FVector3d& RevolutionAxisOrigin, const FVector3d& RevolutionAxisDirection,
+		UE::Geometry::FCurveSweepOp& CurveSweepOpOut) const;
+
+	protected:
+		int32 ClampedRevolutionDegreesWatcherIndex;
+		int32 RevolutionDegreesWatcherIndex;
+};
diff --git a/Engine/Plugins/Experimental/MeshModelingToolsetExp/Source/MeshModelingToolsExp/Public/RevolveBoundaryTool.h b/Engine/Plugins/Experimental/MeshModelingToolsetExp/Source/MeshModelingToolsExp/Public/RevolveBoundaryTool.h
new file mode 100644
--- /dev/null
+++ b/Engine/Plugins/Experimental/MeshModelingToolsetExp/Source/MeshModelingToolsExp/Public/RevolveBoundaryTool.h
@@ -0,0 +1,121 @@
+// Copyright Epic Games, Inc. All Rights Reserved.
+
+#pragma once
+
+#include "CoreMinimal.h"
+
+#include "BaseBehaviors/BehaviorTargetInterfaces.h"
+#include "InteractiveToolBuilder.h"
+#include "SingleSelectionTool.h"
+#include "Mechanics/ConstructionPlaneMechanic.h"
+#include "MeshBoundaryToolBase.h"
+#include "MeshOpPreviewHelpers.h" //UMeshOpPreviewWithBackgroundCompute
+#include "ModelingOperators.h" //IDynamicMeshOperatorFactory
+#include "InteractiveTool.h"
+#include "InteractiveToolBuilder.h" //UInteractiveToolBuilder
+#include "Properties/MeshMaterialProperties.h"
+#include "Properties/RevolveProperties.h"
+#include "ToolContextInterfaces.h" // FToolBuilderState
+
+#include "RevolveBoundaryTool.generated.h"
+
+// Tool Builder
+
+UCLASS()
+class MESHMODELINGTOOLSEXP_API URevolveBoundaryToolBuilder : public USingleSelectionMeshEditingToolBuilder
+{
+	GENERATED_BODY()
+
+public:
+	virtual USingleSelectionMeshEditingTool* CreateNewTool(const FToolBuilderState& SceneState) const override;
+};
+
+UCLASS()
+class MESHMODELINGTOOLSEXP_API URevolveBoundaryOperatorFactory : public UObject, public UE::Geometry::IDynamicMeshOperatorFactory
+{
+	GENERATED_BODY()
+
+public:
+
+	TUniquePtr<UE::Geometry::FDynamicMeshOperator> MakeNewOperator() override;
+
+	UPROPERTY()
+	TObjectPtr<URevolveBoundaryTool> RevolveBoundaryTool;
+};
+
+UCLASS()
+class MESHMODELINGTOOLSEXP_API URevolveBoundaryToolProperties : public URevolveProperties
+{
+	GENERATED_BODY()
+
+public:
+	
+	UPROPERTY(EditAnywhere, Category = RevolveSettings, AdvancedDisplay)
+	bool bDisplayOriginalMesh = false;
+
+	UPROPERTY(EditAnywhere, Category = RevolutionAxis)
+	FVector AxisOrigin = FVector(0, 0, 0);
+
+	//~ We don't use a rotator for axis orientation because one of the components (roll) 
+	//~ will never do anything in the case of our axis.
+	UPROPERTY(EditAnywhere, Category = RevolutionAxis, meta = (
+		UIMin = -180, UIMax = 180, ClampMin = -180000, ClampMax = 18000))
+	float AxisYaw = 0;
+
+	UPROPERTY(EditAnywhere, Category = RevolutionAxis, meta = (
+		UIMin = -180, UIMax = 180, ClampMin = -180000, ClampMax = 18000))
+	float AxisPitch = 0;
+};
+
+/** 
+ * Tool that revolves the boundary of a mesh around an axis to create a new mesh. Mainly useful for
+ * revolving planar meshes. 
+ */
+UCLASS()
+class MESHMODELINGTOOLSEXP_API URevolveBoundaryTool : public UMeshBoundaryToolBase, public IClickBehaviorTarget
+{
+	GENERATED_BODY()
+
+public:
+
+	virtual bool CanAccept() const override;
+
+	virtual void Setup() override;
+	virtual void Shutdown(EToolShutdownType ShutdownType) override;
+
+	virtual void Render(IToolsContextRenderAPI* RenderAPI) override;
+	virtual void OnTick(float DeltaTime) override;
+	virtual void OnPropertyModified(UObject* PropertySet, FProperty* Property) override;
+
+	virtual void OnUpdateModifierState(int ModifierID, bool bIsOn) override;
+protected:
+
+	// Support for Ctrl+Clicking a boundary to align the revolution axis to that segment
+	bool bAlignAxisOnClick = false;
+	int32 AlignAxisModifier = 2;
+
+	UPROPERTY()
+	TObjectPtr<URevolveBoundaryToolProperties> Settings = nullptr;
+
+	UPROPERTY()
+	TObjectPtr<UNewMeshMaterialProperties> MaterialProperties;
+
+	UPROPERTY()
+	TObjectPtr<UConstructionPlaneMechanic> PlaneMechanic = nullptr;
+
+	UPROPERTY()
+	TObjectPtr<UMeshOpPreviewWithBackgroundCompute> Preview = nullptr;
+
+	FVector3d RevolutionAxisOrigin;
+	FVector3d RevolutionAxisDirection;
+
+	void GenerateAsset(const FDynamicMeshOpResult& Result);
+	void UpdateRevolutionAxis();
+	void StartPreview();
+
+	// IClickBehaviorTarget API
+	virtual FInputRayHit IsHitByClick(const FInputDeviceRay& ClickPos) override;
+	virtual void OnClicked(const FInputDeviceRay& ClickPos) override;
+
+	friend class URevolveBoundaryOperatorFactory;
+};
\ No newline at end of file
diff --git a/Engine/Plugins/Experimental/MeshModelingToolsetExp/Source/MeshModelingToolsExp/Public/Sculpting/MeshMoveBrushOps.h b/Engine/Plugins/Experimental/MeshModelingToolsetExp/Source/MeshModelingToolsExp/Public/Sculpting/MeshMoveBrushOps.h
new file mode 100644
--- /dev/null
+++ b/Engine/Plugins/Experimental/MeshModelingToolsetExp/Source/MeshModelingToolsExp/Public/Sculpting/MeshMoveBrushOps.h
@@ -0,0 +1,89 @@
+// Copyright Epic Games, Inc. All Rights Reserved.
+
+#pragma once
+
+#include "Sculpting/MeshBrushOpBase.h"
+#include "DynamicMesh/DynamicMesh3.h"
+#include "DynamicMesh/MeshNormals.h"
+#include "MeshWeights.h"
+#include "Async/ParallelFor.h"
+#include "PropertySets/AxisFilterPropertyType.h"
+#include "MeshMoveBrushOps.generated.h"
+
+
+
+UCLASS()
+class MESHMODELINGTOOLSEXP_API UMoveBrushOpProps : public UMeshSculptBrushOpProps
+{
+	GENERATED_BODY()
+public:
+	/** Strength of the Brush */
+	UPROPERTY(EditAnywhere, Category = MoveBrush, meta = (DisplayName = "Strength", UIMin = "0.0", UIMax = "1.0", ClampMin = "0.0", ClampMax = "1.0"))
+	float Strength = 1.0;
+
+	/** Amount of falloff to apply */
+	UPROPERTY(EditAnywhere, Category = MoveBrush, meta = (DisplayName = "Falloff", UIMin = "0.0", UIMax = "1.0", ClampMin = "0.0", ClampMax = "1.0"))
+	float Falloff = 0.5;
+
+	/** Depth of Brush into surface along view ray */
+	UPROPERTY(EditAnywhere, Category = MoveBrush, meta = (UIMin = "-0.5", UIMax = "0.5", ClampMin = "-1.0", ClampMax = "1.0"))
+	float Depth = 0;
+
+	virtual float GetStrength() override { return Strength; }
+	virtual float GetFalloff() override { return Falloff; }
+	virtual float GetDepth() override { return Depth; }
+
+	/** Axis filters restrict mesh movement to World X/Y/Z axes */
+	UPROPERTY(EditAnywhere, Category = MoveBrush, AdvancedDisplay)
+	FModelingToolsAxisFilter AxisFilters;
+};
+
+
+class FMoveBrushOp : public FMeshSculptBrushOp
+{
+public:
+
+	virtual void ApplyStamp(const FDynamicMesh3* Mesh, const FSculptBrushStamp& Stamp, const TArray<int32>& Vertices, TArray<FVector3d>& NewPositionsOut) override
+	{
+		const FVector3d& StampPos = Stamp.LocalFrame.Origin;
+
+		double UsePower = Stamp.Power;
+		FVector3d MoveVec = Stamp.LocalFrame.Origin - Stamp.PrevLocalFrame.Origin;
+		
+		UMoveBrushOpProps* Props = GetPropertySetAs<UMoveBrushOpProps>();
+
+		// if we have axis constraints we want to apply them in world space, and then remap
+		// that world vector to local space
+		if (Props->AxisFilters.AnyAxisFiltered())
+		{
+			FVector3d WorldMoveVec = Stamp.WorldFrame.Origin - Stamp.PrevWorldFrame.Origin;
+			WorldMoveVec = MoveVec.Length() * UE::Geometry::Normalized(WorldMoveVec);		// apply local space scaling
+			WorldMoveVec.X = (Props->AxisFilters.bAxisX) ? WorldMoveVec.X : 0;
+			WorldMoveVec.Y = (Props->AxisFilters.bAxisY) ? WorldMoveVec.Y : 0;
+			WorldMoveVec.Z = (Props->AxisFilters.bAxisZ) ? WorldMoveVec.Z : 0;
+			MoveVec = Stamp.LocalFrame.FromFrameVector(Stamp.WorldFrame.ToFrameVector(WorldMoveVec));
+		}
+
+		ParallelFor(Vertices.Num(), [&](int32 k)
+		{
+			int32 VertIdx = Vertices[k];
+			FVector3d OrigPos = Mesh->GetVertex(VertIdx);
+
+			double Falloff = GetFalloff().Evaluate(Stamp, OrigPos);
+
+			FVector3d NewPos = OrigPos + Falloff * MoveVec;
+			NewPositionsOut[k] = NewPos;
+		});
+	}
+
+
+	virtual ESculptBrushOpTargetType GetBrushTargetType() const override
+	{
+		return ESculptBrushOpTargetType::ActivePlane;
+	}
+
+	virtual bool IgnoreZeroMovements() const override
+	{
+		return true;
+	}
+};
diff --git a/Engine/Plugins/Experimental/MeshModelingToolsetExp/Source/MeshModelingToolsExp/Public/Sculpting/MeshSculptToolBase.h b/Engine/Plugins/Experimental/MeshModelingToolsetExp/Source/MeshModelingToolsExp/Public/Sculpting/MeshSculptToolBase.h
new file mode 100644
--- /dev/null
+++ b/Engine/Plugins/Experimental/MeshModelingToolsetExp/Source/MeshModelingToolsExp/Public/Sculpting/MeshSculptToolBase.h
@@ -0,0 +1,503 @@
+// Copyright Epic Games, Inc. All Rights Reserved.
+
+#pragma once
+
+#include "CoreMinimal.h"
+#include "BaseTools/MeshSurfacePointTool.h"
+#include "BaseTools/BaseBrushTool.h"
+#include "Sculpting/MeshBrushOpBase.h"
+#include "BoxTypes.h"
+#include "Properties/MeshMaterialProperties.h"
+#include "Changes/ValueWatcher.h"
+#include "Components/BaseDynamicMeshComponent.h"
+#include "MeshSculptToolBase.generated.h"
+
+
+class UMaterialInstanceDynamic;
+class UTransformGizmo;
+class UTransformProxy;
+class UPreviewMesh;
+
+
+/** Mesh Sculpting Brush Types */
+UENUM()
+enum class EMeshSculptFalloffType : uint8
+{
+	Smooth,
+	Linear,
+	Inverse,
+	Round,
+	BoxSmooth,
+	BoxLinear,
+	BoxInverse,
+	BoxRound,
+
+	LastValue UMETA(Hidden)
+};
+
+
+
+
+
+UCLASS()
+class MESHMODELINGTOOLSEXP_API USculptBrushProperties : public UBrushBaseProperties
+{
+	GENERATED_BODY()
+public:
+	USculptBrushProperties()
+	{
+		this->BrushFalloffAmount = 0.5f;
+	}
+
+	/** Depth of Brush into surface along view ray or surface normal, depending on the Active Brush Type */
+	UPROPERTY(EditAnywhere, Category = Brush, meta = (UIMin = "-0.5", UIMax = "0.5", ClampMin = "-1.0", ClampMax = "1.0", DisplayPriority = 5, HideEditConditionToggle, EditConditionHides, EditCondition = "bShowPerBrushProps"))
+	float Depth = 0;
+
+	/** Allow the Brush to hit the back-side of the mesh */
+	UPROPERTY(EditAnywhere, Category = Brush, meta = (DisplayPriority = 6))
+	bool bHitBackFaces = true;
+
+	/** Brush stamps are applied at this time interval. 0 for a single stamp, 1 for continuous stamps, 0.5 is a stamp every half-second */
+	UPROPERTY(EditAnywhere, Category = Brush, meta = (UIMin = "0.0", UIMax = "1.0", ClampMin = "0.0", ClampMax = "1.0", DisplayPriority = 7, HideEditConditionToggle, EditConditionHides, EditCondition = "bShowFlowRate"))
+	float FlowRate = 1.0f;
+
+	/** Minimum world-space spacing between stamps, defined along the brush path. Zero spacing means continuous stamps. */
+	UPROPERTY(EditAnywhere, Category = Brush, meta = (UIMin = "0.0", UIMax = "4.0", ClampMin = "0.0", ClampMax = "1000.0", DisplayPriority = 8, HideEditConditionToggle, EditConditionHides, EditCondition = "bShowSpacing"))
+	float Spacing = 0.0f;
+
+	/** Lazy brush smooths out the brush path by averaging the cursor positions */
+	UPROPERTY(EditAnywhere, Category = Brush, meta = (UIMin = "0.0", UIMax = "1.0", ClampMin = "0.0", ClampMax = "1.0", DisplayPriority = 9, HideEditConditionToggle, EditConditionHides, EditCondition = "bShowLazyness"))
+	float Lazyness = 0;
+
+	/**  */
+	UPROPERTY( meta = (TransientToolProperty))
+	bool bShowPerBrushProps = true;
+
+	/**  */
+	UPROPERTY(meta = (TransientToolProperty))
+	bool bShowLazyness = true;
+
+	/**  */
+	UPROPERTY(meta = (TransientToolProperty))
+	bool bShowFlowRate = true;
+
+	UPROPERTY(meta = (TransientToolProperty))
+	bool bShowSpacing = true;
+
+};
+
+
+
+
+UCLASS()
+class MESHMODELINGTOOLSEXP_API UKelvinBrushProperties : public UInteractiveToolPropertySet
+{
+	GENERATED_BODY()
+
+public:
+	/** Brush Fall off as fraction of brush size*/
+	UPROPERTY(EditAnywhere, Category = Kelvin, meta = (UIMin = "0.0", ClampMin = "0.0"))
+	float FallOffDistance = 1.f;
+
+	/** How much the mesh resists shear */
+	UPROPERTY(EditAnywhere, Category = Kelvin, meta = (UIMin = "0.0", ClampMin = "0.0"))
+	float Stiffness = 1.f;
+
+	/** How compressible the spatial region is: 1 - 2 x Poisson ratio */
+	UPROPERTY(EditAnywhere, Category = Kelvin, meta = (UIMin = "0.0", UIMax = "0.1", ClampMin = "0.0", ClampMax = "1.0"))
+	float Incompressiblity = 1.f;
+
+	/** Integration steps*/
+	UPROPERTY(EditAnywhere, Category = Kelvin, meta = (UIMin = "0.0", UIMax = "100", ClampMin = "0.0", ClampMax = "100"))
+	int BrushSteps = 3;
+};
+
+
+
+
+UCLASS()
+class MESHMODELINGTOOLSEXP_API UWorkPlaneProperties : public UInteractiveToolPropertySet
+{
+	GENERATED_BODY()
+
+public:
+	UPROPERTY( meta = (TransientToolProperty) )
+	bool bPropertySetEnabled = true;
+
+	/** Toggle whether Work Plane Positioing Gizmo is visible */
+	UPROPERTY(EditAnywhere, Category = TargetPlane, meta = (HideEditConditionToggle, EditCondition = "bPropertySetEnabled == true"))
+	bool bShowGizmo = true;
+
+	UPROPERTY(EditAnywhere, Category = TargetPlane, AdvancedDisplay, meta = (HideEditConditionToggle, EditCondition = "bPropertySetEnabled == true"))
+	FVector Position = FVector::ZeroVector;
+
+	UPROPERTY(EditAnywhere, Category = TargetPlane, AdvancedDisplay, meta = (HideEditConditionToggle, EditCondition = "bPropertySetEnabled == true"))
+	FQuat Rotation = FQuat::Identity;
+};
+
+
+
+
+UCLASS()
+class MESHMODELINGTOOLSEXP_API USculptMaxBrushProperties : public UInteractiveToolPropertySet
+{
+	GENERATED_BODY()
+public:
+	/** Specify maximum displacement height (relative to brush size) */
+	UPROPERTY(EditAnywhere, Category = SculptMaxBrush, meta = (UIMin = "0.0", UIMax = "1.0", ClampMin = "0.0", ClampMax = "1.0"))
+	float MaxHeight = 0.5;
+
+	/** Use maximum height from last brush stroke, regardless of brush size. Note that spatial brush falloff still applies.  */
+	UPROPERTY(EditAnywhere, Category = SculptMaxBrush)
+	bool bFreezeCurrentHeight = false;
+};
+
+
+
+
+/**
+ * Base Tool for mesh sculpting tools, provides some shared functionality
+ */
+UCLASS()
+class MESHMODELINGTOOLSEXP_API UMeshSculptToolBase : public UMeshSurfacePointTool
+{
+	GENERATED_BODY()
+protected:
+	using FTransform3d = UE::Geometry::FTransform3d;
+	using FFrame3d = UE::Geometry::FFrame3d;
+	using FRay3d = UE::Geometry::FRay3d;
+public:
+
+	virtual void RegisterActions(FInteractiveToolActionSet& ActionSet) override;
+
+	virtual void SetWorld(UWorld* World);
+
+	virtual void Setup() override;
+	virtual void Shutdown(EToolShutdownType ShutdownType) override;
+	virtual void Render(IToolsContextRenderAPI* RenderAPI) override;
+
+	// UMeshSurfacePointTool API
+	virtual bool HitTest(const FRay& Ray, FHitResult& OutHit) override;
+	virtual void OnBeginDrag(const FRay& Ray) override;
+	virtual void OnUpdateDrag(const FRay& Ray) override;
+	virtual void OnEndDrag(const FRay& Ray) override;
+	// end UMeshSurfacePointTool API
+
+protected:
+	virtual void OnTick(float DeltaTime) override;
+
+	
+	virtual void OnCompleteSetup();
+	virtual void OnBeginStroke(const FRay& WorldRay) { check(false); }
+	virtual void OnEndStroke() { check(false); }
+
+public:
+	/** Properties that control brush size/etc */
+	UPROPERTY()
+	TObjectPtr<USculptBrushProperties> BrushProperties;
+
+	/** Properties for 3D workplane / gizmo */
+	UPROPERTY()
+	TObjectPtr<UWorkPlaneProperties> GizmoProperties;
+
+
+protected:
+	UWorld* TargetWorld;		// required to spawn UPreviewMesh/etc
+	FViewCameraState CameraState;
+
+	/** Initial transformation on target mesh */
+	UE::Geometry::FTransform3d InitialTargetTransform;
+	/** Active transformation on target mesh, includes baked scale */
+	UE::Geometry::FTransform3d CurTargetTransform;
+
+	FRay3d GetLocalRay(const FRay& WorldRay) const;
+
+
+	/**
+	 * Subclass must implement this and return relevant rendering component
+	 */
+	virtual UBaseDynamicMeshComponent* GetSculptMeshComponent() { check(false); return nullptr; }
+
+	virtual FDynamicMesh3* GetSculptMesh() { return GetSculptMeshComponent()->GetMesh(); }
+	virtual const FDynamicMesh3* GetSculptMesh() const { return const_cast<UMeshSculptToolBase*>(this)->GetSculptMeshComponent()->GetMesh(); }
+
+	virtual FDynamicMesh3* GetBaseMesh() { check(false); return nullptr; }
+	virtual const FDynamicMesh3* GetBaseMesh() const { check(false); return nullptr; }
+
+
+	/**
+	 * Subclass calls this to set up editing component
+	 */
+	void InitializeSculptMeshComponent(UBaseDynamicMeshComponent* Component);
+
+
+	/**
+	 * Subclass can override this to change what results are written.
+	 * Default is to apply a default vertex positions update to the target object.
+	 */
+	virtual void CommitResult(UBaseDynamicMeshComponent* Component, bool bModifiedTopology);
+
+
+	//
+	// Brush Types
+	//
+protected:
+	UPROPERTY()
+	TMap<int32, TObjectPtr<UMeshSculptBrushOpProps>> BrushOpPropSets;
+
+	TMap<int32, TUniquePtr<FMeshSculptBrushOpFactory>> BrushOpFactories;
+
+	UPROPERTY()
+	TMap<int32, TObjectPtr<UMeshSculptBrushOpProps>> SecondaryBrushOpPropSets;
+
+	TMap<int32, TUniquePtr<FMeshSculptBrushOpFactory>> SecondaryBrushOpFactories;
+
+	void RegisterBrushType(int32 Identifier, TUniquePtr<FMeshSculptBrushOpFactory> Factory, UMeshSculptBrushOpProps* PropSet);
+	void RegisterSecondaryBrushType(int32 Identifier, TUniquePtr<FMeshSculptBrushOpFactory> Factory, UMeshSculptBrushOpProps* PropSet);
+
+	virtual void SaveAllBrushTypeProperties(UInteractiveTool* SaveFromTool);
+	virtual void RestoreAllBrushTypeProperties(UInteractiveTool* RestoreToTool);
+
+protected:
+	TUniquePtr<FMeshSculptBrushOp> PrimaryBrushOp;
+	UMeshSculptBrushOpProps* PrimaryVisiblePropSet = nullptr;		// BrushOpPropSets prevents GC of this
+
+	TUniquePtr<FMeshSculptBrushOp> SecondaryBrushOp;
+	UMeshSculptBrushOpProps* SecondaryVisiblePropSet = nullptr;
+
+	bool bBrushOpPropsVisible = true;
+
+	void SetActivePrimaryBrushType(int32 Identifier);
+	void SetActiveSecondaryBrushType(int32 Identifier);
+	TUniquePtr<FMeshSculptBrushOp>& GetActiveBrushOp();
+	void SetBrushOpPropsVisibility(bool bVisible);
+
+	//
+	// Falloff types
+	//
+protected:
+	TSharedPtr<FMeshSculptFallofFunc> PrimaryFalloff;
+
+	void SetPrimaryFalloffType(EMeshSculptFalloffType Falloff);
+
+
+
+
+	//
+	// Brush Size
+	//
+protected:
+	UE::Geometry::FInterval1d BrushRelativeSizeRange;
+	double CurrentBrushRadius = 1.0;
+	virtual void InitializeBrushSizeRange(const UE::Geometry::FAxisAlignedBox3d& TargetBounds);
+	virtual void CalculateBrushRadius();
+	virtual double GetCurrentBrushRadius() const { return CurrentBrushRadius; }
+
+	double CurrentBrushFalloff = 0.5;
+	virtual double GetCurrentBrushFalloff() const { return CurrentBrushFalloff; }
+
+	double ActivePressure = 1.0;
+	virtual double GetActivePressure() const { return ActivePressure; }
+
+	virtual double GetCurrentBrushStrength();
+	virtual double GetCurrentBrushDepth();
+
+public:
+	virtual void IncreaseBrushRadiusAction();
+	virtual void DecreaseBrushRadiusAction();
+	virtual void IncreaseBrushRadiusSmallStepAction();
+	virtual void DecreaseBrushRadiusSmallStepAction();
+
+
+	// client currently needs to implement these...
+	virtual void IncreaseBrushSpeedAction() {}
+	virtual void DecreaseBrushSpeedAction() {}
+	virtual void NextBrushModeAction() {}
+	virtual void PreviousBrushModeAction() {}
+
+
+
+	//
+	// Brush/Stroke stuff
+	//
+protected:
+	FFrame3d LastBrushFrameWorld;
+	FFrame3d LastBrushFrameLocal;
+	int32 LastBrushTriangleID;
+
+	const FFrame3d& GetBrushFrameWorld() const { return LastBrushFrameWorld; }
+	const FFrame3d& GetBrushFrameLocal() const { return LastBrushFrameLocal; }
+	int32 GetBrushTriangleID() const { return LastBrushTriangleID; }
+	void UpdateBrushFrameWorld(const FVector3d& NewPosition, const FVector3d& NewNormal);
+	void AlignBrushToView();
+
+	bool GetBrushCanHitBackFaces() const { return BrushProperties->bHitBackFaces; }
+
+	/** @return hit triangle at ray position - subclass must implement this */
+	virtual int32 FindHitSculptMeshTriangle(const FRay3d& LocalRay) { check(false); return IndexConstants::InvalidID; }
+	/** @return hit triangle at ray position - subclass should implement this for most brushes */
+	virtual int32 FindHitTargetMeshTriangle(const FRay3d& LocalRay) { check(false); return IndexConstants::InvalidID;	}
+
+	virtual bool UpdateBrushPositionOnActivePlane(const FRay& WorldRay);
+	virtual bool UpdateBrushPositionOnTargetMesh(const FRay& WorldRay, bool bFallbackToViewPlane);
+	virtual bool UpdateBrushPositionOnSculptMesh(const FRay& WorldRay, bool bFallbackToViewPlane);
+
+
+
+	//
+	// Brush Target Plane is plane that some brushes move on
+	//
+protected:
+	FFrame3d ActiveBrushTargetPlaneWorld;
+	virtual void UpdateBrushTargetPlaneFromHit(const FRay& WorldRay, const FHitResult& Hit);
+
+
+	//
+	// Stroke Modifiers
+	//
+protected:
+	bool bInStroke = false;
+	bool bSmoothing = false;
+	bool bInvert = false;
+	virtual void SaveActiveStrokeModifiers();
+	virtual bool InStroke() const { return bInStroke; }
+	virtual bool GetInSmoothingStroke() const { return bSmoothing; }
+	virtual bool GetInInvertStroke() const { return bInvert; }
+
+	// when in a stroke, this function determines when a new stamp should be emitted, based on spacing and flow rate settings
+	virtual void UpdateStampPendingState();
+
+	// for tracking stroke time and length, to apply spacing and flow rate settings
+	double ActiveStrokeTime = 0.0;
+	double ActiveStrokePathArcLen = 0.0;
+	int LastFlowTimeStamp = 0;
+	int LastSpacingTimestamp = 0;
+	virtual void ResetStrokeTime();
+	virtual void AccumulateStrokeTime(float DeltaTime);
+
+	//
+	// Stamps
+	//
+protected:
+	bool bIsStampPending = false;
+	FRay PendingStampRay;
+	FSculptBrushStamp HoverStamp;
+	FSculptBrushStamp CurrentStamp;
+	FSculptBrushStamp LastStamp;
+	virtual void UpdateHoverStamp(const FFrame3d& StampFrame);
+	virtual bool IsStampPending() const { return bIsStampPending; }
+	virtual const FRay& GetPendingStampRayWorld() const { return PendingStampRay;  }
+
+
+	//
+	// Stamp ROI Plane is a plane used by some brush ops
+	//
+protected:
+	FFrame3d StampRegionPlane;
+	virtual FFrame3d ComputeStampRegionPlane(const FFrame3d& StampFrame, const TArray<int32>& StampTriangles, bool bIgnoreDepth, bool bViewAligned, bool bInvDistFalloff = true);
+	virtual FFrame3d ComputeStampRegionPlane(const FFrame3d& StampFrame, const TSet<int32>& StampTriangles, bool bIgnoreDepth, bool bViewAligned, bool bInvDistFalloff = true);
+
+
+	// Stroke plane is a plane used by some brush ops
+	//
+protected:
+	FFrame3d StrokePlane;
+	virtual const FFrame3d& GetCurrentStrokeReferencePlane() const { return StrokePlane; }
+
+	virtual void UpdateStrokeReferencePlaneForROI(const FFrame3d& StampFrame, const TArray<int32>& TriangleROI, bool bViewAligned);
+	virtual void UpdateStrokeReferencePlaneFromWorkPlane();
+
+
+	//
+	// Display / Material
+	//
+public:
+	UPROPERTY()
+	TObjectPtr<UMeshEditingViewProperties> ViewProperties;
+
+	UPROPERTY()
+	TObjectPtr<UMaterialInstanceDynamic> ActiveOverrideMaterial;
+
+protected:
+	virtual void SetViewPropertiesEnabled(bool bNewValue);
+	virtual void UpdateWireframeVisibility(bool bNewValue);
+	virtual void UpdateMaterialMode(EMeshEditingMaterialModes NewMode);
+	virtual void UpdateFlatShadingSetting(bool bNewValue);
+	virtual void UpdateColorSetting(FLinearColor NewColor);
+	virtual void UpdateOpacitySetting(double Opacity);
+	virtual void UpdateTwoSidedSetting(bool bOn);
+	virtual void UpdateCustomMaterial(TWeakObjectPtr<UMaterialInterface> NewMaterial);
+	virtual void UpdateImageSetting(UTexture2D* NewImage);
+
+
+
+	//
+	// brush indicator
+	//
+protected:
+	// subclasses should call this to create indicator in their ::Setup()
+	virtual void InitializeIndicator();
+
+	virtual bool GetIsVolumetricIndicator();
+
+	virtual void ConfigureIndicator(bool bVolumetric);
+
+	virtual void SetIndicatorVisibility(bool bVisible);
+	virtual bool GetIndicatorVisibility() const;
+
+protected:
+	UPROPERTY()
+	TObjectPtr<UBrushStampIndicator> BrushIndicator;
+
+	UPROPERTY()
+	bool bIsVolumetricIndicator;
+
+
+	UPROPERTY()
+	TObjectPtr<UMaterialInstanceDynamic> BrushIndicatorMaterial;
+
+	UPROPERTY()
+	TObjectPtr<UPreviewMesh> BrushIndicatorMesh;
+
+	// creates default sphere indicator
+	UPreviewMesh* MakeDefaultIndicatorSphereMesh(UObject* Parent, UWorld* World, int Resolution = 32);
+
+
+
+
+	//
+	// Work Plane
+	//
+public:
+	// plane gizmo
+	UPROPERTY()
+	TObjectPtr<UTransformGizmo> PlaneTransformGizmo;
+
+	UPROPERTY()
+	TObjectPtr<UTransformProxy> PlaneTransformProxy;
+
+protected:
+	virtual void UpdateWorkPlane();
+	virtual bool ShowWorkPlane() const { return false; };
+
+protected:
+	TValueWatcher<FVector> GizmoPositionWatcher;
+	TValueWatcher<FQuat> GizmoRotationWatcher;
+
+	virtual void UpdateGizmoFromProperties();
+	virtual void PlaneTransformChanged(UTransformProxy* Proxy, FTransform Transform);
+
+	enum class EPendingWorkPlaneUpdate
+	{
+		NoUpdatePending,
+		MoveToHitPositionNormal,
+		MoveToHitPosition,
+		MoveToHitPositionViewAligned
+	};
+	EPendingWorkPlaneUpdate PendingWorkPlaneUpdate;
+	virtual void SetFixedSculptPlaneFromWorldPos(const FVector& Position, const FVector& Normal, EPendingWorkPlaneUpdate UpdateType);
+	virtual void UpdateFixedSculptPlanePosition(const FVector& Position);
+	virtual void UpdateFixedSculptPlaneRotation(const FQuat& Rotation);
+	virtual void UpdateFixedPlaneGizmoVisibility(bool bVisible);
+
+};
\ No newline at end of file
diff --git a/Engine/Plugins/Experimental/MeshModelingToolsetExp/Source/MeshModelingToolsExp/Public/SeamSculptTool.h b/Engine/Plugins/Experimental/MeshModelingToolsetExp/Source/MeshModelingToolsExp/Public/SeamSculptTool.h
new file mode 100644
--- /dev/null
+++ b/Engine/Plugins/Experimental/MeshModelingToolsetExp/Source/MeshModelingToolsExp/Public/SeamSculptTool.h
@@ -0,0 +1,116 @@
+// Copyright Epic Games, Inc. All Rights Reserved.
+
+#pragma once
+
+#include "CoreMinimal.h"
+#include "DynamicMeshBrushTool.h"
+#include "SeamSculptTool.generated.h"
+
+class UExistingMeshMaterialProperties;
+class UPreviewGeometry;
+PREDECLARE_USE_GEOMETRY_CLASS(FDynamicMesh3);
+
+/**
+ *
+ */
+UCLASS()
+class MESHMODELINGTOOLSEXP_API USeamSculptToolBuilder : public UMeshSurfacePointToolBuilder
+{
+	GENERATED_BODY()
+
+public:
+	virtual UMeshSurfacePointTool* CreateNewTool(const FToolBuilderState& SceneState) const override;
+};
+
+
+
+
+
+UCLASS()
+class MESHMODELINGTOOLSEXP_API USeamSculptToolProperties : public UInteractiveToolPropertySet
+{
+	GENERATED_BODY()
+public:
+
+	UPROPERTY(EditAnywhere, Category = Options)
+	bool bShowWireframe = false;
+
+
+	UPROPERTY(EditAnywhere, Category = Options)
+	bool bHitBackFaces = false;
+};
+
+
+
+
+
+
+
+UCLASS(Transient)
+class MESHMODELINGTOOLSEXP_API USeamSculptTool : public UDynamicMeshBrushTool
+{
+	GENERATED_BODY()
+public:
+	USeamSculptTool();
+
+	virtual void SetWorld(UWorld* World);
+
+	virtual void Setup() override;
+	virtual void OnTick(float DeltaTime) override;
+	virtual void Render(IToolsContextRenderAPI* RenderAPI) override;
+
+	virtual bool HasCancel() const override { return true; }
+	virtual bool HasAccept() const override { return true; }
+
+	// UBaseBrushTool overrides
+	virtual void OnBeginDrag(const FRay& Ray) override;
+	virtual void OnUpdateDrag(const FRay& Ray) override;
+	virtual void OnEndDrag(const FRay& Ray) override;
+	virtual bool OnUpdateHover(const FInputDeviceRay& DevicePos) override;
+
+
+protected:
+	virtual void ApplyStamp(const FBrushStampData& Stamp);
+
+	virtual void OnShutdown(EToolShutdownType ShutdownType) override;
+
+
+public:
+	UPROPERTY()
+	TObjectPtr<USeamSculptToolProperties> Settings;
+
+
+protected:
+	UPROPERTY()
+	TObjectPtr<UPreviewGeometry> PreviewGeom;
+
+	TSharedPtr<FDynamicMesh3, ESPMode::ThreadSafe> InputMesh;
+	UE::Geometry::FTransform3d MeshTransform;
+	double NormalOffset = 0;
+
+	void InitPreviewGeometry();
+	void UpdatePreviewGeometry();
+	bool bPreviewGeometryNeedsUpdate = false;
+
+	FVector3d CurrentSnapPositionLocal = FVector3d::Zero();
+	int32 CurrentSnapVertex = -1;
+
+	FVector3d DrawPathStartPositionLocal = FVector3d::Zero();
+	int32 DrawPathStartVertex = -1;
+
+	TArray<int32> CurDrawPath;
+	void UpdateCurrentDrawPath();
+
+
+	void CreateSeamAlongPath();
+
+	enum class EActiveCaptureState
+	{
+		NoState,
+		DrawNewPath
+	};
+	EActiveCaptureState CaptureState = EActiveCaptureState::NoState;
+
+	UWorld* TargetWorld;
+
+};
\ No newline at end of file
diff --git a/Engine/Plugins/Experimental/MeshModelingToolsetExp/Source/MeshModelingToolsExp/Public/SelfUnionMeshesTool.h b/Engine/Plugins/Experimental/MeshModelingToolsetExp/Source/MeshModelingToolsExp/Public/SelfUnionMeshesTool.h
new file mode 100644
--- /dev/null
+++ b/Engine/Plugins/Experimental/MeshModelingToolsetExp/Source/MeshModelingToolsExp/Public/SelfUnionMeshesTool.h
@@ -0,0 +1,115 @@
+// Copyright Epic Games, Inc. All Rights Reserved.
+
+#pragma once
+
+#include "CoreMinimal.h"
+#include "UObject/NoExportTypes.h"
+#include "Drawing/LineSetComponent.h"
+#include "BaseTools/BaseCreateFromSelectedTool.h"
+
+#include "CompositionOps/SelfUnionMeshesOp.h"
+
+#include "SelfUnionMeshesTool.generated.h"
+
+// predeclarations
+PREDECLARE_USE_GEOMETRY_CLASS(FDynamicMesh3);
+
+
+/**
+ * Standard properties of the self-union operation
+ */
+UCLASS()
+class MESHMODELINGTOOLSEXP_API USelfUnionMeshesToolProperties : public UInteractiveToolPropertySet
+{
+	GENERATED_BODY()
+public:
+	/** Automatically attempt to fill any holes left by merging (e.g. due to numerical errors) */
+	UPROPERTY(EditAnywhere, Category = Options)
+	bool bAttemptFixHoles = false;
+
+	/** Try to collapse extra edges created by the Boolean operation */
+	UPROPERTY(EditAnywhere, Category = Options)
+	bool bCollapseExtraEdges = true;
+
+	/** Show boundary edges created by the union operation -- often due to numerical error */
+	UPROPERTY(EditAnywhere, Category = Options)
+	bool bShowNewBoundaryEdges = true;
+
+	/** If true, remove open, visible geometry */
+	UPROPERTY(EditAnywhere, Category = Options)
+	bool bTrimFlaps = false;
+
+	/** Winding number threshold to determine what is consider inside the mesh */
+	UPROPERTY(EditAnywhere, Category = Options, AdvancedDisplay)
+	double WindingNumberThreshold = .5;
+
+	/** If true, only the first mesh will keep its materials assignments; all other triangles will be assigned material 0 */
+	UPROPERTY(EditAnywhere, Category = Materials)
+	bool bOnlyUseFirstMeshMaterials = false;
+};
+
+
+
+/**
+ * Union of meshes, resolving self intersections
+ */
+UCLASS()
+class MESHMODELINGTOOLSEXP_API USelfUnionMeshesTool : public UBaseCreateFromSelectedTool
+{
+	GENERATED_BODY()
+
+public:
+
+	USelfUnionMeshesTool() {}
+
+protected:
+
+	void TransformChanged(UTransformProxy* Proxy, FTransform Transform) override;
+
+	virtual void OnPropertyModified(UObject* PropertySet, FProperty* Property) override;
+
+	virtual void ConvertInputsAndSetPreviewMaterials(bool bSetPreviewMesh = true) override;
+
+	virtual void SetupProperties() override;
+	virtual void SaveProperties() override;
+	virtual void SetPreviewCallbacks() override;
+
+	virtual FString GetCreatedAssetName() const;
+	virtual FText GetActionName() const;
+
+	// IDynamicMeshOperatorFactory API
+	virtual TUniquePtr<UE::Geometry::FDynamicMeshOperator> MakeNewOperator() override;
+
+protected:
+
+	void UpdateVisualization();
+
+protected:
+
+	UPROPERTY()
+	TObjectPtr<USelfUnionMeshesToolProperties> Properties;
+
+	UPROPERTY()
+	TObjectPtr<ULineSetComponent> DrawnLineSet;
+
+	TSharedPtr<FDynamicMesh3, ESPMode::ThreadSafe> CombinedSourceMeshes;
+
+	// for visualization of any errors in the currently-previewed merge operation
+	TArray<int> CreatedBoundaryEdges;
+};
+
+
+UCLASS()
+class MESHMODELINGTOOLSEXP_API USelfUnionMeshesToolBuilder : public UBaseCreateFromSelectedToolBuilder
+{
+	GENERATED_BODY()
+
+public:
+	virtual UBaseCreateFromSelectedTool* MakeNewToolInstance(UObject* Outer) const override
+	{
+		return NewObject<USelfUnionMeshesTool>(Outer);
+	}
+};
+
+
+
diff --git a/Engine/Plugins/Experimental/MeshModelingToolsetExp/Source/MeshModelingToolsExp/Public/SplitMeshesTool.h b/Engine/Plugins/Experimental/MeshModelingToolsetExp/Source/MeshModelingToolsExp/Public/SplitMeshesTool.h
new file mode 100644
--- /dev/null
+++ b/Engine/Plugins/Experimental/MeshModelingToolsetExp/Source/MeshModelingToolsExp/Public/SplitMeshesTool.h
@@ -0,0 +1,83 @@
+// Copyright Epic Games, Inc. All Rights Reserved.
+
+#pragma once
+
+#include "CoreMinimal.h"
+#include "MultiSelectionTool.h"
+#include "InteractiveToolBuilder.h"
+#include "PropertySets/CreateMeshObjectTypeProperties.h"
+#include "SplitMeshesTool.generated.h"
+
+class UMaterialInterface;
+
+
+UCLASS()
+class MESHMODELINGTOOLSEXP_API USplitMeshesToolBuilder : public UInteractiveToolBuilder
+{
+	GENERATED_BODY()
+public:
+	virtual bool CanBuildTool(const FToolBuilderState& SceneState) const override;
+	virtual UInteractiveTool* BuildTool(const FToolBuilderState& SceneState) const override;
+
+protected:
+	virtual const FToolTargetTypeRequirements& GetTargetRequirements() const override;
+};
+
+
+
+UCLASS()
+class MESHMODELINGTOOLSEXP_API USplitMeshesToolProperties : public UInteractiveToolPropertySet
+{
+	GENERATED_BODY()
+public:
+	UPROPERTY(EditAnywhere, Category = Options)
+	bool bTransferMaterials = true;
+};
+
+
+
+UCLASS()
+class MESHMODELINGTOOLSEXP_API USplitMeshesTool : public UMultiSelectionTool
+{
+	GENERATED_BODY()
+
+public:
+	virtual void SetWorld(UWorld* World);
+
+	virtual void Setup() override;
+	virtual void Shutdown(EToolShutdownType ShutdownType) override;
+
+	virtual bool HasCancel() const override { return true; }
+	virtual bool HasAccept() const override { return true; }
+	virtual bool CanAccept() const override;
+
+	UPROPERTY()
+	TObjectPtr<USplitMeshesToolProperties> BasicProperties;
+
+	UPROPERTY()
+	TObjectPtr<UCreateMeshObjectTypeProperties> OutputTypeProperties;
+
+protected:
+	UWorld* TargetWorld;
+
+	struct FSourceMeshInfo
+	{
+		UE::Geometry::FDynamicMesh3 Mesh;
+		TArray<UMaterialInterface*> Materials;
+	};
+	TArray<FSourceMeshInfo> SourceMeshes;
+
+
+	struct FComponentsInfo
+	{
+		bool bNoComponents;
+		TArray<UE::Geometry::FDynamicMesh3> Meshes;
+		TArray<TArray<UMaterialInterface*>> Materials;
+		TArray<FVector3d> Origins;
+	};
+	TArray<FComponentsInfo> SplitMeshes;
+
+	int32 NoSplitCount = 0;
+
+	void UpdateSplitMeshes();
+};
diff --git a/Engine/Plugins/Experimental/MeshModelingToolsetExp/Source/MeshModelingToolsExp/Public/TransferMeshTool.h b/Engine/Plugins/Experimental/MeshModelingToolsetExp/Source/MeshModelingToolsExp/Public/TransferMeshTool.h
new file mode 100644
--- /dev/null
+++ b/Engine/Plugins/Experimental/MeshModelingToolsetExp/Source/MeshModelingToolsExp/Public/TransferMeshTool.h
@@ -0,0 +1,60 @@
+// Copyright Epic Games, Inc. All Rights Reserved.
+
+#pragma once
+
+#include "CoreMinimal.h"
+#include "MultiSelectionTool.h"
+#include "InteractiveToolBuilder.h"
+#include "TransferMeshTool.generated.h"
+
+/**
+ *
+ */
+UCLASS()
+class MESHMODELINGTOOLSEXP_API UTransferMeshToolBuilder : public UInteractiveToolBuilder
+{
+	GENERATED_BODY()
+
+public:
+	virtual bool CanBuildTool(const FToolBuilderState& SceneState) const override;
+	virtual UInteractiveTool* BuildTool(const FToolBuilderState& SceneState) const override;
+
+protected:
+	virtual const FToolTargetTypeRequirements& GetTargetRequirements() const override;
+};
+
+/**
+ * Standard properties of the Transfer operation
+ */
+UCLASS()
+class MESHMODELINGTOOLSEXP_API UTransferMeshToolProperties : public UInteractiveToolPropertySet
+{
+	GENERATED_BODY()
+public:
+	UPROPERTY(EditAnywhere, Category = Options)
+	bool bTransferMaterials = true;
+};
+
+
+
+UCLASS()
+class MESHMODELINGTOOLSEXP_API UTransferMeshTool : public UMultiSelectionTool
+{
+	GENERATED_BODY()
+
+public:
+	virtual void SetWorld(UWorld* World);
+
+	virtual void Setup() override;
+	virtual void Shutdown(EToolShutdownType ShutdownType) override;
+
+	virtual bool HasCancel() const override { return true; }
+	virtual bool HasAccept() const override { return true; }
+	virtual bool CanAccept() const override;
+
+	UPROPERTY()
+	TObjectPtr<UTransferMeshToolProperties> BasicProperties;
+
+protected:
+	UWorld* TargetWorld;
+};
diff --git a/Engine/Plugins/Experimental/MeshModelingToolsetExp/Source/MeshModelingToolsExp/Public/VolumeToMeshTool.h b/Engine/Plugins/Experimental/MeshModelingToolsetExp/Source/MeshModelingToolsExp/Public/VolumeToMeshTool.h
new file mode 100644
--- /dev/null
+++ b/Engine/Plugins/Experimental/MeshModelingToolsetExp/Source/MeshModelingToolsExp/Public/VolumeToMeshTool.h
@@ -0,0 +1,110 @@
+// Copyright Epic Games, Inc. All Rights Reserved.
+
+#pragma once
+
+#include "CoreMinimal.h"
+#include "GameFramework/Volume.h"
+#include "SingleSelectionTool.h"
+#include "InteractiveToolBuilder.h"
+#include "DynamicMesh/DynamicMesh3.h"
+#include "PreviewMesh.h"
+#include "Drawing/LineSetComponent.h"
+#include "VolumeToMeshTool.generated.h"
+
+/**
+ *
+ */
+UCLASS()
+class MESHMODELINGTOOLSEXP_API UVolumeToMeshToolBuilder : public UInteractiveToolBuilder
+{
+	GENERATED_BODY()
+public:
+	virtual bool CanBuildTool(const FToolBuilderState& SceneState) const override;
+	virtual UInteractiveTool* BuildTool(const FToolBuilderState& SceneState) const override;
+};
+
+
+
+
+UENUM()
+enum class EVolumeToMeshMode
+{
+	/** Convert based on Angle Tolerance between Face Normals */
+	TriangulatePolygons,
+	/** Create PolyGroups based on UV Islands */
+	MinimalPolygons
+};
+
+
+
+UCLASS()
+class MESHMODELINGTOOLSEXP_API UVolumeToMeshToolProperties : public UInteractiveToolPropertySet
+{
+	GENERATED_BODY()
+
+public:
+	UPROPERTY(EditAnywhere, Category = Options)
+	bool bWeldEdges = true;
+
+	UPROPERTY(EditAnywhere, Category = Options)
+	bool bAutoRepair = true;
+
+	UPROPERTY(EditAnywhere, Category = Options)
+	bool bOptimizeMesh = true;
+
+	UPROPERTY(EditAnywhere, Category = Options)
+	bool bShowWireframe = true;
+};
+
+
+
+/**
+ *
+ */
+UCLASS()
+class MESHMODELINGTOOLSEXP_API UVolumeToMeshTool : public UInteractiveTool
+{
+	GENERATED_BODY()
+
+public:
+	UVolumeToMeshTool();
+
+	virtual void SetWorld(UWorld* World) { this->TargetWorld = World; }
+	virtual void SetSelection(AVolume* Volume);
+
+	virtual void Setup() override;
+	virtual void Shutdown(EToolShutdownType ShutdownType) override;
+
+	virtual void Render(IToolsContextRenderAPI* RenderAPI) override;
+	virtual void OnTick(float DeltaTime) override;
+
+	virtual bool HasCancel() const override { return true; }
+	virtual bool HasAccept() const override { return true; }
+	virtual bool CanAccept() const override;
+
+
+protected:
+	UPROPERTY()
+	TObjectPtr<UVolumeToMeshToolProperties> Settings;
+
+	UPROPERTY()
+	TObjectPtr<UPreviewMesh> PreviewMesh;
+
+	UPROPERTY()
+	TLazyObjectPtr<AVolume> TargetVolume;
+
+	UPROPERTY()
+	TObjectPtr<ULineSetComponent> VolumeEdgesSet;
+
+protected:
+	UWorld* TargetWorld = nullptr;
+
+	UE::Geometry::FDynamicMesh3 CurrentMesh;
+
+	void RecalculateMesh();
+
+	void UpdateLineSet();
+
+	bool bResultValid = false;
+
+};
diff --git a/Engine/Plugins/Experimental/MeshModelingToolsetExp/Source/MeshModelingToolsExp/Public/VoxelMorphologyMeshesTool.h b/Engine/Plugins/Experimental/MeshModelingToolsetExp/Source/MeshModelingToolsExp/Public/VoxelMorphologyMeshesTool.h
new file mode 100644
--- /dev/null
+++ b/Engine/Plugins/Experimental/MeshModelingToolsetExp/Source/MeshModelingToolsExp/Public/VoxelMorphologyMeshesTool.h
@@ -0,0 +1,87 @@
+// Copyright Epic Games, Inc. All Rights Reserved.
+
+#pragma once
+
+#include "CoreMinimal.h"
+#include "UObject/NoExportTypes.h"
+#include "BaseTools/BaseVoxelTool.h"
+#include "CompositionOps/VoxelMorphologyMeshesOp.h"
+
+#include "VoxelMorphologyMeshesTool.generated.h"
+
+
+
+/**
+ * Properties of the morphology tool
+ */
+UCLASS()
+class MESHMODELINGTOOLSEXP_API UVoxelMorphologyMeshesToolProperties : public UInteractiveToolPropertySet
+{
+	GENERATED_BODY()
+
+public:
+
+	UPROPERTY(EditAnywhere, Category = Morphology)
+	EMorphologyOperation Operation = EMorphologyOperation::Dilate;
+
+	UPROPERTY(EditAnywhere, Category = Morphology, meta = (UIMin = ".1", UIMax = "100", ClampMin = ".001", ClampMax = "1000"))
+	double Distance = 5;
+
+	/** Apply a "VoxWrap" operation to input mesh(es) before computing the morphology.  Fixes results for inputs with holes and/or self-intersections. */
+	UPROPERTY(EditAnywhere, Category = VoxWrapPreprocess)
+	bool bVoxWrap = false;
+
+	/** Remove internal surfaces from the VoxWrap output, before computing the morphology. */
+	UPROPERTY(EditAnywhere, Category = VoxWrapPreprocess, meta = (EditCondition = "bVoxWrap == true"))
+	bool bRemoveInternalsAfterVoxWrap = false;
+
+	/** Thicken open-boundary surfaces (extrude them inwards) before VoxWrapping them. Units are in world space. If 0, no extrusion is applied. */
+	UPROPERTY(EditAnywhere, Category = VoxWrapPreprocess, meta = (EditCondition = "bVoxWrap == true", UIMin = "0", UIMax = "100", ClampMin = "0", ClampMax = "1000"))
+	double ThickenShells = 0.0;
+};
+
+
+
+/**
+ * Morphology tool -- dilate, contract, close, open operations on the input shape
+ */
+UCLASS()
+class MESHMODELINGTOOLSEXP_API UVoxelMorphologyMeshesTool : public UBaseVoxelTool
+{
+	GENERATED_BODY()
+
+public:
+
+	UVoxelMorphologyMeshesTool() {}
+
+protected:
+
+	virtual void SetupProperties() override;
+	virtual void SaveProperties() override;
+
+	virtual FString GetCreatedAssetName() const override;
+	virtual FText GetActionName() const override;
+
+	// IDynamicMeshOperatorFactory API
+	virtual TUniquePtr<UE::Geometry::FDynamicMeshOperator> MakeNewOperator() override;
+
+	UPROPERTY()
+	TObjectPtr<UVoxelMorphologyMeshesToolProperties> MorphologyProperties;
+
+};
+
+
+
+
+UCLASS()
+class MESHMODELINGTOOLSEXP_API UVoxelMorphologyMeshesToolBuilder : public UBaseCreateFromSelectedToolBuilder
+{
+	GENERATED_BODY()
+
+public:
+	virtual UBaseCreateFromSelectedTool* MakeNewToolInstance(UObject* Outer) const override
+	{
+		return NewObject<UVoxelMorphologyMeshesTool>(Outer);
+	}
+};
+
diff --git a/Engine/Plugins/Experimental/MeshModelingToolsetExp/Source/MeshModelingToolsExp/Public/WeldMeshEdgesTool.h b/Engine/Plugins/Experimental/MeshModelingToolsetExp/Source/MeshModelingToolsExp/Public/WeldMeshEdgesTool.h
new file mode 100644
--- /dev/null
+++ b/Engine/Plugins/Experimental/MeshModelingToolsetExp/Source/MeshModelingToolsExp/Public/WeldMeshEdgesTool.h
@@ -0,0 +1,107 @@
+// Copyright Epic Games, Inc. All Rights Reserved.
+
+#pragma once
+
+#include "CoreMinimal.h"
+#include "BaseTools/SingleSelectionMeshEditingTool.h"
+#include "PreviewMesh.h"
+#include "ModelingOperators.h"
+#include "MeshOpPreviewHelpers.h"
+#include "DynamicMesh/DynamicMesh3.h"
+#include "DynamicMesh/DynamicMeshAABBTree3.h"
+#include "WeldMeshEdgesTool.generated.h"
+
+// predeclarations
+struct FMeshDescription;
+class UMeshElementsVisualizer;
+class UWeldMeshEdgesTool;
+class FWeldMeshEdgesOp;
+
+/**
+ *
+ */
+UCLASS()
+class MESHMODELINGTOOLSEXP_API UWeldMeshEdgesToolBuilder : public USingleSelectionMeshEditingToolBuilder
+{
+	GENERATED_BODY()
+public:
+	virtual USingleSelectionMeshEditingTool* CreateNewTool(const FToolBuilderState& SceneState) const override;
+};
+
+
+
+UCLASS()
+class MESHMODELINGTOOLSEXP_API UWeldMeshEdgesToolProperties : public UInteractiveToolPropertySet
+{
+	GENERATED_BODY()
+
+public:
+	/** Edges are considered matching if both pairs of endpoint vertices are closer than this distance */
+	UPROPERTY(EditAnywhere, Category = Options, meta = (UIMin = "0.000001", UIMax = "0.01", ClampMin = "0.00000001", ClampMax = "1000.0"))
+	float Tolerance = FMathf::ZeroTolerance;
+
+	/** Only merge unambiguous pairs that have unique duplicate-edge matches */
+	UPROPERTY(EditAnywhere, Category = Options)
+	bool bOnlyUnique = false;
+};
+
+
+
+
+UCLASS()
+class MESHMODELINGTOOLSEXP_API UWeldMeshEdgesOperatorFactory : public UObject, public UE::Geometry::IDynamicMeshOperatorFactory
+{
+	GENERATED_BODY()
+
+public:
+	// IDynamicMeshOperatorFactory API
+	virtual TUniquePtr<UE::Geometry::FDynamicMeshOperator> MakeNewOperator() override;
+
+	UPROPERTY()
+	TObjectPtr<UWeldMeshEdgesTool> WeldMeshEdgesTool;
+};
+
+
+/**
+ * Mesh Weld Edges Tool
+ */
+UCLASS()
+class MESHMODELINGTOOLSEXP_API UWeldMeshEdgesTool : public USingleSelectionMeshEditingTool
+{
+	GENERATED_BODY()
+public:
+	UWeldMeshEdgesTool();
+
+	virtual void Setup() override;
+	virtual void Shutdown(EToolShutdownType ShutdownType) override;
+
+	virtual void OnTick(float DeltaTime) override;
+
+	virtual bool HasCancel() const override { return true; }
+	virtual bool HasAccept() const override { return true; }
+	virtual bool CanAccept() const override;
+
+	// update parameters in FWeldMeshEdgesOp based on current Settings
+	void UpdateOpParameters(FWeldMeshEdgesOp& Op) const;
+
+protected:
+	UPROPERTY()
+	TObjectPtr<UWeldMeshEdgesToolProperties> Settings;
+
+	UPROPERTY()
+	TObjectPtr<UMeshOpPreviewWithBackgroundCompute> PreviewCompute = nullptr;
+
+	UPROPERTY()
+	TObjectPtr<UMeshElementsVisualizer> MeshElementsDisplay = nullptr;
+
+	UPROPERTY()
+	TObjectPtr<UWeldMeshEdgesOperatorFactory> OperatorFactory;
+
+protected:
+
+	TSharedPtr<UE::Geometry::FDynamicMesh3, ESPMode::ThreadSafe> SourceMesh;
+
+
+
+
+};
diff --git a/Engine/Plugins/Experimental/MeshModelingToolsetExp/Source/ModelingComponents/ModelingComponents.Build.cs b/Engine/Plugins/Experimental/MeshModelingToolsetExp/Source/ModelingComponents/ModelingComponents.Build.cs
new file mode 100644
--- /dev/null
+++ b/Engine/Plugins/Experimental/MeshModelingToolsetExp/Source/ModelingComponents/ModelingComponents.Build.cs
@@ -0,0 +1,70 @@
+// Copyright Epic Games, Inc. All Rights Reserved.
+
+using UnrealBuildTool;
+
+public class ModelingComponents : ModuleRules
+{
+	public ModelingComponents(ReadOnlyTargetRules Target) : base(Target)
+	{
+		PCHUsage = ModuleRules.PCHUsageMode.UseExplicitOrSharedPCHs;
+		
+		PublicIncludePaths.AddRange(
+			new string[] {
+				// ... add public include paths required here ...
+			}
+			);
+				
+		
+		PrivateIncludePaths.AddRange(
+			new string[] {
+				// ... add other private include paths required here ...
+			}
+			);
+			
+		
+		PublicDependencyModuleNames.AddRange(
+			new string[]
+			{
+				"Core",
+				"PhysicsCore",
+				"InteractiveToolsFramework",
+                "MeshDescription",
+				"StaticMeshDescription",
+				"GeometryCore",
+				"GeometryFramework",
+				"GeometryAlgorithms",
+				"DynamicMesh",
+				"MeshConversion",
+				"ModelingOperators"
+				// ... add other public dependencies that you statically link with here ...
+			}
+            );
+			
+		
+		PrivateDependencyModuleNames.AddRange(
+			new string[]
+			{
+				"CoreUObject",
+				"Engine",
+				"InputCore",
+                "RenderCore",
+                "RHI",
+				"ImageWriteQueue"
+
+				//"MeshUtilities",    // temp for saving mesh asset
+				//"UnrealEd"
+				//"Slate",
+				//"SlateCore",
+				// ... add private dependencies that you statically link with here ...	
+			}
+			);
+		
+		
+		DynamicallyLoadedModuleNames.AddRange(
+			new string[]
+			{
+				// ... add any modules that your module loads dynamically here ...
+			}
+			);
+	}
+}
diff --git a/Engine/Plugins/Experimental/MeshModelingToolsetExp/Source/ModelingComponents/Private/BaseTools/BaseCreateFromSelectedTool.cpp b/Engine/Plugins/Experimental/MeshModelingToolsetExp/Source/ModelingComponents/Private/BaseTools/BaseCreateFromSelectedTool.cpp
new file mode 100644
--- /dev/null
+++ b/Engine/Plugins/Experimental/MeshModelingToolsetExp/Source/ModelingComponents/Private/BaseTools/BaseCreateFromSelectedTool.cpp
@@ -0,0 +1,388 @@
+// Copyright Epic Games, Inc. All Rights Reserved.
+
+#include "BaseTools/BaseCreateFromSelectedTool.h"
+#include "InteractiveToolManager.h"
+#include "ToolBuilderUtil.h"
+#include "ToolSetupUtil.h"
+#include "Components/DynamicMeshComponent.h"
+#include "BaseGizmos/TransformGizmoUtil.h"
+#include "Async/Async.h"
+
+#include "DynamicMesh/MeshNormals.h"
+#include "DynamicMesh/MeshTransforms.h"
+#include "DynamicMeshToMeshDescription.h"
+#include "MeshDescriptionToDynamicMesh.h"
+
+#include "ModelingObjectsCreationAPI.h"
+#include "Selection/ToolSelectionUtil.h"
+
+#include "TargetInterfaces/MeshDescriptionCommitter.h"
+#include "TargetInterfaces/MeshDescriptionProvider.h"
+#include "TargetInterfaces/PrimitiveComponentBackedTarget.h"
+#include "TargetInterfaces/MaterialProvider.h"
+#include "TargetInterfaces/AssetBackedTarget.h"
+#include "ToolTargetManager.h"
+#include "ModelingToolTargetUtil.h"
+
+#define LOCTEXT_NAMESPACE "UBaseCreateFromSelectedTool"
+
+using namespace UE::Geometry;
+
+/*
+ * ToolBuilder
+ */
+
+const FToolTargetTypeRequirements& UBaseCreateFromSelectedToolBuilder::GetTargetRequirements() const
+{
+	static FToolTargetTypeRequirements TypeRequirements({
+		UMeshDescriptionCommitter::StaticClass(),
+		UMeshDescriptionProvider::StaticClass(),
+		UPrimitiveComponentBackedTarget::StaticClass(),
+		UMaterialProvider::StaticClass()
+		});
+	return TypeRequirements;
+}
+
+bool UBaseCreateFromSelectedToolBuilder::CanBuildTool(const FToolBuilderState& SceneState) const
+{
+	int32 ComponentCount = SceneState.TargetManager->CountSelectedAndTargetable(SceneState, GetTargetRequirements());
+	return ComponentCount >= MinComponentsSupported() && (!MaxComponentsSupported().IsSet() || ComponentCount <= MaxComponentsSupported().GetValue());
+}
+
+
+UInteractiveTool* UBaseCreateFromSelectedToolBuilder::BuildTool(const FToolBuilderState& SceneState) const
+{
+	UBaseCreateFromSelectedTool* NewTool = MakeNewToolInstance(SceneState.ToolManager);
+
+	TArray<TObjectPtr<UToolTarget>> Targets = SceneState.TargetManager->BuildAllSelectedTargetable(SceneState, GetTargetRequirements());
+	NewTool->SetTargets(MoveTemp(Targets));
+	NewTool->SetWorld(SceneState.World);
+
+	return NewTool;
+}
+
+
+
+
+/*
+ * Tool
+ */
+
+void UBaseCreateFromSelectedTool::Setup()
+{
+	UInteractiveTool::Setup();
+
+	// hide input StaticMeshComponents
+	for (int32 ComponentIdx = 0; ComponentIdx < Targets.Num(); ComponentIdx++)
+	{
+		TargetComponentInterface(ComponentIdx)->SetOwnerVisibility(false);
+	}
+
+	// initialize our properties
+
+	SetupProperties();
+	TransformProperties = NewObject<UTransformInputsToolProperties>(this);
+	TransformProperties->RestoreProperties(this);
+	AddToolPropertySource(TransformProperties);
+	
+	OutputTypeProperties = NewObject<UCreateMeshObjectTypeProperties>(this);
+	OutputTypeProperties->InitializeDefaultWithAuto();
+	OutputTypeProperties->OutputType = UCreateMeshObjectTypeProperties::AutoIdentifier;
+	OutputTypeProperties->RestoreProperties(this, TEXT("OutputTypeFromInputTool"));
+	OutputTypeProperties->WatchProperty(OutputTypeProperties->OutputType, [this](FString) { OutputTypeProperties->UpdatePropertyVisibility(); });
+	AddToolPropertySource(OutputTypeProperties);
+
+	HandleSourcesProperties = NewObject<UBaseCreateFromSelectedHandleSourceProperties>(this);
+	HandleSourcesProperties->RestoreProperties(this);
+	AddToolPropertySource(HandleSourcesProperties);
+
+	Preview = NewObject<UMeshOpPreviewWithBackgroundCompute>(this);
+	Preview->Setup(this->TargetWorld, this);
+
+	SetPreviewCallbacks();
+	Preview->OnMeshUpdated.AddLambda(
+		[this](const UMeshOpPreviewWithBackgroundCompute* UpdatedPreview)
+		{
+			UpdateAcceptWarnings(UpdatedPreview->HaveEmptyResult() ? EAcceptWarning::EmptyForbidden : EAcceptWarning::NoWarning);
+		}
+	);
+
+	SetTransformGizmos();
+
+	ConvertInputsAndSetPreviewMaterials(true);
+
+	// output name fields
+	HandleSourcesProperties->OutputName = PrefixWithSourceNameIfSingleSelection(GetCreatedAssetName());
+	HandleSourcesProperties->WatchProperty(HandleSourcesProperties->WriteOutputTo, [&](EBaseCreateFromSelectedTargetType NewType)
+	{
+		SetToolPropertySourceEnabled(OutputTypeProperties, (NewType == EBaseCreateFromSelectedTargetType::NewAsset));
+		if (NewType == EBaseCreateFromSelectedTargetType::NewAsset)
+		{
+			HandleSourcesProperties->OutputAsset = TEXT("");
+			UpdateGizmoVisibility();
+		}
+		else
+		{
+			int32 Index = (HandleSourcesProperties->WriteOutputTo == EBaseCreateFromSelectedTargetType::FirstInputAsset) ? 0 : Targets.Num() - 1;
+			HandleSourcesProperties->OutputAsset = UE::Modeling::GetComponentAssetBaseName(TargetComponentInterface(Index)->GetOwnerComponent(), false);
+
+			// Reset the hidden gizmo to its initial position
+			FTransform ComponentTransform = TargetComponentInterface(Index)->GetWorldTransform();
+			TransformGizmos[Index]->SetNewGizmoTransform(ComponentTransform, true);
+			UpdateGizmoVisibility();
+		}
+	});
+
+	Preview->InvalidateResult();
+}
+
+
+int32 UBaseCreateFromSelectedTool::GetHiddenGizmoIndex() const
+{
+	if (HandleSourcesProperties->WriteOutputTo == EBaseCreateFromSelectedTargetType::NewAsset)
+	{
+		return -1;
+	}
+	else
+	{
+		return (HandleSourcesProperties->WriteOutputTo == EBaseCreateFromSelectedTargetType::FirstInputAsset) ? 0 : Targets.Num() - 1;
+	}
+}
+
+
+void UBaseCreateFromSelectedTool::SetWorld(UWorld* World)
+{
+	this->TargetWorld = World;
+}
+
+
+void UBaseCreateFromSelectedTool::OnTick(float DeltaTime)
+{
+	Preview->Tick(DeltaTime);
+}
+
+
+
+void UBaseCreateFromSelectedTool::UpdateGizmoVisibility()
+{
+	for (int32 GizmoIndex = 0; GizmoIndex < TransformGizmos.Num(); GizmoIndex++)
+	{
+		UTransformGizmo* Gizmo = TransformGizmos[GizmoIndex];
+		Gizmo->SetVisibility(TransformProperties->bShowTransformUI && GizmoIndex != GetHiddenGizmoIndex());
+	}
+}
+
+
+
+void UBaseCreateFromSelectedTool::SetTransformGizmos()
+{
+	UInteractiveGizmoManager* GizmoManager = GetToolManager()->GetPairedGizmoManager();
+
+	for (int ComponentIdx = 0; ComponentIdx < Targets.Num(); ComponentIdx++)
+	{
+		UTransformProxy* Proxy = TransformProxies.Add_GetRef(NewObject<UTransformProxy>(this));
+		UTransformGizmo* Gizmo = TransformGizmos.Add_GetRef(UE::TransformGizmoUtil::Create3AxisTransformGizmo(GizmoManager, this));
+		Proxy->SetTransform(TargetComponentInterface(ComponentIdx)->GetWorldTransform());
+		Gizmo->SetActiveTarget(Proxy, GetToolManager());
+		Proxy->OnTransformChanged.AddUObject(this, &UBaseCreateFromSelectedTool::TransformChanged);
+	}
+	UpdateGizmoVisibility();
+}
+
+
+void UBaseCreateFromSelectedTool::TransformChanged(UTransformProxy* Proxy, FTransform Transform)
+{
+	Preview->InvalidateResult();
+}
+
+
+FText UBaseCreateFromSelectedTool::GetActionName() const
+{
+	return LOCTEXT("BaseCreateFromSelectedTool", "Generated Mesh");
+}
+
+
+TArray<UMaterialInterface*> UBaseCreateFromSelectedTool::GetOutputMaterials() const
+{
+	return Preview->StandardMaterials;
+}
+
+
+void UBaseCreateFromSelectedTool::GenerateAsset(const FDynamicMeshOpResult& OpResult)
+{
+	if (OpResult.Mesh.Get() == nullptr) return;
+
+	FTransform3d NewTransform;
+	if (Targets.Num() == 1) // in the single-selection case, shove the result back into the original component space
+	{
+		FTransform3d ToSourceComponentSpace = (FTransform3d)TargetComponentInterface(0)->GetWorldTransform().Inverse();
+		MeshTransforms::ApplyTransform(*OpResult.Mesh, ToSourceComponentSpace);
+		NewTransform = (FTransform3d)TargetComponentInterface(0)->GetWorldTransform();
+	}
+	else // in the multi-selection case, center the pivot for the combined result
+	{
+		FVector3d Center = OpResult.Mesh->GetCachedBounds().Center();
+		double Rescale = OpResult.Transform.GetScale().X;
+		FTransform3d LocalTransform(-Center * Rescale);
+		LocalTransform.SetScale(FVector3d(Rescale, Rescale, Rescale));
+		MeshTransforms::ApplyTransform(*OpResult.Mesh, LocalTransform);
+		NewTransform = OpResult.Transform;
+		NewTransform.SetScale(FVector3d::One());
+		NewTransform.SetTranslation(NewTransform.GetTranslation() + NewTransform.TransformVector(Center * Rescale));
+	}
+
+	// max len explicitly enforced here, would ideally notify user
+	FString UseBaseName = HandleSourcesProperties->OutputName.Left(250);
+	if (UseBaseName.IsEmpty())
+	{
+		UseBaseName = PrefixWithSourceNameIfSingleSelection(GetCreatedAssetName());
+	}
+
+	FCreateMeshObjectParams NewMeshObjectParams;
+	NewMeshObjectParams.TargetWorld = TargetWorld;
+	NewMeshObjectParams.Transform = (FTransform)NewTransform;
+	NewMeshObjectParams.BaseName = UseBaseName;
+	NewMeshObjectParams.Materials = GetOutputMaterials();
+	NewMeshObjectParams.SetMesh(OpResult.Mesh.Get());
+	if (OutputTypeProperties->OutputType == UCreateMeshObjectTypeProperties::AutoIdentifier)
+	{
+		UE::ToolTarget::ConfigureCreateMeshObjectParams(Targets[0], NewMeshObjectParams);
+	}
+	else
+	{
+		OutputTypeProperties->ConfigureCreateMeshObjectParams(NewMeshObjectParams);
+	}
+	FCreateMeshObjectResult Result = UE::Modeling::CreateMeshObject(GetToolManager(), MoveTemp(NewMeshObjectParams));
+	if (Result.IsOK() && Result.NewActor != nullptr)
+	{
+		ToolSelectionUtil::SetNewActorSelection(GetToolManager(), Result.NewActor);
+	}
+}
+
+
+
+void UBaseCreateFromSelectedTool::UpdateAsset(const FDynamicMeshOpResult& Result, UToolTarget* UpdateTarget)
+{
+	check(Result.Mesh.Get() != nullptr);
+
+	IPrimitiveComponentBackedTarget* TargetComponent = Cast<IPrimitiveComponentBackedTarget>(UpdateTarget);
+	IMeshDescriptionCommitter* TargetMeshCommitter = Cast<IMeshDescriptionCommitter>(UpdateTarget);
+	IMaterialProvider* TargetMaterial = Cast<IMaterialProvider>(UpdateTarget);
+
+	FTransform3d TargetToWorld = (FTransform3d)TargetComponent->GetWorldTransform();
+
+	FTransform3d ResultTransform = Result.Transform;
+	MeshTransforms::ApplyTransform(*Result.Mesh, ResultTransform);
+	MeshTransforms::ApplyTransformInverse(*Result.Mesh, TargetToWorld);
+
+	TargetMeshCommitter->CommitMeshDescription([&](const IMeshDescriptionCommitter::FCommitterParams& CommitParams)
+	{
+		FDynamicMeshToMeshDescription Converter;
+		Converter.Convert(Result.Mesh.Get(), *CommitParams.MeshDescriptionOut);
+	});
+
+	FComponentMaterialSet MaterialSet;
+	MaterialSet.Materials = GetOutputMaterials();
+	TargetMaterial->CommitMaterialSetUpdate(MaterialSet, true);
+}
+
+
+FString UBaseCreateFromSelectedTool::PrefixWithSourceNameIfSingleSelection(const FString& AssetName) const
+{
+	if (Targets.Num() == 1)
+	{
+		FString CurName = UE::Modeling::GetComponentAssetBaseName(TargetComponentInterface(0)->GetOwnerComponent());
+		return FString::Printf(TEXT("%s_%s"), *CurName, *AssetName);
+	}
+	else
+	{
+		return AssetName;
+	}
+}
+
+
+void UBaseCreateFromSelectedTool::OnPropertyModified(UObject* PropertySet, FProperty* Property)
+{
+	if (Property && (Property->GetFName() == GET_MEMBER_NAME_CHECKED(UTransformInputsToolProperties, bShowTransformUI)))
+	{
+		UpdateGizmoVisibility();
+	}
+	else if (Property && 
+		(  PropertySet == HandleSourcesProperties
+		))
+	{
+		// nothing
+	}
+	else
+	{
+		Preview->InvalidateResult();
+	}
+}
+
+
+void UBaseCreateFromSelectedTool::Shutdown(EToolShutdownType ShutdownType)
+{
+	SaveProperties();
+	HandleSourcesProperties->SaveProperties(this);
+	OutputTypeProperties->SaveProperties(this, TEXT("OutputTypeFromInputTool"));
+	TransformProperties->SaveProperties(this);
+
+	FDynamicMeshOpResult Result = Preview->Shutdown();
+	// Restore (unhide) the source meshes
+	for (int32 ComponentIdx = 0; ComponentIdx < Targets.Num(); ComponentIdx++)
+	{
+		TargetComponentInterface(ComponentIdx)->SetOwnerVisibility(true);
+	}
+
+	if (ShutdownType == EToolShutdownType::Accept)
+	{
+		GetToolManager()->BeginUndoTransaction(GetActionName());
+
+		// Generate the result
+		AActor* KeepActor = nullptr;
+		if (HandleSourcesProperties->WriteOutputTo == EBaseCreateFromSelectedTargetType::NewAsset)
+		{
+			GenerateAsset(Result);
+		}
+		else
+		{
+			int32 TargetIndex = (HandleSourcesProperties->WriteOutputTo == EBaseCreateFromSelectedTargetType::FirstInputAsset) ? 0 : (Targets.Num() - 1);
+			KeepActor = TargetComponentInterface(TargetIndex)->GetOwnerActor();
+
+			UpdateAsset(Result, Targets[TargetIndex]);
+		}
+
+		TArray<AActor*> Actors;
+		for (int32 ComponentIdx = 0; ComponentIdx < Targets.Num(); ComponentIdx++)
+		{
+			AActor* Actor = TargetComponentInterface(ComponentIdx)->GetOwnerActor();
+			if (Actor != KeepActor)
+			{
+				Actors.Add(Actor);
+			}
+		}
+		HandleSourcesProperties->ApplyMethod(Actors, GetToolManager());
+
+		if (KeepActor != nullptr)
+		{
+			// select the actor we kept
+			ToolSelectionUtil::SetNewActorSelection(GetToolManager(), KeepActor);
+		}
+
+		GetToolManager()->EndUndoTransaction();
+	}
+
+	UInteractiveGizmoManager* GizmoManager = GetToolManager()->GetPairedGizmoManager();
+	GizmoManager->DestroyAllGizmosByOwner(this);
+}
+
+
+bool UBaseCreateFromSelectedTool::CanAccept() const
+{
+	return Super::CanAccept() && Preview->HaveValidNonEmptyResult();
+}
+
+
+
+
+#undef LOCTEXT_NAMESPACE
diff --git a/Engine/Plugins/Experimental/MeshModelingToolsetExp/Source/ModelingComponents/Private/ModelingComponentsModule.cpp b/Engine/Plugins/Experimental/MeshModelingToolsetExp/Source/ModelingComponents/Private/ModelingComponentsModule.cpp
new file mode 100644
--- /dev/null
+++ b/Engine/Plugins/Experimental/MeshModelingToolsetExp/Source/ModelingComponents/Private/ModelingComponentsModule.cpp
@@ -0,0 +1,20 @@
+// Copyright Epic Games, Inc. All Rights Reserved.
+
+#include "ModelingComponentsModule.h"
+
+#define LOCTEXT_NAMESPACE "FModelingComponentsModule"
+
+void FModelingComponentsModule::StartupModule()
+{
+	// This code will execute after your module is loaded into memory; the exact timing is specified in the .uplugin file per-module
+}
+
+void FModelingComponentsModule::ShutdownModule()
+{
+	// This function may be called during shutdown to clean up your module.  For modules that support dynamic reloading,
+	// we call this function before unloading the module.
+}
+
+#undef LOCTEXT_NAMESPACE
+	
+IMPLEMENT_MODULE(FModelingComponentsModule, ModelingComponents)
\ No newline at end of file
