MeshModelingToolset: Split plugin into non-experimental and experimental

diff --git a/Engine/Plugins/Experimental/MeshModelingToolsetExp/Source/ModelingComponents/Private/Selection/GroupTopologySelector.cpp b/Engine/Plugins/Experimental/MeshModelingToolsetExp/Source/ModelingComponents/Private/Selection/GroupTopologySelector.cpp
new file mode 100644
--- /dev/null
+++ b/Engine/Plugins/Experimental/MeshModelingToolsetExp/Source/ModelingComponents/Private/Selection/GroupTopologySelector.cpp
@@ -0,0 +1,900 @@
+// Copyright Epic Games, Inc. All Rights Reserved.
+
+
+#include "Selection/GroupTopologySelector.h"
+#include "Mechanics/RectangleMarqueeMechanic.h"
+#include "MeshQueries.h"
+#include "ToolDataVisualizer.h"
+#include "ToolSceneQueriesUtil.h"
+
+#include "ExplicitUseGeometryMathTypes.h"		// using UE::Geometry::(math types)
+using namespace UE::Geometry;
+
+// Local utility function forward declarations
+bool IsOccluded(const FGeometrySet3::FNearest& ClosestElement, const FVector3d& ViewOrigin, const FDynamicMeshAABBTree3* Spatial);
+bool IsOccluded(const FVector3d& Point, const FVector3d& ViewOrigin, const FDynamicMeshAABBTree3* Spatial);
+void AddNewEdgeLoopEdgesFromCorner(const FGroupTopology& Topology, int32 EdgeID, int32 CornerID, TSet<int32>& EdgeSet);
+bool GetNextEdgeLoopEdge(const FGroupTopology& Topology, int32 IncomingEdgeID, int32 CornerID, int32& NextEdgeIDOut);
+void AddNewEdgeRingEdges(const FGroupTopology& Topology, int32 StartEdgeID, int32 ForwardGroupID, TSet<int32>& EdgeSet);
+bool GetQuadOppositeEdge(const FGroupTopology& Topology, int32 EdgeIDIn, int32 GroupID, int32& OppositeEdgeIDOut);
+
+FGroupTopologySelector::FGroupTopologySelector()
+{
+	// initialize to sane values
+	PointsWithinToleranceTest =
+		[](const FVector3d& A, const FVector3d& B, double TolScale) { return Distance(A, B) < (TolScale*1.0); };
+	GetSpatial =
+		[]() { return nullptr; };
+}
+
+
+
+void FGroupTopologySelector::Initialize(const FDynamicMesh3* MeshIn, const FGroupTopology* TopologyIn)
+{
+	Mesh = MeshIn;
+	Topology = TopologyIn;
+	bGeometryInitialized = false;
+	bGeometryUpToDate = false;
+}
+
+
+void FGroupTopologySelector::Invalidate(bool bTopologyDeformed, bool bTopologyModified)
+{
+	if (bTopologyDeformed)
+	{
+		bGeometryUpToDate = false;
+	}
+	if (bTopologyModified)
+	{
+		bGeometryUpToDate = bGeometryInitialized = false;
+	}
+}
+
+
+const FGeometrySet3& FGroupTopologySelector::GetGeometrySet()
+{
+	if (bGeometryInitialized == false)
+	{
+		GeometrySet.Reset();
+		int NumCorners = Topology->Corners.Num();
+		for (int CornerID = 0; CornerID < NumCorners; ++CornerID)
+		{
+			FVector3d Position = Mesh->GetVertex(Topology->Corners[CornerID].VertexID);
+			GeometrySet.AddPoint(CornerID, Position);
+		}
+		int NumEdges = Topology->Edges.Num();
+		for (int EdgeID = 0; EdgeID < NumEdges; ++EdgeID)
+		{
+			FPolyline3d Polyline;
+			Topology->Edges[EdgeID].Span.GetPolyline(Polyline);
+			GeometrySet.AddCurve(EdgeID, Polyline);
+		}
+
+		bGeometryInitialized = true;
+		bGeometryUpToDate = true;
+	}
+
+	if (bGeometryUpToDate == false)
+	{
+		int NumCorners = Topology->Corners.Num();
+		for (int CornerID = 0; CornerID < NumCorners; ++CornerID)
+		{
+			FVector3d Position = Mesh->GetVertex(Topology->Corners[CornerID].VertexID);
+			GeometrySet.UpdatePoint(CornerID, Position);
+		}
+		int NumEdges = Topology->Edges.Num();
+		for (int EdgeID = 0; EdgeID < NumEdges; ++EdgeID)
+		{
+			FPolyline3d Polyline;
+			Topology->Edges[EdgeID].Span.GetPolyline(Polyline);
+			GeometrySet.UpdateCurve(EdgeID, Polyline);
+		}
+		bGeometryUpToDate = true;
+	}
+
+	return GeometrySet;
+}
+
+bool FGroupTopologySelector::FindSelectedElement(const FSelectionSettings& Settings, const FRay3d& Ray,
+	FGroupTopologySelection& ResultOut, FVector3d& SelectedPositionOut, FVector3d& SelectedNormalOut, int32* EdgeSegmentIdOut)
+{
+	// These get used for finding intersections with triangles and corners/edges, repectively.
+	FDynamicMeshAABBTree3* Spatial = GetSpatial();
+	const FGeometrySet3& TopoSpatial = GetGeometrySet();
+
+	// We start by intersecting with the mesh triangles because even when selecting corners or edges, we set
+	// the normal based on the true triangle that we hit. If we end up with a simple face selection, we will
+	// end up using this result.
+	double RayParameter = -1;
+	int HitTriangleID = IndexConstants::InvalidID;
+	FVector3d TriangleHitPos;
+	bool bActuallyHitSurface = (Spatial != nullptr) ? Spatial->FindNearestHitTriangle(Ray, RayParameter, HitTriangleID) : false;
+	if (bActuallyHitSurface)
+	{
+		TriangleHitPos = Ray.PointAt(RayParameter);
+		SelectedNormalOut = Mesh->GetTriNormal(HitTriangleID);
+	}
+	else
+	{
+		SelectedNormalOut = FVector3d::UnitZ();
+	}
+	bool bHaveFaceHit = (bActuallyHitSurface && Settings.bEnableFaceHits);
+	
+	// Deal with corner hits first (and edges that project to a corner)
+	FGroupTopologySelection CornerResults;
+	FVector3d CornerPosition;
+	int32 CornerSegmentEdgeID = 0;
+	bool bHaveCornerHit = false;
+	double CornerAngle = TNumericLimits<double>::Max();
+	if (Settings.bEnableCornerHits || (Settings.bEnableEdgeHits && Settings.bPreferProjectedElement))
+	{
+		if (DoCornerBasedSelection(Settings, Ray, Spatial, TopoSpatial, CornerResults, CornerPosition, &CornerSegmentEdgeID))
+		{
+			bHaveCornerHit = true;
+			CornerAngle = UE::Geometry::AngleD(Ray.Direction, UE::Geometry::Normalized(CornerPosition - Ray.Origin) );
+		}
+	}
+
+	// If corner selection didn't yield results, try edge selection
+	FGroupTopologySelection EdgeResults;
+	FVector3d EdgePosition;
+	int32 EdgeSegmentEdgeID = 0;
+	bool bHaveEdgeHit = false;
+	double EdgeAngle = TNumericLimits<double>::Max();
+	if (Settings.bEnableEdgeHits || (Settings.bEnableFaceHits && Settings.bPreferProjectedElement))
+	{
+		if (DoEdgeBasedSelection(Settings, Ray, Spatial, TopoSpatial, EdgeResults, EdgePosition, &EdgeSegmentEdgeID))
+		{
+			bHaveEdgeHit = true;
+			EdgeAngle = UE::Geometry::AngleD(Ray.Direction, UE::Geometry::Normalized(EdgePosition - Ray.Origin) );
+		}
+	}
+
+	// if we have both corner and edge, want to keep the one we are closer to
+	if (bHaveCornerHit && bHaveEdgeHit)
+	{
+		if (PointsWithinToleranceTest(CornerPosition, Ray.NearestPoint(CornerPosition), 0.75))
+		{
+			bHaveEdgeHit = false;
+		}
+		else
+		{
+			bHaveCornerHit = false;
+		}
+	}
+
+	// if we have a corner or edge hit, and a face hit, pick face unless we are really close to corner/edge
+	if ((bHaveCornerHit || bHaveEdgeHit) && bHaveFaceHit)
+	{
+		FVector3d TestPos = (bHaveCornerHit) ? CornerPosition : EdgePosition;
+		if (!PointsWithinToleranceTest(TestPos, Ray.NearestPoint(TestPos), 0.15))
+		{
+			bHaveEdgeHit = bHaveCornerHit = false;
+		}
+	}
+
+
+	if (bHaveCornerHit)
+	{
+		ResultOut = CornerResults;
+		SelectedPositionOut = CornerPosition;
+		if (EdgeSegmentIdOut != nullptr)
+		{
+			*EdgeSegmentIdOut = CornerSegmentEdgeID;
+		}
+		return true;
+	}
+	else if (bHaveEdgeHit)
+	{
+		ResultOut = EdgeResults;
+		SelectedPositionOut = EdgePosition;
+		if (EdgeSegmentIdOut != nullptr)
+		{
+			*EdgeSegmentIdOut = EdgeSegmentEdgeID;
+		}
+		return true;
+	}
+	else if (bHaveFaceHit)
+	{
+		// If we still haven't found a selection, go ahead and select the face that we found earlier
+		ResultOut.SelectedGroupIDs.Add(Topology->GetGroupID(HitTriangleID));
+		SelectedPositionOut = TriangleHitPos;
+		return true;
+	}
+
+
+	return false;
+}
+
+bool FGroupTopologySelector::DoCornerBasedSelection(const FSelectionSettings& Settings,
+	const FRay3d& Ray, FDynamicMeshAABBTree3* Spatial, const FGeometrySet3& TopoSpatial,
+	FGroupTopologySelection& ResultOut, FVector3d& SelectedPositionOut, int32 *EdgeSegmentIdOut) const
+{
+	// These will store our results, depending on whether we select all along the ray or not.
+	FGeometrySet3::FNearest SingleElement;
+	// ElementsWithinTolerance gives all nearby elements returned by the topology selector, whereas DownRayElements filters these to just
+	// the closest element and those that project directly onto it (and it only stores the IDs, since we don't care about anything else).
+	TArray<FGeometrySet3::FNearest> ElementsWithinTolerance;
+	TArray<int32> DownRayElements;
+
+	auto LocalTolTest = [this](const FVector3d& A, const FVector3d& B) { return PointsWithinToleranceTest(A, B, 1.0); };
+
+	// Start by getting the closest element
+	const FGeometrySet3::FNearest* ClosestElement = nullptr;
+	if (!Settings.bSelectDownRay)
+	{
+		if (TopoSpatial.FindNearestPointToRay(Ray, SingleElement, LocalTolTest))
+		{
+			ClosestElement = &SingleElement;
+		}
+	}
+	else
+	{
+		// We're collecting all corners within tolerance, but we still need the closest element
+		if (TopoSpatial.CollectPointsNearRay(Ray, ElementsWithinTolerance, LocalTolTest))
+		{
+			double MinRayT = TNumericLimits<double>::Max();
+			for (const FGeometrySet3::FNearest& Element : ElementsWithinTolerance)
+			{
+				if (Element.RayParam < MinRayT)
+				{
+					MinRayT = Element.RayParam;
+					ClosestElement = &Element;
+				}
+			}
+		}//end if found corners
+	}//end if selecting down ray
+
+	// Bail early if we haven't found a corner
+	if (ClosestElement == nullptr)
+	{
+		return false;
+	}
+
+	// Also bail if the closest element is not visible.
+	if (!Settings.bIgnoreOcclusion && IsOccluded(*ClosestElement, Ray.Origin, Spatial))
+	{
+		return false;
+	}
+
+	// The closest point is already found
+	SelectedPositionOut = ClosestElement->NearestGeoPoint;
+
+	// If we have other corners, we actually need to filter them to only those that lie in line with the closest element. Note that
+	// this would be done differently depending on whether we're dealing with an orthographic or perspective view: in an
+	// orthographic view, we need to see that they lie on a ray with view ray direction and closest corner origin, whereas in
+	// perspective, they would need to lie on a ray from camer through closest corner (which will differ due to tolerance).
+	// Because the "select down ray" behavior is only useful in orthographic viewports in the first place, we do it that way.
+	if (Settings.bSelectDownRay)
+	{
+		DownRayElements.Add(ClosestElement->ID);
+		for (const FGeometrySet3::FNearest& Element : ElementsWithinTolerance)
+		{
+			if (ClosestElement == &Element)
+			{
+				continue; // Already added
+			}
+
+			// Make sure that closest corner to current element is parallel with view ray
+			FVector3d ClosestTowardElement = UE::Geometry::Normalized(Element.NearestGeoPoint - ClosestElement->NearestGeoPoint);
+			// There would usually be one more abs() in here, but we know that other elements are down ray direction
+			if (abs(ClosestTowardElement.Dot(Ray.Direction) - 1.0) < KINDA_SMALL_NUMBER)
+			{
+				DownRayElements.Add(Element.ID);
+			}
+		}//end assembling aligned corners
+	}
+
+	// Try to select edges that project to corners.
+	if (Settings.bPreferProjectedElement && Settings.bEnableEdgeHits)
+	{
+		TSet<int32> AddedTopologyEdges;
+
+		// See if the closest vertex has an attached edge that is colinear with the view ray. Due to the fact that
+		// topology "edges" are actually polylines, we could actually have more than one even for the closest corner
+		// (if the "edge" towards us curves away). We'll only grab an edge down our view ray.
+		int32 ClosestVid = Topology->GetCornerVertexID(ClosestElement->ID);
+		for (int32 Eid : Mesh->VtxEdgesItr(ClosestVid))
+		{
+			FDynamicMesh3::FEdge Edge = Mesh->GetEdge(Eid);
+			int32 OtherVid = (Edge.Vert.A == ClosestVid) ? Edge.Vert.B : Edge.Vert.A;
+			FVector3d EdgeVector = UE::Geometry::Normalized(Mesh->GetVertex(OtherVid) - Mesh->GetVertex(ClosestVid));
+			if (abs(EdgeVector.Dot(Ray.Direction) - 1.0) < KINDA_SMALL_NUMBER)
+			{
+				int TopologyEdgeIndex = Topology->FindGroupEdgeID(Eid);
+				if (TopologyEdgeIndex >= 0)
+				{
+					ResultOut.SelectedEdgeIDs.Add(TopologyEdgeIndex);
+					AddedTopologyEdges.Add(TopologyEdgeIndex);
+
+					if (EdgeSegmentIdOut)
+					{
+						Topology->GetGroupEdgeEdges(TopologyEdgeIndex).Find(Eid, *EdgeSegmentIdOut);
+					}
+
+					break;
+				}
+			}
+		}
+
+		// If relevant, get all the other colinear edges
+		if (Settings.bSelectDownRay && AddedTopologyEdges.Num() > 0)
+		{
+			for (int i = 1; i < DownRayElements.Num(); ++i) // skip 0 because it is closest and done
+			{
+				// Look though any attached edges.
+				for (int32 Eid : Mesh->VtxEdgesItr(Topology->GetCornerVertexID(DownRayElements[i])))
+				{
+					FDynamicMesh3::FEdge Edge = Mesh->GetEdge(Eid);
+					FVector3d EdgeVector = UE::Geometry::Normalized(Mesh->GetVertex(Edge.Vert.A) - Mesh->GetVertex(Edge.Vert.B));
+
+					// Compare absolute value of dot product to 1. We already made sure that one of the vertices is in line
+					// with the closest vertex earlier.
+					if (abs(abs(EdgeVector.Dot(Ray.Direction)) - 1.0) < KINDA_SMALL_NUMBER)
+					{
+						int TopologyEdgeIndex = Topology->FindGroupEdgeID(Eid);
+						if (TopologyEdgeIndex >= 0 && !AddedTopologyEdges.Contains(TopologyEdgeIndex))
+						{
+							ResultOut.SelectedEdgeIDs.Add(TopologyEdgeIndex);
+							AddedTopologyEdges.Add(TopologyEdgeIndex);
+							// Don't break here because we may have parallel edges in both directions, since we aren't
+							// going through the vertices in a particular order.
+						}
+					}
+				}//end checking edges
+			}//end going through corners
+		}
+
+		// If we found edges, we're done
+		if (AddedTopologyEdges.Num() > 0)
+		{
+			return true;
+		}
+	}//end selecting projected edges
+
+	// If getting projected edges didn't work out, go ahead and add the corners.
+	if (Settings.bEnableCornerHits)
+	{
+		if (Settings.bSelectDownRay)
+		{
+			for (int32 Id : DownRayElements)
+			{
+				ResultOut.SelectedCornerIDs.Add(Id);
+			}
+		}
+		else
+		{
+			ResultOut.SelectedCornerIDs.Add(ClosestElement->ID);
+		}
+		return true;
+	}
+
+	return false;
+}
+
+bool FGroupTopologySelector::DoEdgeBasedSelection(const FSelectionSettings& Settings, const FRay3d& Ray,
+	FDynamicMeshAABBTree3* Spatial, const FGeometrySet3& TopoSpatial,
+	FGroupTopologySelection& ResultOut, FVector3d& SelectedPositionOut, int32* EdgeSegmentIdOut) const
+{
+	// These will store our results, depending on whether we select all along the ray or not.
+	FGeometrySet3::FNearest SingleElement;
+	// ElementsWithinTolerance gives all nearby elements returned by the topology selector, whereas DownRayElements filters these to just
+	// the closest element and those that project directly onto it (it stores the polyline and segment IDs, which are all we care about).
+	TArray<FGeometrySet3::FNearest> ElementsWithinTolerance;
+	TArray<FIndex2i> DownRayElements;
+
+	auto LocalTolTest = [this](const FVector3d& A, const FVector3d& B) { return PointsWithinToleranceTest(A, B, 1.0); };
+
+	// Start by getting the closest element
+	const FGeometrySet3::FNearest* ClosestElement = nullptr;
+	if (!Settings.bSelectDownRay)
+	{
+		if (TopoSpatial.FindNearestCurveToRay(Ray, SingleElement, LocalTolTest))
+		{
+			ClosestElement = &SingleElement;
+		}
+	}
+	else
+	{
+		// Need all curves within tolerance, but also need to know the closest.
+		if (TopoSpatial.CollectCurvesNearRay(Ray, ElementsWithinTolerance, LocalTolTest))
+		{
+			double MinRayT = TNumericLimits<double>::Max();
+			for (const FGeometrySet3::FNearest& Element : ElementsWithinTolerance)
+			{
+				if (Element.RayParam < MinRayT)
+				{
+					MinRayT = Element.RayParam;
+					ClosestElement = &Element;
+				}
+			}
+		}//end if found edges
+	}//end if selecting down ray
+
+	// Bail early if we haven't found at least one edge
+	if (ClosestElement == nullptr)
+	{
+		return false;
+	}
+
+	// Also bail if the closest element is not visible.
+	if (!Settings.bIgnoreOcclusion && IsOccluded(*ClosestElement, Ray.Origin, Spatial))
+	{
+		return false;
+	}
+
+	// The closest point is already found
+	SelectedPositionOut = ClosestElement->NearestGeoPoint;
+
+	// If we have other edges, we need to filter them to only those that project onto the closest element. This would be done
+	// differently for perspective cameras vs orthographic projection, but since the behavior is only useful in ortho mode,
+	// we do it that way.
+	if (Settings.bSelectDownRay)
+	{
+		// Closest element is a given
+		DownRayElements.Add(FIndex2i(ClosestElement->ID, ClosestElement->PolySegmentIdx));
+
+		// We want edges that lie in a plane through the closest edge that is coplanar with the view direction. If we had wanted 
+		// to do this for perspective mode, we would have picked a plane through the edge and the camera origin.
+		int32 ClosestEid = Topology->GetGroupEdgeEdges(ClosestElement->ID)[ClosestElement->PolySegmentIdx];
+		FDynamicMesh3::FEdge ClosestEdge = Mesh->GetEdge(ClosestEid);
+		UE::Geometry::FPlane3d PlaneThroughClosestEdge(Mesh->GetVertex(ClosestEdge.Vert.A), Mesh->GetVertex(ClosestEdge.Vert.B), 
+			Mesh->GetVertex(ClosestEdge.Vert.A) + Ray.Direction);
+
+		for (const FGeometrySet3::FNearest& Element : ElementsWithinTolerance)
+		{
+			if (ClosestElement == &Element)
+			{
+				continue; // already added
+			}
+
+			// See if the edge endpoints lie in the plane
+			int32 Eid = Topology->GetGroupEdgeEdges(Element.ID)[Element.PolySegmentIdx];
+			FDynamicMesh3::FEdge Edge = Mesh->GetEdge(Eid);
+			if (abs(PlaneThroughClosestEdge.DistanceTo(Mesh->GetVertex(Edge.Vert.A))) < KINDA_SMALL_NUMBER
+				&& abs(PlaneThroughClosestEdge.DistanceTo(Mesh->GetVertex(Edge.Vert.B))) < KINDA_SMALL_NUMBER)
+			{
+				DownRayElements.Add(FIndex2i(Element.ID, Element.PolySegmentIdx));
+			}
+		}
+	}
+
+	// Try to select faces that project to the closest edge
+	if (Settings.bPreferProjectedElement && Settings.bEnableFaceHits)
+	{
+		TSet<int32> AddedGroups;
+
+		// Start with the closest edge
+		int32 Eid = Topology->GetGroupEdgeEdges(ClosestElement->ID)[ClosestElement->PolySegmentIdx];
+		FDynamicMesh3::FEdge Edge = Mesh->GetEdge(Eid);
+		FVector3d VertA = Mesh->GetVertex(Edge.Vert.A);
+		FVector3d VertB = Mesh->GetVertex(Edge.Vert.B);
+
+
+		// Grab a plane through the two verts that contains the ray direction (again, this is assuming that we'd
+		// only do this in ortho mode, otherwise the plane would go through the ray origin.
+		UE::Geometry::FPlane3d PlaneThroughClosestEdge(VertA, VertB, VertA + Ray.Direction);
+
+		// Checking that the face is coplanar simply entails checking that the opposite vert is in the plane. However,
+		// it is possible even for the closest edge to have multiple coplanar faces if a group going towards
+		// the camera curves away. We need to be able to select just one, and we want to select the one that
+		// extends down the view ray. For that, we'll make sure that a vector to the opposite vertex lies on the same
+		// side of the edge as the view ray.
+		FVector3d EdgeVector = VertB - VertA;
+		FVector3d EdgeVecCrossDirection = EdgeVector.Cross(Ray.Direction);
+		FIndex2i OppositeVids = Mesh->GetEdgeOpposingV(Eid);
+
+		FVector3d OppositeVert = Mesh->GetVertex(OppositeVids.A);
+		if (abs(PlaneThroughClosestEdge.DistanceTo(OppositeVert)) < KINDA_SMALL_NUMBER
+			&& EdgeVector.Cross(OppositeVert - VertA).Dot(EdgeVecCrossDirection) > 0)
+		{
+			int GroupId = Topology->GetGroupID(Edge.Tri.A);
+			ResultOut.SelectedGroupIDs.Add(GroupId);
+			AddedGroups.Add(GroupId);
+		}
+		else if (OppositeVids.B != FDynamicMesh3::InvalidID)
+		{
+			OppositeVert = Mesh->GetVertex(OppositeVids.B);
+
+			if (abs(PlaneThroughClosestEdge.DistanceTo(OppositeVert)) < KINDA_SMALL_NUMBER
+				&& EdgeVector.Cross(OppositeVert - VertA).Dot(EdgeVecCrossDirection) > 0)
+			{
+				int GroupId = Topology->GetGroupID(Edge.Tri.B);
+				ResultOut.SelectedGroupIDs.Add(GroupId);
+				AddedGroups.Add(GroupId);
+			}
+		}
+
+		// If relevant, get all the other coplanar faces
+		if (Settings.bSelectDownRay && AddedGroups.Num() > 0)
+		{
+			for (int i = 1; i < DownRayElements.Num(); ++i) // skip 0 because it is closest and done
+			{
+				// We already made sure that all these edges are coplanar, so we'll just be checking opposite verts.
+				Eid = Topology->GetGroupEdgeEdges(DownRayElements[i].A)[DownRayElements[i].B];
+				Edge = Mesh->GetEdge(Eid);
+				OppositeVids = Mesh->GetEdgeOpposingV(Eid);
+
+				// No need to check directionality of faces here, just need them to be coplanar.
+				if (abs(PlaneThroughClosestEdge.DistanceTo(Mesh->GetVertex(OppositeVids.A))) < KINDA_SMALL_NUMBER)
+				{
+					int GroupId = Topology->GetGroupID(Edge.Tri.A);
+					if (!AddedGroups.Contains(GroupId))
+					{
+						ResultOut.SelectedGroupIDs.Add(GroupId);
+						AddedGroups.Add(GroupId);
+					}
+				}
+				if (OppositeVids.B != FDynamicMesh3::InvalidID
+					&& abs(PlaneThroughClosestEdge.DistanceTo(Mesh->GetVertex(OppositeVids.B))) < KINDA_SMALL_NUMBER)
+				{
+					int GroupId = Topology->GetGroupID(Edge.Tri.B);
+					if (!AddedGroups.Contains(GroupId))
+					{
+						ResultOut.SelectedGroupIDs.Add(GroupId);
+						AddedGroups.Add(GroupId);
+					}
+				}
+			}//end going through other edges
+		}
+
+		// If we selected faces, we're done
+		if (AddedGroups.Num() > 0)
+		{
+			return true;
+		}
+	}
+
+	// If we didn't end up selecting projected faces, and we have edges to select, select them
+	if (Settings.bEnableEdgeHits)
+	{
+		if (Settings.bSelectDownRay)
+		{
+			for (const FIndex2i ElementTuple : DownRayElements)
+			{
+				ResultOut.SelectedEdgeIDs.Add(ElementTuple.A);
+			}
+		}
+		else
+		{
+			ResultOut.SelectedEdgeIDs.Add(ClosestElement->ID);
+		}
+
+		if (EdgeSegmentIdOut)
+		{
+			*EdgeSegmentIdOut = ClosestElement->PolySegmentIdx;
+		}
+
+		return true;
+	}
+	return false;
+}
+
+bool IsOccluded(const FGeometrySet3::FNearest& ClosestElement, const FVector3d& ViewOrigin, const FDynamicMeshAABBTree3* Spatial)
+{
+	return IsOccluded(ClosestElement.NearestGeoPoint, ViewOrigin, Spatial);
+}
+
+bool IsOccluded(const FVector3d& Point, const FVector3d& ViewOrigin, const FDynamicMeshAABBTree3* Spatial)
+{
+	// Shoot ray backwards to see if we hit something. 
+	FRay3d ToEyeRay(Point, UE::Geometry::Normalized(ViewOrigin - Point), true);
+	ToEyeRay.Origin += (double)(100 * FMathf::ZeroTolerance) * ToEyeRay.Direction;
+	if (Spatial->FindNearestHitTriangle(ToEyeRay) >= 0)
+	{
+		return true;
+	}
+	return false;
+}
+
+bool FGroupTopologySelector::FindSelectedElement(const FSelectionSettings& Settings, 
+	const FCameraRectangle& CameraRectangle, FTransform3d TargetTransform, FGroupTopologySelection& ResultOut)
+{
+	// One minor easy thing we can do to speed up the below is to detect cases where transform does not have a non-uniform
+	// scale, and in those cases we transform the rectangle once rather than transforming all the query points/curves. 
+	// This isn't worth doing for non-uniform scale transforms because the rectangle basis stops being orthonormal and 
+	// becomes a problem to work with.
+
+	ResultOut.Clear();
+	FDynamicMeshAABBTree3* Spatial = GetSpatial();
+
+	// Needed for occlusion test, which happens in local space.
+	FVector3d LocalCameraOrigin = TargetTransform.InverseTransformPosition((FVector3d)CameraRectangle.CameraOrigin);
+	
+	// Corner selection takes priority over edges.
+	if (Settings.bEnableCornerHits)
+	{
+		GeometrySet.ParallelFindAllPointsSatisfying(
+			[&CameraRectangle, &TargetTransform, &Settings, &LocalCameraOrigin, Spatial](const FVector3d& PointPosition)
+			{
+				return CameraRectangle.IsProjectedPointInRectangle((FVector)TargetTransform.TransformPosition(PointPosition))
+					&& (Settings.bIgnoreOcclusion || !IsOccluded(PointPosition, LocalCameraOrigin, Spatial));
+			},
+			ResultOut.SelectedCornerIDs);
+	}
+
+	// If we didn't get corners, look for edges.
+	if (ResultOut.SelectedCornerIDs.IsEmpty() && Settings.bEnableEdgeHits)
+	{
+		GeometrySet.ParallelFindAllCurvesSatisfying(
+			[&CameraRectangle, &TargetTransform, &Settings, &LocalCameraOrigin, Spatial](const FPolyline3d& Curve)
+			{
+				// Testing occlusion properly seems like it would be a pain, so we'll just consider something occluded
+				// if one of the endpoints is occluded. This will handle the common case of not wanting to select a
+				// hidden edge that is connected to a visible corner.
+				if (!Settings.bIgnoreOcclusion && (IsOccluded(Curve.Start(), LocalCameraOrigin, Spatial)
+						|| IsOccluded(Curve.End(), LocalCameraOrigin, Spatial)))
+				{
+					return false;
+				}
+
+				// Check whether any of the component segments intersect the rectangle
+				const TArray<FVector3<double>>& Verts = Curve.GetVertices();
+				FVector3d CurrentVert = TargetTransform.TransformPosition(Verts[0]);
+				for (int32 i = 1; i < Verts.Num(); ++i)
+				{
+					FVector3d NextVert = TargetTransform.TransformPosition(Verts[i]);
+					if (CameraRectangle.IsProjectedSegmentIntersectingRectangle((FVector)CurrentVert, (FVector)NextVert))
+					{
+						return true;
+					}
+					CurrentVert = NextVert;
+				}
+				return false;
+			},
+			ResultOut.SelectedEdgeIDs);
+	}
+	
+	return ResultOut.IsEmpty();
+}
+
+bool FGroupTopologySelector::ExpandSelectionByEdgeLoops(FGroupTopologySelection& Selection)
+{
+	TSet<int32> EdgeSet(Selection.SelectedEdgeIDs);
+	int32 OriginalNumEdges = Selection.SelectedEdgeIDs.Num();
+	for (int32 Eid : Selection.SelectedEdgeIDs)
+	{
+		const FGroupTopology::FGroupEdge& Edge = Topology->Edges[Eid];
+		if (Edge.EndpointCorners[0] == IndexConstants::InvalidID)
+		{
+			continue; // This FGroupEdge is a loop unto itself (and already in our selection, since we're looking at it).
+		}
+
+		// Go forward and backward adding edges
+		AddNewEdgeLoopEdgesFromCorner(*Topology, Eid, Edge.EndpointCorners[0], EdgeSet);
+		AddNewEdgeLoopEdgesFromCorner(*Topology, Eid, Edge.EndpointCorners[1], EdgeSet);
+	}
+
+	if (EdgeSet.Num() > OriginalNumEdges)
+	{
+		Selection.SelectedEdgeIDs.Append(EdgeSet);
+		return true;
+	}
+	else
+	{
+		return false;
+	}
+}
+
+void AddNewEdgeLoopEdgesFromCorner(const FGroupTopology& Topology, int32 EdgeID, int32 CornerID, TSet<int32>& EdgeSet)
+{
+	const FGroupTopology::FCorner& CurrentCorner = Topology.Corners[CornerID];
+
+	int32 LastCornerID = CornerID;
+	int32 LastEdgeID = EdgeID;
+	while (true)
+	{
+		int32 NextEid;
+		if (!GetNextEdgeLoopEdge(Topology, LastEdgeID, LastCornerID, NextEid))
+		{
+			break; // Probably not a valence 4 corner
+		}
+		if (EdgeSet.Contains(NextEid))
+		{
+			break; // Either we finished the loop, or we'll continue it from another selection
+		}
+
+		EdgeSet.Add(NextEid);
+
+		LastEdgeID = NextEid;
+		const FGroupTopology::FGroupEdge& LastEdge = Topology.Edges[LastEdgeID];
+		LastCornerID = LastEdge.EndpointCorners[0] == LastCornerID ? LastEdge.EndpointCorners[1] : LastEdge.EndpointCorners[0];
+
+		check(LastCornerID != IndexConstants::InvalidID);
+	}
+}
+
+bool GetNextEdgeLoopEdge(const FGroupTopology& Topology, int32 IncomingEdgeID, int32 CornerID, int32& NextEdgeIDOut)
+{
+	// It's worth noting that the approach here breaks down in pathological cases where the same group is present
+	// multiple times around a corner (i.e. the group is not contiguous, and separate islands share a corner).
+	// It's not practical to worry about those cases.
+
+	NextEdgeIDOut = IndexConstants::InvalidID;
+	const FGroupTopology::FCorner& CurrentCorner = Topology.Corners[CornerID];
+
+	if (CurrentCorner.NeighbourGroupIDs.Num() != 4)
+	{
+		return false; // Not a valence 4 corner
+	}
+
+	const FGroupTopology::FGroupEdge& IncomingEdge = Topology.Edges[IncomingEdgeID];
+
+	// We want to find the edge that shares this corner but does not border either of the neighboring groups of
+	// the incoming edge.
+
+	for (int32 Gid : CurrentCorner.NeighbourGroupIDs)
+	{
+		if (Gid == IncomingEdge.Groups[0] || Gid == IncomingEdge.Groups[1])
+		{
+			continue; // This is one of the neighboring groups of the incoming edge
+		}
+
+		// Iterate through all edges of group
+		const FGroupTopology::FGroup* Group = Topology.FindGroupByID(Gid);
+		for (const FGroupTopology::FGroupBoundary& Boundary : Group->Boundaries)
+		{
+			for (int32 Eid : Boundary.GroupEdges)
+			{
+				const FGroupTopology::FGroupEdge& CandidateEdge = Topology.Edges[Eid];
+
+				// Edge must share corner but not neighboring groups
+				if ((CandidateEdge.EndpointCorners[0] == CornerID || CandidateEdge.EndpointCorners[1] == CornerID)
+					&& CandidateEdge.Groups[0] != IncomingEdge.Groups[0] && CandidateEdge.Groups[0] != IncomingEdge.Groups[1]
+					&& CandidateEdge.Groups[1] != IncomingEdge.Groups[0] && CandidateEdge.Groups[1] != IncomingEdge.Groups[1])
+				{
+					NextEdgeIDOut = Eid;
+					return true;
+				}
+			}
+		}
+	}
+	return false;
+}
+
+bool FGroupTopologySelector::ExpandSelectionByEdgeRings(FGroupTopologySelection& Selection)
+{
+	TSet<int32> EdgeSet(Selection.SelectedEdgeIDs);
+	int32 OriginalNumEdges = Selection.SelectedEdgeIDs.Num();
+	for (int32 Eid : Selection.SelectedEdgeIDs)
+	{
+		const FGroupTopology::FGroupEdge& Edge = Topology->Edges[Eid];
+
+		// Go forward and backward adding edges
+		if (Edge.Groups[0] != IndexConstants::InvalidID)
+		{
+			AddNewEdgeRingEdges(*Topology, Eid, Edge.Groups[0], EdgeSet);
+		}
+		if (Edge.Groups[0] != IndexConstants::InvalidID)
+		{
+			AddNewEdgeRingEdges(*Topology, Eid, Edge.Groups[1], EdgeSet);
+		}
+	}
+
+	if (EdgeSet.Num() > OriginalNumEdges)
+	{
+		Selection.SelectedEdgeIDs.Append(EdgeSet);
+		return true;
+	}
+	else
+	{
+		return false;
+	}
+}
+
+void AddNewEdgeRingEdges(const FGroupTopology& Topology, int32 StartEdgeID, int32 ForwardGroupID, TSet<int32>& EdgeSet)
+{
+	int32 CurrentEdgeID = StartEdgeID;
+	int32 CurrentForwardGroupID = ForwardGroupID;
+	while (true)
+	{
+		if (CurrentForwardGroupID == IndexConstants::InvalidID)
+		{
+			break; // Reached a boundary
+		}
+
+		int32 NextEdgeID;
+		if (!GetQuadOppositeEdge(Topology, CurrentEdgeID, CurrentForwardGroupID, NextEdgeID))
+		{
+			break; // Probably not a quad
+		}
+		if (EdgeSet.Contains(NextEdgeID))
+		{
+			break; // Either we finished the loop, or we'll continue it from another selection
+		}
+
+		EdgeSet.Add(NextEdgeID);
+
+		CurrentEdgeID = NextEdgeID;
+		const FGroupTopology::FGroupEdge& Edge = Topology.Edges[CurrentEdgeID];
+		CurrentForwardGroupID = (Edge.Groups[0] == CurrentForwardGroupID) ? Edge.Groups[1] : Edge.Groups[0];
+	}
+}
+
+bool GetQuadOppositeEdge(const FGroupTopology& Topology, int32 EdgeIDIn, int32 GroupID, int32& OppositeEdgeIDOut)
+{
+	const FGroupTopology::FGroup* Group = Topology.FindGroupByID(GroupID);
+	check(Group);
+
+	// Find the boundary that contains this edge
+	for (int32 i = 0; i < Group->Boundaries.Num(); ++i)
+	{
+		const FGroupTopology::FGroupBoundary& Boundary = Group->Boundaries[i];
+		int32 EdgeIndex = Boundary.GroupEdges.IndexOfByKey(EdgeIDIn);
+		if (EdgeIndex != INDEX_NONE)
+		{
+			if (Boundary.GroupEdges.Num() != 4)
+			{
+				return false;
+			}
+
+			OppositeEdgeIDOut = Boundary.GroupEdges[(EdgeIndex + 2) % 4];
+			return true;
+		}
+	}
+	check(false); // No boundary of the given group contained the given edge
+	return false;
+}
+
+
+void FGroupTopologySelector::DrawSelection(const FGroupTopologySelection& Selection, FToolDataVisualizer* Renderer, const FViewCameraState* CameraState)
+{
+	FLinearColor UseColor = Renderer->LineColor;
+	float LineWidth = Renderer->LineThickness;
+
+	for (int CornerID : Selection.SelectedCornerIDs)
+	{
+		int VertexID = Topology->GetCornerVertexID(CornerID);
+		FVector Position = (FVector)Mesh->GetVertex(VertexID);
+		FVector WorldPosition = Renderer->TransformP(Position);
+		
+		// Depending on whether we're in an orthographic view or not, we set the radius based on visual angle or based on ortho 
+		// viewport width (divided into 90 segments like the FOV is divided into 90 degrees).
+		float Radius = (CameraState->bIsOrthographic) ? CameraState->OrthoWorldCoordinateWidth * 0.5 / 90.0
+			: (float)ToolSceneQueriesUtil::CalculateDimensionFromVisualAngleD(*CameraState, (FVector3d)WorldPosition, 0.5);
+		Renderer->DrawViewFacingCircle(Position, Radius, 16, UseColor, LineWidth, false);
+	}
+
+	for (int EdgeID : Selection.SelectedEdgeIDs)
+	{
+		const TArray<int>& Vertices = Topology->GetGroupEdgeVertices(EdgeID);
+		int NV = Vertices.Num() - 1;
+
+		// Draw the edge, but also draw the endpoints in ortho mode (to make projected edges visible)
+		FVector A = (FVector)Mesh->GetVertex(Vertices[0]);
+		if (CameraState->bIsOrthographic)
+		{
+			Renderer->DrawPoint(A, UseColor, 10, false);
+		}
+		for (int k = 0; k < NV; ++k)
+		{
+			FVector B = (FVector)Mesh->GetVertex(Vertices[k+1]);
+			Renderer->DrawLine(A, B, UseColor, LineWidth, false);
+			A = B;
+		}
+		if (CameraState->bIsOrthographic)
+		{
+			Renderer->DrawPoint(A, UseColor, LineWidth, false);
+		}
+	}
+
+	// We are not responsible for drawing the faces, but do draw the sides of the faces in ortho mode to make them visible
+	// when they project to an edge.
+	if (CameraState->bIsOrthographic && Selection.SelectedGroupIDs.Num() > 0)
+	{
+		Topology->ForGroupSetEdges(Selection.SelectedGroupIDs,
+			[&UseColor, LineWidth, Renderer, this] (FGroupTopology::FGroupEdge Edge, int EdgeID) 
+		{
+			Topology->GetGroupEdgeVertices(EdgeID);
+			const TArray<int>& Vertices = Topology->GetGroupEdgeVertices(EdgeID);
+			int NV = Vertices.Num() - 1;
+			FVector A = (FVector)Mesh->GetVertex(Vertices[0]);
+			for (int k = 0; k < NV; ++k)
+			{
+				FVector B = (FVector)Mesh->GetVertex(Vertices[k + 1]);
+				Renderer->DrawLine(A, B, UseColor, LineWidth, false);
+				A = B;
+			}
+		});
+	}
+}
\ No newline at end of file
diff --git a/Engine/Plugins/Experimental/MeshModelingToolsetExp/Source/ModelingComponents/Private/Selection/PolygonSelectionMechanic.cpp b/Engine/Plugins/Experimental/MeshModelingToolsetExp/Source/ModelingComponents/Private/Selection/PolygonSelectionMechanic.cpp
new file mode 100644
--- /dev/null
+++ b/Engine/Plugins/Experimental/MeshModelingToolsetExp/Source/ModelingComponents/Private/Selection/PolygonSelectionMechanic.cpp
@@ -0,0 +1,768 @@
+// Copyright Epic Games, Inc. All Rights Reserved.
+
+#include "Selection/PolygonSelectionMechanic.h"
+
+#include "BaseBehaviors/ClickDragBehavior.h"
+#include "BaseBehaviors/MouseHoverBehavior.h"
+#include "BaseBehaviors/SingleClickBehavior.h"
+#include "InteractiveToolManager.h"
+#include "Util/ColorConstants.h"
+#include "Selection/GroupTopologyStorableSelection.h"
+#include "ToolSceneQueriesUtil.h"
+#include "ToolSetupUtil.h"
+
+#include "ExplicitUseGeometryMathTypes.h"		// using UE::Geometry::(math types)
+using namespace UE::Geometry;
+
+#define LOCTEXT_NAMESPACE "UPolygonSelectionMechanic"
+
+UPolygonSelectionMechanic::~UPolygonSelectionMechanic()
+{
+	checkf(PreviewGeometryActor == nullptr, TEXT("Shutdown() should be called before UPolygonSelectionMechanic is destroyed."));
+}
+
+void UPolygonSelectionMechanic::Setup(UInteractiveTool* ParentToolIn)
+{
+	UInteractionMechanic::Setup(ParentToolIn);
+
+	TopoSelector = MakeShared<FGroupTopologySelector, ESPMode::ThreadSafe>();
+
+	ClickBehavior = NewObject<USingleClickInputBehavior>();
+	ClickBehavior->Initialize(this);
+	ClickBehavior->SetDefaultPriority(BasePriority);
+	ParentToolIn->AddInputBehavior(ClickBehavior, this);
+
+	HoverBehavior = NewObject<UMouseHoverBehavior>();
+	HoverBehavior->Modifiers.RegisterModifier(ShiftModifierID, FInputDeviceState::IsShiftKeyDown);
+	HoverBehavior->Modifiers.RegisterModifier(CtrlModifierID, FInputDeviceState::IsCtrlKeyDown);
+	HoverBehavior->Initialize(this);
+	HoverBehavior->SetDefaultPriority(BasePriority);
+	ParentToolIn->AddInputBehavior(HoverBehavior, this);
+
+	MarqueeMechanic = NewObject<URectangleMarqueeMechanic>();
+	MarqueeMechanic->Setup(ParentToolIn);
+	MarqueeMechanic->OnDragRectangleStarted.AddUObject(this, &UPolygonSelectionMechanic::OnDragRectangleStarted);
+	MarqueeMechanic->OnDragRectangleChanged.AddUObject(this, &UPolygonSelectionMechanic::OnDragRectangleChanged);
+	MarqueeMechanic->OnDragRectangleFinished.AddUObject(this, &UPolygonSelectionMechanic::OnDragRectangleFinished);
+	MarqueeMechanic->SetBasePriority(BasePriority.MakeLower());
+
+	Properties = NewObject<UPolygonSelectionMechanicProperties>(this);
+	if (bAddSelectionFilterPropertiesToParentTool)
+	{
+		AddToolPropertySource(Properties);
+	}
+	Properties->WatchProperty(Properties->bSelectVertices, [this](bool bSelectVertices) {
+		UpdateMarqueeEnabled(); 
+	});
+	Properties->WatchProperty(Properties->bSelectEdges, [this](bool bSelectVertices) { 
+		UpdateMarqueeEnabled(); 
+	});
+	Properties->WatchProperty(Properties->bEnableMarquee, [this](bool bEnableMarquee) { 
+		UpdateMarqueeEnabled(); 
+	});
+
+	// set up visualizers
+	PolyEdgesRenderer.LineColor = FLinearColor::Red;
+	PolyEdgesRenderer.LineThickness = 2.0;
+	HilightRenderer.LineColor = FLinearColor::Green;
+	HilightRenderer.LineThickness = 4.0f;
+	SelectionRenderer.LineColor = LinearColors::Gold3f();
+	SelectionRenderer.LineThickness = 4.0f;
+
+	float HighlightedFacePercentDepthOffset = 0.5f;
+	HighlightedFaceMaterial = ToolSetupUtil::GetSelectionMaterial(FLinearColor::Green, ParentToolIn->GetToolManager(), HighlightedFacePercentDepthOffset);
+	// The rest of the highlighting setup has to be done in Initialize(), since we need the world to set up our drawing component.
+}
+
+void UPolygonSelectionMechanic::Shutdown()
+{
+	if (PreviewGeometryActor)
+	{
+		PreviewGeometryActor->Destroy();
+		PreviewGeometryActor = nullptr;
+	}
+}
+
+void UPolygonSelectionMechanic::GetStorableSelection(UGroupTopologyStorableSelection& StorableSelectionOut, const FCompactMaps* CompactMapsToApply) const
+{
+	if (CompactMapsToApply)
+	{
+		StorableSelectionOut.SetSelection(*Topology, PersistentSelection, *CompactMapsToApply);
+	}
+	else
+	{
+		StorableSelectionOut.SetSelection(*Topology, PersistentSelection);
+	}
+}
+
+void UPolygonSelectionMechanic::LoadStorableSelection(const UGroupTopologyStorableSelection& StorableSelectionIn)
+{
+	StorableSelectionIn.ExtractIntoSelectionObject(*Topology, PersistentSelection);
+}
+
+void UPolygonSelectionMechanic::Initialize(
+	const FDynamicMesh3* MeshIn,
+	FTransform TargetTransformIn,
+	UWorld* WorldIn,
+	const FGroupTopology* TopologyIn,
+	TFunction<FDynamicMeshAABBTree3 * ()> GetSpatialSourceFuncIn)
+{
+	this->Mesh = MeshIn;
+	this->Topology = TopologyIn;
+	this->TargetTransform = FTransform3d(TargetTransformIn);
+
+	TopoSelector->Initialize(Mesh, Topology);
+	this->GetSpatialFunc = GetSpatialSourceFuncIn;
+	TopoSelector->SetSpatialSource(GetSpatialFunc);
+	TopoSelector->PointsWithinToleranceTest = [this](const FVector3d& Position1, const FVector3d& Position2, double TolScale) {
+		if (CameraState.bIsOrthographic)
+		{
+			// We could just always use ToolSceneQueriesUtil::PointSnapQuery. But in ortho viewports, we happen to know
+			// that the only points that we will ever give this function will be the closest points between a ray and
+			// some geometry, meaning that the vector between them will be orthogonal to the view ray. With this knowledge,
+			// we can do the tolerance computation more efficiently than PointSnapQuery can, since we don't need to project
+			// down to the view plane.
+			// As in PointSnapQuery, we convert our angle-based tolerance to one we can use in an ortho viewport (instead of
+			// dividing our field of view into 90 visual angle degrees, we divide the plane into 90 units).
+			float OrthoTolerance = ToolSceneQueriesUtil::GetDefaultVisualAngleSnapThreshD() * CameraState.OrthoWorldCoordinateWidth / 90.0;
+			OrthoTolerance *= TolScale;
+			return TargetTransform.TransformPosition(Position1).DistanceSquared(TargetTransform.TransformPosition(Position2)) < OrthoTolerance * OrthoTolerance;
+		}
+		else
+		{
+			return ToolSceneQueriesUtil::PointSnapQuery(CameraState,
+				TargetTransform.TransformPosition(Position1), TargetTransform.TransformPosition(Position2),
+				ToolSceneQueriesUtil::GetDefaultVisualAngleSnapThreshD() * TolScale);
+		}
+	};
+
+	// Set up the component we use to draw highlighted triangles. Only needs to be done once, not when the mesh
+	// changes (we are assuming that we won't swap worlds without creating a new mechanic).
+	if (PreviewGeometryActor == nullptr)
+	{
+		FRotator Rotation(0.0f, 0.0f, 0.0f);
+		FActorSpawnParameters SpawnInfo;;
+		PreviewGeometryActor = WorldIn->SpawnActor<APreviewGeometryActor>(FVector::ZeroVector, Rotation, SpawnInfo);
+
+		DrawnTriangleSetComponent = NewObject<UTriangleSetComponent>(PreviewGeometryActor);
+		PreviewGeometryActor->SetRootComponent(DrawnTriangleSetComponent);
+		DrawnTriangleSetComponent->RegisterComponent();
+	}
+
+	PreviewGeometryActor->SetActorTransform(TargetTransformIn);
+
+	DrawnTriangleSetComponent->Clear();
+	CurrentlyHighlightedGroups.Empty();
+}
+
+void UPolygonSelectionMechanic::Initialize(
+	UDynamicMeshComponent* MeshComponentIn,
+	const FGroupTopology* TopologyIn,
+	TFunction<FDynamicMeshAABBTree3 * ()> GetSpatialSourceFuncIn)
+{
+
+	Initialize(MeshComponentIn->GetMesh(),
+		MeshComponentIn->GetComponentTransform(),
+		MeshComponentIn->GetWorld(),
+		TopologyIn,
+		GetSpatialSourceFuncIn);
+}
+
+void UPolygonSelectionMechanic::DisableBehaviors(UInteractiveTool* ParentToolIn)
+{
+	ParentToolIn->RemoveInputBehaviorsBySource(this);
+	ParentToolIn->RemoveInputBehaviorsBySource(MarqueeMechanic);
+
+	// TODO: Is it worth adding a way to remove the property watchers for marquee?
+}
+
+void UPolygonSelectionMechanic::SetIsEnabled(bool bOn)
+{
+	bIsEnabled = bOn;
+	UpdateMarqueeEnabled();
+}
+
+void UPolygonSelectionMechanic::SetBasePriority(const FInputCapturePriority &Priority)
+{
+	BasePriority = Priority;
+	if (ClickBehavior)
+	{
+		ClickBehavior->SetDefaultPriority(Priority);
+	}
+	if (HoverBehavior)
+	{
+		HoverBehavior->SetDefaultPriority(Priority);
+	}
+	if (MarqueeMechanic)
+	{
+		MarqueeMechanic->SetBasePriority(Priority.MakeLower());
+	}
+}
+
+TPair<FInputCapturePriority, FInputCapturePriority> UPolygonSelectionMechanic::GetPriorityRange() const
+{
+	TPair<FInputCapturePriority, FInputCapturePriority> Result;
+	Result.Key = BasePriority;
+	Result.Value = MarqueeMechanic->GetPriorityRange().Value;
+	return Result;
+}
+
+void UPolygonSelectionMechanic::Render(IToolsContextRenderAPI* RenderAPI)
+{
+	MarqueeMechanic->Render(RenderAPI);
+
+	// Cache the view camera state so we can use for snapping/etc.
+	// This should not happen in Render() though...
+	GetParentTool()->GetToolManager()->GetContextQueriesAPI()->GetCurrentViewState(CameraState);
+
+	FViewCameraState RenderCameraState = RenderAPI->GetCameraState();
+
+	const FDynamicMesh3* TargetMesh = this->Mesh;
+	FTransform Transform = (FTransform)TargetTransform;
+
+	PolyEdgesRenderer.BeginFrame(RenderAPI, RenderCameraState);
+	PolyEdgesRenderer.SetTransform(Transform);
+	for (const FGroupTopology::FGroupEdge& Edge : Topology->Edges)
+	{
+		FVector3d A, B;
+		for (int32 eid : Edge.Span.Edges)
+		{
+			TargetMesh->GetEdgeV(eid, A, B);
+			PolyEdgesRenderer.DrawLine(A, B);
+		}
+	}
+	PolyEdgesRenderer.EndFrame();
+
+	if (PersistentSelection.IsEmpty() == false)
+	{
+		SelectionRenderer.BeginFrame(RenderAPI, RenderCameraState);
+		SelectionRenderer.SetTransform(Transform);
+		TopoSelector->DrawSelection(PersistentSelection, &SelectionRenderer, &RenderCameraState);
+		SelectionRenderer.EndFrame();
+	}
+
+	HilightRenderer.BeginFrame(RenderAPI, RenderCameraState);
+	HilightRenderer.SetTransform(Transform);
+	TopoSelector->DrawSelection(HilightSelection, &HilightRenderer, &RenderCameraState);
+	HilightRenderer.EndFrame();
+}
+
+void UPolygonSelectionMechanic::DrawHUD(FCanvas* Canvas, IToolsContextRenderAPI* RenderAPI)
+{
+	MarqueeMechanic->DrawHUD(Canvas, RenderAPI);
+}
+
+
+
+
+void UPolygonSelectionMechanic::ClearHighlight()
+{
+	checkf(DrawnTriangleSetComponent != nullptr, TEXT("Initialize() not called on UPolygonSelectionMechanic."));
+
+	HilightSelection.Clear();
+	DrawnTriangleSetComponent->Clear();
+	CurrentlyHighlightedGroups.Empty();
+}
+
+
+void UPolygonSelectionMechanic::NotifyMeshChanged(bool bTopologyModified)
+{
+	ClearHighlight();
+	TopoSelector->Invalidate(true, bTopologyModified);
+	if (bTopologyModified)
+	{
+		PersistentSelection.Clear();
+		SelectionTimestamp++;
+		OnSelectionChanged.Broadcast();
+	}
+}
+
+
+bool UPolygonSelectionMechanic::TopologyHitTest(const FRay& WorldRay, FHitResult& OutHit, bool bUseOrthoSettings)
+{
+	FGroupTopologySelection Selection;
+	return TopologyHitTest(WorldRay, OutHit, Selection, bUseOrthoSettings);
+}
+
+bool UPolygonSelectionMechanic::TopologyHitTest(const FRay& WorldRay, FHitResult& OutHit, FGroupTopologySelection& OutSelection, bool bUseOrthoSettings)
+{
+	FRay3d LocalRay(TargetTransform.InverseTransformPosition((FVector3d)WorldRay.Origin),
+		TargetTransform.InverseTransformVector((FVector3d)WorldRay.Direction));
+	UE::Geometry::Normalize(LocalRay.Direction);
+
+	FVector3d LocalPosition, LocalNormal;
+	int32 EdgeSegmentId; // Only used if hit is an edge
+	FGroupTopologySelector::FSelectionSettings TopoSelectorSettings = GetTopoSelectorSettings(bUseOrthoSettings);
+	if (TopoSelector->FindSelectedElement(TopoSelectorSettings, LocalRay, OutSelection, LocalPosition, LocalNormal, &EdgeSegmentId) == false)
+	{
+		return false;
+	}
+
+	if (OutSelection.SelectedCornerIDs.Num() > 0)
+	{
+		OutHit.FaceIndex = OutSelection.GetASelectedCornerID();
+		OutHit.Distance = LocalRay.Project(LocalPosition);
+		OutHit.ImpactPoint = (FVector)TargetTransform.TransformPosition(LocalRay.PointAt(OutHit.Distance));
+	}
+	else if (OutSelection.SelectedEdgeIDs.Num() > 0)
+	{
+		OutHit.FaceIndex = OutSelection.GetASelectedEdgeID();
+		OutHit.Distance = LocalRay.Project(LocalPosition);
+		OutHit.ImpactPoint = (FVector)TargetTransform.TransformPosition(LocalRay.PointAt(OutHit.Distance));
+		OutHit.Item = EdgeSegmentId;
+	}
+	else
+	{
+		FDynamicMeshAABBTree3* Spatial = GetSpatialFunc();
+		int HitTID = Spatial->FindNearestHitTriangle(LocalRay);
+		if (HitTID != IndexConstants::InvalidID)
+		{
+			FTriangle3d Triangle;
+			Spatial->GetMesh()->GetTriVertices(HitTID, Triangle.V[0], Triangle.V[1], Triangle.V[2]);
+			FIntrRay3Triangle3d Query(LocalRay, Triangle);
+			if (Query.Find())
+			{
+				OutHit.FaceIndex = HitTID;
+				OutHit.Distance = (float)Query.RayParameter;
+				OutHit.Normal = (FVector)TargetTransform.TransformVectorNoScale(Spatial->GetMesh()->GetTriNormal(HitTID));
+				OutHit.ImpactPoint = (FVector)TargetTransform.TransformPosition(LocalRay.PointAt(Query.RayParameter));
+			}
+			else
+			{
+				return false;
+			}
+		}
+	}
+	return true;
+}
+
+
+
+FGroupTopologySelector::FSelectionSettings UPolygonSelectionMechanic::GetTopoSelectorSettings(bool bUseOrthoSettings)
+{
+	FGroupTopologySelector::FSelectionSettings Settings;
+
+	Settings.bEnableFaceHits = Properties->bSelectFaces;
+	Settings.bEnableEdgeHits = Properties->bSelectEdges;
+	Settings.bEnableCornerHits = Properties->bSelectVertices;
+
+	if (!PersistentSelection.IsEmpty() && (ShouldAddToSelectionFunc() || ShouldRemoveFromSelectionFunc()))
+	{
+		Settings.bEnableFaceHits = Settings.bEnableFaceHits && PersistentSelection.SelectedGroupIDs.Num() > 0;
+		Settings.bEnableEdgeHits = Settings.bEnableEdgeHits && PersistentSelection.SelectedEdgeIDs.Num() > 0;
+		Settings.bEnableCornerHits = Settings.bEnableCornerHits && PersistentSelection.SelectedCornerIDs.Num() > 0;
+	}
+
+	if (bUseOrthoSettings)
+	{
+		Settings.bPreferProjectedElement = Properties->bPreferProjectedElement;
+		Settings.bSelectDownRay = Properties->bSelectDownRay;
+		Settings.bIgnoreOcclusion = Properties->bIgnoreOcclusion;
+	}
+
+	return Settings;
+}
+
+
+
+
+bool UPolygonSelectionMechanic::UpdateHighlight(const FRay& WorldRay)
+{
+	checkf(DrawnTriangleSetComponent != nullptr, TEXT("Initialize() not called on UPolygonSelectionMechanic."));
+
+	FRay3d LocalRay(TargetTransform.InverseTransformPosition((FVector3d)WorldRay.Origin),
+		TargetTransform.InverseTransformVector((FVector3d)WorldRay.Direction));
+	UE::Geometry::Normalize(LocalRay.Direction);
+
+	HilightSelection.Clear();
+	FVector3d LocalPosition, LocalNormal;
+	FGroupTopologySelector::FSelectionSettings TopoSelectorSettings = GetTopoSelectorSettings(CameraState.bIsOrthographic);
+	bool bHit = TopoSelector->FindSelectedElement(TopoSelectorSettings, LocalRay, HilightSelection, LocalPosition, LocalNormal);
+
+	if (HilightSelection.SelectedEdgeIDs.Num() > 0 && Properties->bSelectEdgeRings && ShouldSelectEdgeRingsFunc())
+	{
+		TopoSelector->ExpandSelectionByEdgeRings(HilightSelection);
+	}
+	if (HilightSelection.SelectedEdgeIDs.Num() > 0 && Properties->bSelectEdgeLoops && ShouldSelectEdgeLoopsFunc())
+	{
+		TopoSelector->ExpandSelectionByEdgeLoops(HilightSelection);
+	}
+
+	// Don't hover highlight a selection that we already selected, because people didn't like that
+	if (PersistentSelection.Contains(HilightSelection))
+	{
+		HilightSelection.Clear();
+	}
+
+	// Currently we draw highlighted edges/vertices differently from highlighted faces. Edges/vertices
+	// get drawn in the Render() call, so it is sufficient to just update HighlightSelection above.
+	// Faces, meanwhile, get placed into a Component that is rendered through the normal rendering system.
+	// So, we need to update the component when the highlighted selection changes.
+
+	// Put hovered groups in a set to easily compare to current
+	TSet<int> NewlyHighlightedGroups;
+	NewlyHighlightedGroups.Append(HilightSelection.SelectedGroupIDs);
+
+	// See if we're currently highlighting any groups that we're not supposed to
+	if (!NewlyHighlightedGroups.Includes(CurrentlyHighlightedGroups))
+	{
+		DrawnTriangleSetComponent->Clear();
+		CurrentlyHighlightedGroups.Empty();
+	}
+
+	// See if we need to add any groups
+	if (!CurrentlyHighlightedGroups.Includes(NewlyHighlightedGroups))
+	{
+		// Add triangles for each new group
+		for (int Gid : HilightSelection.SelectedGroupIDs)
+		{
+			if (!CurrentlyHighlightedGroups.Contains(Gid))
+			{
+				for (int32 Tid : Topology->GetGroupTriangles(Gid))
+				{
+					// We use the triangle normals because the normal overlay isn't guaranteed to be valid as we edit the mesh
+					FVector3d TriangleNormal = Mesh->GetTriNormal(Tid);
+
+					FIndex3i VertIndices = Mesh->GetTriangle(Tid);
+					DrawnTriangleSetComponent->AddTriangle(FRenderableTriangle(HighlightedFaceMaterial,
+						FRenderableTriangleVertex((FVector)Mesh->GetVertex(VertIndices.A), (FVector2D)Mesh->GetVertexUV(VertIndices.A), (FVector)TriangleNormal, UE::Geometry::ToFColor(Mesh->GetVertexColor(VertIndices.A))),
+						FRenderableTriangleVertex((FVector)Mesh->GetVertex(VertIndices.B), (FVector2D)Mesh->GetVertexUV(VertIndices.B), (FVector)TriangleNormal, UE::Geometry::ToFColor(Mesh->GetVertexColor(VertIndices.B))),
+						FRenderableTriangleVertex((FVector)Mesh->GetVertex(VertIndices.C), (FVector2D)Mesh->GetVertexUV(VertIndices.C), (FVector)TriangleNormal, UE::Geometry::ToFColor(Mesh->GetVertexColor(VertIndices.C))) ));
+				}
+
+				CurrentlyHighlightedGroups.Add(Gid);
+			}
+		}//end iterating through groups
+	}//end if groups need to be added
+
+	return bHit;
+}
+
+
+
+
+bool UPolygonSelectionMechanic::HasSelection() const
+{
+	return PersistentSelection.IsEmpty() == false;
+}
+
+
+bool UPolygonSelectionMechanic::UpdateSelection(const FRay& WorldRay, FVector3d& LocalHitPositionOut, FVector3d& LocalHitNormalOut)
+{
+	FRay3d LocalRay(TargetTransform.InverseTransformPosition((FVector3d)WorldRay.Origin),
+		TargetTransform.InverseTransformVector((FVector3d)WorldRay.Direction));
+	UE::Geometry::Normalize(LocalRay.Direction);
+
+	bool bSelectionModified = false;
+	FVector3d LocalPosition, LocalNormal;
+	FGroupTopologySelection Selection;
+	FGroupTopologySelector::FSelectionSettings TopoSelectorSettings = GetTopoSelectorSettings(CameraState.bIsOrthographic);
+	if (TopoSelector->FindSelectedElement(TopoSelectorSettings, LocalRay, Selection, LocalPosition, LocalNormal))
+	{
+		LocalHitPositionOut = LocalPosition;
+		LocalHitNormalOut = LocalNormal;
+
+		if (Selection.SelectedEdgeIDs.Num() > 0 && Properties->bSelectEdgeRings && ShouldSelectEdgeRingsFunc())
+		{
+			TopoSelector->ExpandSelectionByEdgeRings(Selection);
+		}
+		if (Selection.SelectedEdgeIDs.Num() > 0 && Properties->bSelectEdgeLoops && ShouldSelectEdgeLoopsFunc())
+		{
+			TopoSelector->ExpandSelectionByEdgeLoops(Selection);
+		}
+
+		if (ShouldAddToSelectionFunc())
+		{
+			if (ShouldRemoveFromSelectionFunc())
+			{
+				PersistentSelection.Toggle(Selection);
+			}
+			else
+			{
+				PersistentSelection.Append(Selection);
+			}
+		}
+		else if (ShouldRemoveFromSelectionFunc())
+		{
+			PersistentSelection.Remove(Selection);
+		}
+		else
+		{
+			PersistentSelection = Selection;
+		}
+
+		bSelectionModified = true;
+	}
+	else
+	{
+		bSelectionModified = (PersistentSelection.IsEmpty() == false);
+		PersistentSelection.Clear();
+	}
+
+	if (bSelectionModified)
+	{
+		SelectionTimestamp++;
+		OnSelectionChanged.Broadcast();
+	}
+
+	return bSelectionModified;
+}
+
+
+void UPolygonSelectionMechanic::SetSelection(const FGroupTopologySelection& Selection, bool bBroadcast)
+{
+	PersistentSelection = Selection;
+	SelectionTimestamp++;
+	if (bBroadcast)
+	{
+		OnSelectionChanged.Broadcast();
+	}
+}
+
+
+void UPolygonSelectionMechanic::ClearSelection()
+{
+	PersistentSelection.Clear();
+	SelectionTimestamp++;
+	OnSelectionChanged.Broadcast();
+}
+
+FInputRayHit UPolygonSelectionMechanic::IsHitByClick(const FInputDeviceRay& ClickPos)
+{
+	if (!bIsEnabled)
+	{
+		return FInputRayHit(); // bHit is false
+	}
+
+	FHitResult OutHit;
+	FGroupTopologySelection Selection;
+	if (TopologyHitTest(ClickPos.WorldRay, OutHit, Selection, CameraState.bIsOrthographic))
+	{
+		return FInputRayHit(OutHit.Distance);
+	}
+
+	// We didn't hit selectable geometry. If we're using the marquee mechanic, give it a
+	// chance to capture. Otherwise still capture ourselves so that we are able to clear selections.
+	return MarqueeMechanic->IsEnabled() ?
+		FInputRayHit() // bHit is false
+		: FInputRayHit(TNumericLimits<float>::Max()); // bHit is true
+}
+
+void UPolygonSelectionMechanic::OnClicked(const FInputDeviceRay& ClickPos)
+{
+	// update selection
+	ParentTool->GetToolManager()->BeginUndoTransaction(LOCTEXT("SelectionChange", "Selection"));
+	BeginChange();
+
+	// This will fire off a OnSelectionChanged delegate.
+	UpdateSelection(ClickPos.WorldRay, LastClickedHitPosition, LastClickedHitNormal);
+
+	EndChangeAndEmitIfModified();
+	ParentTool->GetToolManager()->EndUndoTransaction();
+}
+
+FInputRayHit UPolygonSelectionMechanic::BeginHoverSequenceHitTest(const FInputDeviceRay& PressPos)
+{
+	FHitResult OutHit;
+	if (bIsEnabled && TopologyHitTest(PressPos.WorldRay, OutHit))
+	{
+		return FInputRayHit(OutHit.Distance);
+	}
+	return FInputRayHit(); // bHit is false
+}
+
+void UPolygonSelectionMechanic::OnBeginHover(const FInputDeviceRay& DevicePos)
+{
+}
+
+bool UPolygonSelectionMechanic::OnUpdateHover(const FInputDeviceRay& DevicePos)
+{
+	UpdateHighlight(DevicePos.WorldRay);
+	return true;
+}
+
+void UPolygonSelectionMechanic::OnEndHover()
+{
+	ClearHighlight();
+}
+
+void UPolygonSelectionMechanic::OnUpdateModifierState(int ModifierID, bool bIsOn)
+{
+	switch (ModifierID)
+	{
+	case ShiftModifierID:
+		bShiftToggle = bIsOn;
+		break;
+	case CtrlModifierID:
+		bCtrlToggle = bIsOn;
+		break;
+	}
+}
+
+void UPolygonSelectionMechanic::OnDragRectangleStarted()
+{
+	ParentTool->GetToolManager()->BeginUndoTransaction(LOCTEXT("SelectionChange", "Selection"));
+	BeginChange();
+
+	PreDragPersistentSelection = PersistentSelection;
+	PersistentSelection.Clear();
+}
+
+void UPolygonSelectionMechanic::OnDragRectangleChanged(const FCameraRectangle& CurrentRectangle)
+{
+	FGroupTopologySelection RectangleSelection;
+
+	FGroupTopologySelector::FSelectionSettings TopoSelectorSettings = GetTopoSelectorSettings(false);
+	TopoSelectorSettings.bIgnoreOcclusion = Properties->bMarqueeIgnoreOcclusion; // uses a separate setting for marquee
+	
+	TopoSelector->FindSelectedElement(TopoSelectorSettings, CurrentRectangle, TargetTransform,
+		RectangleSelection);
+
+	if (ShouldAddToSelectionFunc())
+	{
+		PersistentSelection = PreDragPersistentSelection;
+		if (ShouldRemoveFromSelectionFunc())
+		{
+			PersistentSelection.Toggle(RectangleSelection);
+		}
+		else
+		{
+			PersistentSelection.Append(RectangleSelection);
+		}
+	}
+	else if (ShouldRemoveFromSelectionFunc())
+	{
+		PersistentSelection = PreDragPersistentSelection;
+		PersistentSelection.Remove(RectangleSelection);
+	}
+	else
+	{
+		// Neither key pressed.
+		PersistentSelection = RectangleSelection;
+	}
+}
+
+void UPolygonSelectionMechanic::OnDragRectangleFinished()
+{
+	if (!(PersistentSelection == PreDragPersistentSelection))
+	{
+		SelectionTimestamp++;
+		OnSelectionChanged.Broadcast();
+	}
+
+	EndChangeAndEmitIfModified();
+	ParentTool->GetToolManager()->EndUndoTransaction();
+}
+
+void UPolygonSelectionMechanic::UpdateMarqueeEnabled()
+{
+	MarqueeMechanic->SetIsEnabled(
+		bIsEnabled 
+		&& Properties->bEnableMarquee 
+		&& (Properties->bSelectVertices || Properties->bSelectEdges));
+}
+
+
+void UPolygonSelectionMechanic::BeginChange()
+{
+	check(ActiveChange.IsValid() == false);
+	ActiveChange = MakeUnique<FPolygonSelectionMechanicSelectionChange>();
+	ActiveChange->Before = PersistentSelection;
+	ActiveChange->Timestamp = SelectionTimestamp;
+}
+
+TUniquePtr<FToolCommandChange> UPolygonSelectionMechanic::EndChange()
+{
+	check(ActiveChange.IsValid() == true);
+	ActiveChange->After = PersistentSelection;
+	if (SelectionTimestamp != ActiveChange->Timestamp)
+	{
+		return MoveTemp(ActiveChange);
+	}
+	ActiveChange = TUniquePtr<FPolygonSelectionMechanicSelectionChange>();
+	return TUniquePtr<FToolCommandChange>();
+}
+
+bool UPolygonSelectionMechanic::EndChangeAndEmitIfModified()
+{
+	check(ActiveChange.IsValid() == true);
+	ActiveChange->After = PersistentSelection;
+	if (SelectionTimestamp != ActiveChange->Timestamp)
+	{
+		GetParentTool()->GetToolManager()->EmitObjectChange(this, MoveTemp(ActiveChange),
+			LOCTEXT("SelectionChangeMessage", "Selection Change"));
+		return true;
+	}
+	ActiveChange = TUniquePtr<FPolygonSelectionMechanicSelectionChange>();
+	return false;
+}
+
+void UPolygonSelectionMechanic::GetClickedHitPosition(FVector3d& PositionOut, FVector3d& NormalOut) const
+{
+	PositionOut = LastClickedHitPosition;
+	NormalOut = LastClickedHitNormal;
+}
+
+FFrame3d UPolygonSelectionMechanic::GetSelectionFrame(bool bWorld, FFrame3d* InitialLocalFrame) const
+{
+	FFrame3d UseFrame;
+	if (PersistentSelection.IsEmpty() == false)
+	{
+		UseFrame = Topology->GetSelectionFrame(PersistentSelection, InitialLocalFrame);
+	}
+
+	if (bWorld)
+	{
+		UseFrame.Transform(TargetTransform);
+	}
+
+	return UseFrame;
+}
+
+FAxisAlignedBox3d UPolygonSelectionMechanic::GetSelectionBounds(bool bWorld) const
+{
+	if ( ! PersistentSelection.IsEmpty() )
+	{
+		if (bWorld)
+		{
+			return Topology->GetSelectionBounds(PersistentSelection, [this](const FVector3d& Pos) { return TargetTransform.TransformPosition(Pos); });
+		}
+		else
+		{
+			return Topology->GetSelectionBounds(PersistentSelection, [this](const FVector3d& Pos) { return Pos; });
+		}
+	}
+	else 
+	{
+		return Mesh->GetBounds();
+	}
+}
+
+
+
+void FPolygonSelectionMechanicSelectionChange::Apply(UObject* Object)
+{
+	UPolygonSelectionMechanic* Mechanic = Cast<UPolygonSelectionMechanic>(Object);
+	if (Mechanic)
+	{
+		Mechanic->PersistentSelection = After;
+		Mechanic->OnSelectionChanged.Broadcast();
+	}
+}
+void FPolygonSelectionMechanicSelectionChange::Revert(UObject* Object)
+{
+	UPolygonSelectionMechanic* Mechanic = Cast<UPolygonSelectionMechanic>(Object);
+	if (Mechanic)
+	{
+		Mechanic->PersistentSelection = Before;
+		Mechanic->OnSelectionChanged.Broadcast();
+	}
+}
+FString FPolygonSelectionMechanicSelectionChange::ToString() const
+{
+	return TEXT("FPolygonSelectionMechanicSelectionChange");
+}
+
+
+
+#undef LOCTEXT_NAMESPACE
\ No newline at end of file
diff --git a/Engine/Plugins/Experimental/MeshModelingToolsetExp/Source/ModelingComponents/Public/BaseTools/BaseCreateFromSelectedTool.h b/Engine/Plugins/Experimental/MeshModelingToolsetExp/Source/ModelingComponents/Public/BaseTools/BaseCreateFromSelectedTool.h
new file mode 100644
--- /dev/null
+++ b/Engine/Plugins/Experimental/MeshModelingToolsetExp/Source/ModelingComponents/Public/BaseTools/BaseCreateFromSelectedTool.h
@@ -0,0 +1,208 @@
+// Copyright Epic Games, Inc. All Rights Reserved.
+
+#pragma once
+
+#include "CoreMinimal.h"
+#include "MultiSelectionTool.h"
+#include "InteractiveToolBuilder.h"
+#include "MeshOpPreviewHelpers.h"
+#include "PropertySets/OnAcceptProperties.h"
+#include "PropertySets/CreateMeshObjectTypeProperties.h"
+#include "BaseCreateFromSelectedTool.generated.h"
+
+class UTransformGizmo;
+class UTransformProxy;
+
+/**
+ * ToolBuilder for UBaseCreateFromSelectedTool
+ */
+UCLASS()
+class MODELINGCOMPONENTS_API UBaseCreateFromSelectedToolBuilder : public UInteractiveToolBuilder
+{
+	GENERATED_BODY()
+
+public:
+	virtual bool CanBuildTool(const FToolBuilderState& SceneState) const override;
+	virtual UInteractiveTool* BuildTool(const FToolBuilderState& SceneState) const override;
+
+public:
+	virtual TOptional<int32> MaxComponentsSupported() const { return TOptional<int32>(); }
+	virtual int32 MinComponentsSupported() const { return 1; }
+	
+	// subclass must override!
+	virtual UBaseCreateFromSelectedTool* MakeNewToolInstance(UObject* Outer) const { check(false); return nullptr; }
+
+protected:
+	virtual const FToolTargetTypeRequirements& GetTargetRequirements() const override;
+};
+
+
+UENUM()
+enum class EBaseCreateFromSelectedTargetType
+{
+	/** Create a new asset containing the result mesh */
+	NewAsset,
+	/** Store the result mesh in the first selected input asset */
+	FirstInputAsset,
+	/** Store the result mesh in the last selected input asset */
+	LastInputAsset
+};
+
+
+UCLASS()
+class MODELINGCOMPONENTS_API UBaseCreateFromSelectedHandleSourceProperties : public UOnAcceptHandleSourcesProperties
+{
+	GENERATED_BODY()
+public:
+	/** Where should the output mesh produced by this operation be stored */
+	UPROPERTY(EditAnywhere, Category = ToolOutputOptions)
+	EBaseCreateFromSelectedTargetType WriteOutputTo = EBaseCreateFromSelectedTargetType::NewAsset;
+
+	/** Base name for newly-generated asset */
+	UPROPERTY(EditAnywhere, Category = ToolOutputOptions, meta = (TransientToolProperty, EditCondition = "WriteOutputTo == EBaseCreateFromSelectedTargetType::NewAsset", EditConditionHides))
+	FString OutputName;
+
+	/** Name of asset that will be updated */
+	UPROPERTY(VisibleAnywhere, Category = ToolOutputOptions, meta = (TransientToolProperty, EditCondition = "WriteOutputTo != EBaseCreateFromSelectedTargetType::NewAsset", EditConditionHides))
+	FString OutputAsset;
+};
+
+
+
+
+
+/**
+ * Properties of UI to adjust input meshes
+ */
+UCLASS()
+class MODELINGCOMPONENTS_API UTransformInputsToolProperties : public UInteractiveToolPropertySet
+{
+	GENERATED_BODY()
+
+public:
+	/** Show UI to allow changing translation, rotation and scale of input meshes */
+	UPROPERTY(EditAnywhere, Category = Transform)
+	bool bShowTransformUI = true;
+};
+
+
+/**
+ * UBaseCreateFromSelectedTool is a base Tool (must be subclassed) 
+ * that provides support for common functionality in tools that create a new mesh from a selection of one or more existing meshes
+ */
+UCLASS()
+class MODELINGCOMPONENTS_API UBaseCreateFromSelectedTool : public UMultiSelectionTool, public UE::Geometry::IDynamicMeshOperatorFactory
+{
+	GENERATED_BODY()
+protected:
+	using FTransform3d = UE::Geometry::FTransform3d;
+	using FFrame3d = UE::Geometry::FFrame3d;
+	using FRay3d = UE::Geometry::FRay3d;
+public:
+	UBaseCreateFromSelectedTool() = default;
+
+	virtual void SetWorld(UWorld* World);
+
+	//
+	// InteractiveTool API - generally does not need to be modified by subclasses
+	//
+
+	virtual void Setup() override;
+	virtual void Shutdown(EToolShutdownType ShutdownType) override;
+
+	virtual void OnTick(float DeltaTime) override;
+
+	virtual bool HasCancel() const override { return true; }
+	virtual bool HasAccept() const override { return true; }
+	virtual bool CanAccept() const override;
+
+	virtual void OnPropertyModified(UObject* PropertySet, FProperty* Property) override;
+
+	
+
+protected:
+
+	//
+	// UBaseCreateFromSelectedTool API - subclasses typically implement these functions
+	//
+
+	/**
+	 * After preview is created, this is called to convert inputs and set preview materials
+	 * (grouped together because materials may come from inputs)
+	 * Subclasses should always implement this.
+	 * @param bSetPreviewMesh If true, function may try to set an initial "early" preview mesh to have some initial surface on tool start.  (Not all tools will actually create this.)
+	 *						  This boolean is here in case a subclass needs to call this setup function again later (e.g. to change the materials used), when it won't need or want the preview surface to be created
+	 */
+	virtual void ConvertInputsAndSetPreviewMaterials(bool bSetPreviewMesh = true) { check(false);  }
+
+	/** overload to initialize any added properties in subclasses; called during setup */
+	virtual void SetupProperties() {}
+
+	/** overload to save any added properties in the subclasses; called on shutdown */
+	virtual void SaveProperties() {}
+
+	/** optional overload to set callbacks on preview, e.g. to visualize results; called after preview is created. */
+	virtual void SetPreviewCallbacks() {}
+
+	/** Return the name to be used for generated assets.  Note: Asset name will be prefixed by source actor name if only actor was selected. */
+	virtual FString GetCreatedAssetName() const { return TEXT("Generated"); }
+
+	/** Return the name of the action to be used in the Undo stack */
+	virtual FText GetActionName() const;
+
+	/** Return the materials to be used on the output mesh on tool accept; defaults to the materials set on the preview */
+	virtual TArray<UMaterialInterface*> GetOutputMaterials() const;
+
+
+
+	/**
+	 * IDynamicMeshOperatorFactory implementation that subclass must override and implement
+	 */
+	virtual TUniquePtr<UE::Geometry::FDynamicMeshOperator> MakeNewOperator() override
+	{
+		check(false);
+		return TUniquePtr<UE::Geometry::FDynamicMeshOperator>();
+	}
+
+protected:
+
+	/** Helper to build asset names */
+	FString PrefixWithSourceNameIfSingleSelection(const FString& AssetName) const;
+
+	// Helpers for managing transform gizoms; typically do not need to be overloaded
+	virtual void UpdateGizmoVisibility();
+	virtual void SetTransformGizmos();
+	virtual void TransformChanged(UTransformProxy* Proxy, FTransform Transform);
+
+	// Helper to generate assets when a result is accepted; typically does not need to be overloaded
+	virtual void GenerateAsset(const FDynamicMeshOpResult& Result);
+
+	// Helper to generate assets when a result is accepted; typically does not need to be overloaded
+	virtual void UpdateAsset(const FDynamicMeshOpResult& Result, UToolTarget* Target);
+
+	// Which of the transform gizmos to hide, or -1 if all gizmos can be shown
+	virtual int32 GetHiddenGizmoIndex() const;
+
+protected:
+
+	UPROPERTY()
+	TObjectPtr<UTransformInputsToolProperties> TransformProperties;
+
+	UPROPERTY()
+	TObjectPtr<UCreateMeshObjectTypeProperties> OutputTypeProperties;
+
+	UPROPERTY()
+	TObjectPtr<UBaseCreateFromSelectedHandleSourceProperties> HandleSourcesProperties;
+
+	UPROPERTY()
+	TObjectPtr<UMeshOpPreviewWithBackgroundCompute> Preview;
+
+	UPROPERTY()
+	TArray<TObjectPtr<UTransformProxy>> TransformProxies;
+
+	UPROPERTY()
+	TArray<TObjectPtr<UTransformGizmo>> TransformGizmos;
+
+	UWorld* TargetWorld = nullptr;
+};
+
diff --git a/Engine/Plugins/Experimental/MeshModelingToolsetExp/Source/ModelingComponents/Public/BaseTools/BaseVoxelTool.h b/Engine/Plugins/Experimental/MeshModelingToolsetExp/Source/ModelingComponents/Public/BaseTools/BaseVoxelTool.h
new file mode 100644
--- /dev/null
+++ b/Engine/Plugins/Experimental/MeshModelingToolsetExp/Source/ModelingComponents/Public/BaseTools/BaseVoxelTool.h
@@ -0,0 +1,42 @@
+// Copyright Epic Games, Inc. All Rights Reserved.
+
+#pragma once
+
+#include "CoreMinimal.h"
+#include "MultiSelectionTool.h"
+#include "BaseTools/BaseCreateFromSelectedTool.h"
+#include "PropertySets/VoxelProperties.h"
+
+#include "BaseVoxelTool.generated.h"
+
+PREDECLARE_USE_GEOMETRY_CLASS(FDynamicMesh3);
+
+/**
+ * Base for Voxel tools
+ */
+UCLASS()
+class MODELINGCOMPONENTS_API UBaseVoxelTool : public UBaseCreateFromSelectedTool
+{
+	GENERATED_BODY()
+
+protected:
+
+	/** Sets up VoxProperties; typically need to overload and call "Super::SetupProperties();" */
+	virtual void SetupProperties() override;
+
+	/** Saves VoxProperties; typically need to overload and call "Super::SaveProperties();" */
+	virtual void SaveProperties() override;
+
+	/** Sets up the default preview and converted inputs for voxel tools; Typically do not need to overload */
+	virtual void ConvertInputsAndSetPreviewMaterials(bool bSetPreviewMesh = true) override;
+
+	/** Sets the output material to the default "world grid" material */
+	virtual TArray<UMaterialInterface*> GetOutputMaterials() const override;
+
+
+	UPROPERTY()
+	TObjectPtr<UVoxelProperties> VoxProperties;
+
+	TArray<TSharedPtr<FDynamicMesh3, ESPMode::ThreadSafe>> OriginalDynamicMeshes;
+};
+
diff --git a/Engine/Plugins/Experimental/MeshModelingToolsetExp/Source/ModelingComponents/Public/Drawing/PolyEditPreviewMesh.h b/Engine/Plugins/Experimental/MeshModelingToolsetExp/Source/ModelingComponents/Public/Drawing/PolyEditPreviewMesh.h
new file mode 100644
--- /dev/null
+++ b/Engine/Plugins/Experimental/MeshModelingToolsetExp/Source/ModelingComponents/Public/Drawing/PolyEditPreviewMesh.h
@@ -0,0 +1,84 @@
+// Copyright Epic Games, Inc. All Rights Reserved.
+
+#pragma once
+
+#include "PreviewMesh.h"
+#include "DynamicMesh/DynamicMesh3.h"
+#include "DynamicMesh/DynamicMeshAABBTree3.h"
+#include "DynamicSubmesh3.h"
+#include "EdgeLoop.h"
+#include "PolyEditPreviewMesh.generated.h"
+
+/**
+ * UPolyEditPreviewMesh is a variant of UPreviewMesh intended for use as a 'live preview' of
+ * a mesh creation/editing operation. The class supports initializing the preview mesh in various
+ * ways, generally as a submesh of a base mesh.
+ */
+UCLASS(Transient)
+class MODELINGCOMPONENTS_API UPolyEditPreviewMesh : public UPreviewMesh
+{
+	GENERATED_BODY()
+	using FTransform3d = UE::Geometry::FTransform3d;	
+public:
+
+	//
+	// "Static Type" is just a static mesh
+	//
+
+	void InitializeStaticType(const FDynamicMesh3* SourceMesh, const TArray<int32>& Triangles,
+		const FTransform3d* MeshTransform = nullptr);
+	void UpdateStaticType(TFunctionRef<void(FDynamicMesh3&)> UpdateMeshFunc, bool bFullRecalculate);
+	void MakeStaticTypeTargetMesh(FDynamicMesh3& TargetMesh);
+
+	//
+	// "Extrude Type" duplicates the input faces, offsets them using FExtrudeMesh, and stitches them together
+	//
+
+	void InitializeExtrudeType(const FDynamicMesh3* SourceMesh, const TArray<int32>& Triangles,
+		const FVector3d& TransformedOffsetDirection,
+		const FTransform3d* MeshTransform = nullptr,
+		bool bDeleteExtrudeBaseFaces = true);
+	void InitializeExtrudeType(FDynamicMesh3&& BaseMesh,
+		const FVector3d& TransformedOffsetDirection,
+		const FTransform3d* MeshTransform = nullptr,
+		bool bDeleteExtrudeBaseFaces = true);
+	/** Update extrude-type preview mesh by moving existing offset vertices */
+	void UpdateExtrudeType(double NewOffset, bool bUseNormalDirection = false);
+	/** Update extrude-type preview mesh by moving existing offset vertices along their connected triangle normals and then averaging positions  */
+	void UpdateExtrudeType_FaceNormalAvg(double NewOffset);
+	/** Update extrude-type preview mesh using external function. if bFullRecalculate, mesh is re-initialized w/ initial extrusion patch SourceMesh before calling function */
+	void UpdateExtrudeType(TFunctionRef<void(FDynamicMesh3&)> UpdateMeshFunc, bool bFullRecalculate);
+	/** Make a hit-target mesh that is an infinite extrusion along extrude direction. If bUseNormalDirection, use per-vertex normals instead */
+	void MakeExtrudeTypeHitTargetMesh(FDynamicMesh3& TargetMesh, bool bUseNormalDirection = false);
+
+
+	//
+	// "Inset Type" duplicates the input faces and insets them using FInsetMeshRegion
+	//
+
+	void InitializeInsetType(const FDynamicMesh3* SourceMesh, const TArray<int32>& Triangles,
+		const FTransform3d* MeshTransform = nullptr);
+	void UpdateInsetType(double NewOffset, bool bReproject = false, double Softness = 0.0, double AreaScaleT = 1.0, bool bBoundaryOnly = false);
+	void MakeInsetTypeTargetMesh(FDynamicMesh3& TargetMesh);
+
+	const FDynamicMesh3& GetInitialPatchMesh() const { return InitialEditPatch; }
+
+	// must be non-const because query functions are non-const
+	UE::Geometry::FDynamicMeshAABBTree3& GetInitialPatchMeshSpatial() { return InitialEditPatchBVTree; }
+
+protected:
+	TUniquePtr<UE::Geometry::FDynamicSubmesh3> ActiveSubmesh;
+	FDynamicMesh3 InitialEditPatch;
+	UE::Geometry::FDynamicMeshAABBTree3 InitialEditPatchBVTree;
+
+	TArray<int32> EditVertices;
+	TArray<FVector3d> InitialPositions;
+	TArray<FVector3d> InitialNormals;
+	TArray<FVector3d> InitialTriNormals;
+	TMap<int32,int32> ExtrudeToInitialVerts;
+
+	FVector3d InputDirection;
+
+	UE::Geometry::FTransform3d MeshTransform;
+	bool bHaveMeshTransform;
+};
\ No newline at end of file
diff --git a/Engine/Plugins/Experimental/MeshModelingToolsetExp/Source/ModelingComponents/Public/Drawing/UVLayoutPreview.h b/Engine/Plugins/Experimental/MeshModelingToolsetExp/Source/ModelingComponents/Public/Drawing/UVLayoutPreview.h
new file mode 100644
--- /dev/null
+++ b/Engine/Plugins/Experimental/MeshModelingToolsetExp/Source/ModelingComponents/Public/Drawing/UVLayoutPreview.h
@@ -0,0 +1,160 @@
+// Copyright Epic Games, Inc. All Rights Reserved.
+
+#pragma once
+
+#include "CoreMinimal.h"
+#include "PreviewMesh.h"
+#include "TriangleSetComponent.h"
+#include "InteractiveTool.h"
+#include "UVLayoutPreview.generated.h"
+
+PREDECLARE_USE_GEOMETRY_CLASS(FDynamicMesh3);
+
+/**
+ * Where should in-viewport UVLayoutPreview be shown, relative to target object
+ */
+UENUM()
+enum class EUVLayoutPreviewSide
+{
+	Left = 0,
+	Right = 1
+};
+
+
+/**
+ * Visualization settings for UV layout preview
+ */
+UCLASS()
+class MODELINGCOMPONENTS_API UUVLayoutPreviewProperties : public UInteractiveToolPropertySet
+{
+	GENERATED_BODY()
+public:
+	/** Should be UV Layout be shown */
+	UPROPERTY(EditAnywhere, Category = UVLayoutPreview)
+	bool bVisible = true;
+
+	/** World-space scaling factor on the UV Layout */
+	UPROPERTY(EditAnywhere, Category = UVLayoutPreview, meta = (UIMin = "0.1", UIMax = "10.0", ClampMin = "0.0001", ClampMax = "1000"))
+	float ScaleFactor = 1.0;
+
+	/** Where should the UV layout be positioned relative to the target object, relative to camera */
+	UPROPERTY(EditAnywhere, Category = UVLayoutPreview)
+	EUVLayoutPreviewSide WhichSide = EUVLayoutPreviewSide::Right;
+
+	/** If true, wireframe is shown for the UV layout */
+	UPROPERTY(EditAnywhere, Category = UVLayoutPreview)
+	bool bShowWireframe = true;
+
+	UPROPERTY(EditAnywhere, Category = UVLayoutPreview)
+	FVector2D Shift = FVector2D(1.0, 0.5);
+};
+
+
+
+
+/**
+ * UUVLayoutPreview is a utility object that creates and manages a 3D plane on which a UV layout
+ * for a 3D mesh is rendered. The UV layout
+ */
+UCLASS(Transient)
+class MODELINGCOMPONENTS_API UUVLayoutPreview : public UObject
+{
+	GENERATED_BODY()
+
+public:
+	virtual ~UUVLayoutPreview();
+
+
+	/**
+	 * Create preview mesh in the World with the given transform
+	 */
+	void CreateInWorld(UWorld* World);
+
+	/**
+	 * Remove and destroy preview mesh
+	 */
+	void Disconnect();
+
+
+	/**
+	 * Configure material set for UV-space preview mesh
+	 */
+	void SetSourceMaterials(const FComponentMaterialSet& MaterialSet);
+
+	/**
+	 * Specify the current world transform/bounds for the target object. 
+	 * UV layout preview is positioned relative to this box
+	 */
+	void SetSourceWorldPosition(FTransform WorldTransform, FBox WorldBounds);
+
+	/**
+	 * Update the current camera state, used to auto-position the UV layout preview
+	 */
+	void SetCurrentCameraState(const FViewCameraState& CameraState);
+
+	/**
+	 * Notify the UV Layout Preview that the source UVs have been modified
+	 */
+	void UpdateUVMesh(const FDynamicMesh3* SourceMesh, int32 SourceUVLayer = 0);
+
+	/**
+	 * Tick the UV Layout Preview, allowing it to upodate various settings
+	 */
+	void OnTick(float DeltaTime);
+
+	/**
+	 * Render the UV Layout Preview, allowing it to upodate various settings
+	 */
+	void Render(IToolsContextRenderAPI* RenderAPI);
+
+
+	/**
+	 * Set the transform on the UV Layout preview mesh
+	 */
+	void SetTransform(const FTransform& UseTransform);
+
+	/**
+	 * Set the visibility of the UV Layout  preview mesh
+	 */
+	void SetVisible(bool bVisible);
+
+
+
+public:
+	/** Visualization settings */
+	UPROPERTY()
+	TObjectPtr<UUVLayoutPreviewProperties> Settings;
+
+	/** PreviewMesh is initialized with a copy of an input mesh with UVs mapped to position, ie such that (X,Y,Z) = (U,V,0) */
+	UPROPERTY()
+	TObjectPtr<UPreviewMesh> PreviewMesh;
+
+	/** Set of additional triangles to draw, eg for backing rectangle, etc */
+	UPROPERTY()
+	TObjectPtr<UTriangleSetComponent> TriangleComponent;
+
+	/** Configure whether the backing rectangle should be shown */
+	UPROPERTY()
+	bool bShowBackingRectangle = true;
+
+	/** Configure the backing rectangle material */
+	UPROPERTY()
+	TObjectPtr<UMaterialInterface> BackingRectangleMaterial = nullptr;
+
+
+protected:
+	void RecalculatePosition();
+
+	FComponentMaterialSet SourceMaterials;
+
+	UE::Geometry::FFrame3d SourceObjectFrame;
+	UE::Geometry::FAxisAlignedBox3d SourceObjectWorldBounds;
+
+	UE::Geometry::FFrame3d CurrentWorldFrame;
+
+	FViewCameraState CameraState;
+
+	bool bSettingsModified = false;
+
+	float GetCurrentScale();
+};
\ No newline at end of file
diff --git a/Engine/Plugins/Experimental/MeshModelingToolsetExp/Source/ModelingComponents/Public/Mechanics/CollectSurfacePathMechanic.h b/Engine/Plugins/Experimental/MeshModelingToolsetExp/Source/ModelingComponents/Public/Mechanics/CollectSurfacePathMechanic.h
new file mode 100644
--- /dev/null
+++ b/Engine/Plugins/Experimental/MeshModelingToolsetExp/Source/ModelingComponents/Public/Mechanics/CollectSurfacePathMechanic.h
@@ -0,0 +1,105 @@
+// Copyright Epic Games, Inc. All Rights Reserved.
+
+#pragma once
+
+#include "CoreMinimal.h"
+#include "InteractionMechanic.h"
+#include "DynamicMesh/DynamicMesh3.h"
+#include "DynamicMesh/DynamicMeshAABBTree3.h"
+#include "ToolDataVisualizer.h"
+#include "CollectSurfacePathMechanic.generated.h"
+
+using UE::Geometry::FDynamicMesh3;
+
+enum class ECollectSurfacePathDoneMode
+{
+	SnapCloseLoop,
+	SnapDoubleClick,
+	SnapDoubleClickOrCloseLoop,
+	ExternalLambda,
+	FixedNumPoints
+};
+
+
+
+/**
+ */
+UCLASS()
+class MODELINGCOMPONENTS_API UCollectSurfacePathMechanic : public UInteractionMechanic
+{
+	GENERATED_BODY()
+public:
+	using FFrame3d = UE::Geometry::FFrame3d;
+	using FRay3d = UE::Geometry::FRay3d;
+
+	TUniqueFunction<bool()> IsDoneFunc = nullptr;
+
+	double ConstantSnapDistance = 10.0f;
+	TUniqueFunction<bool(FVector3d, FVector3d)> SpatialSnapPointsFunc;
+
+	bool bSnapToTargetMeshVertices = false;
+
+	// tfunc to emit changes to...
+
+	TArray<FFrame3d> HitPath;
+
+	FFrame3d PreviewPathPoint;
+
+
+	FToolDataVisualizer PathDrawer;
+	FLinearColor PathColor;
+	FLinearColor PreviewColor;
+	FLinearColor PathCompleteColor;
+	bool bDrawPath = true;
+
+
+public:
+	UCollectSurfacePathMechanic();
+
+	virtual void Setup(UInteractiveTool* ParentTool) override;
+	virtual void Shutdown() override;
+	virtual void Render(IToolsContextRenderAPI* RenderAPI) override;
+
+	/**
+	 * Set the hit target mesh.
+	 */
+	virtual void InitializeMeshSurface(FDynamicMesh3&& TargetSurfaceMesh);
+	virtual void InitializePlaneSurface(const FFrame3d& TargetPlane);
+
+	virtual void SetFixedNumPointsMode(int32 NumPoints);
+	virtual void SetDrawClosedLoopMode();
+	virtual void SetCloseWithLambdaMode();
+	virtual void SetDoubleClickOrCloseLoopMode();
+
+
+	virtual bool IsHitByRay(const FRay3d& Ray, FFrame3d& HitPoint);
+	virtual bool UpdatePreviewPoint(const FRay3d& Ray);
+	virtual bool TryAddPointFromRay(const FRay3d& Ray);
+
+	virtual bool PopLastPoint();
+
+	virtual bool IsDone() const;
+
+	/** Whether the path was finished by the user clicking on the first point */
+	bool LoopWasClosed() const
+	{
+		return bLoopWasClosed;
+	}
+
+protected:
+	FDynamicMesh3 TargetSurface;
+	UE::Geometry::FDynamicMeshAABBTree3 TargetSurfaceAABB;
+
+	FFrame3d TargetPlane;
+	bool bHaveTargetPlane;
+
+	bool RayToPathPoint(const FRay3d& Ray, FFrame3d& PointOut, bool bEnableSnapping);
+
+	ECollectSurfacePathDoneMode DoneMode = ECollectSurfacePathDoneMode::SnapDoubleClick;
+	int32 FixedPointTargetCount = 0;
+	bool bCurrentPreviewWillComplete = false;
+	bool bGeometricCloseOccurred = false;
+	bool CheckGeometricClosure(const FFrame3d& Point, bool* bLoopWasClosedOut = nullptr);
+
+	bool bLoopWasClosed = false;
+};
diff --git a/Engine/Plugins/Experimental/MeshModelingToolsetExp/Source/ModelingComponents/Public/Mechanics/ConstructionPlaneMechanic.h b/Engine/Plugins/Experimental/MeshModelingToolsetExp/Source/ModelingComponents/Public/Mechanics/ConstructionPlaneMechanic.h
new file mode 100644
--- /dev/null
+++ b/Engine/Plugins/Experimental/MeshModelingToolsetExp/Source/ModelingComponents/Public/Mechanics/ConstructionPlaneMechanic.h
@@ -0,0 +1,80 @@
+// Copyright Epic Games, Inc. All Rights Reserved.
+
+#pragma once
+
+#include "CoreMinimal.h"
+#include "InteractionMechanic.h"
+#include "ToolDataVisualizer.h"
+#include "FrameTypes.h"
+#include "BaseBehaviors/BehaviorTargetInterfaces.h"
+#include "Selection/SelectClickedAction.h"
+
+#include "ConstructionPlaneMechanic.generated.h"
+
+class UTransformGizmo;
+class UTransformProxy;
+class IClickBehaviorTarget;
+class USingleClickInputBehavior;
+
+/**
+ * UConstructionPlaneMechanic implements an interaction in which a 3D plane can be
+ * positioned using the standard 3D Gizmo, or placed at hit-locations in the existing scene.
+ * A grid in the plane can optionally be rendered.
+ */
+UCLASS()
+class MODELINGCOMPONENTS_API UConstructionPlaneMechanic : public UInteractionMechanic
+{
+	GENERATED_BODY()
+public:
+
+	/** Replace this to externally control if plane can be updated */
+	TUniqueFunction<bool()> CanUpdatePlaneFunc = []() { return true; };
+
+	bool bShowGrid = true;
+
+	UE::Geometry::FFrame3d Plane;
+
+	DECLARE_MULTICAST_DELEGATE(OnConstructionPlaneChangedEvent);
+	OnConstructionPlaneChangedEvent OnPlaneChanged;
+
+
+public:
+
+	virtual void Setup(UInteractiveTool* ParentTool) override;
+	virtual void Shutdown() override;
+	virtual void Render(IToolsContextRenderAPI* RenderAPI) override;
+	virtual void Tick(float DeltaTime) override;
+
+	void Initialize(UWorld* TargetWorld, const UE::Geometry::FFrame3d& InitialPlane);
+
+	void SetDrawPlaneFromWorldPos(const FVector3d& Position, const FVector3d& Normal, bool bIgnoreNormal);
+
+	/** 
+	 * Sets the plane without broadcasting OnPlaneChanged. Useful when the user of the tool wants to change
+	 * the plane through some other means. Better than setting the Plane field directly because this function
+	 * properly deals with the gizmo.
+	 */
+	void SetPlaneWithoutBroadcast(const UE::Geometry::FFrame3d& Plane);
+
+	void UpdateClickPriority(FInputCapturePriority NewPriority);
+
+public:
+	UPROPERTY()
+	TObjectPtr<UTransformGizmo> PlaneTransformGizmo;
+	
+	UPROPERTY()
+	TObjectPtr<UTransformProxy> PlaneTransformProxy;
+
+	/** 
+	 * This is the behavior target used for the Ctrl+click behavior that sets the plane
+	 * in the world, exposed here so that the user can modify it after Setup() if needed.
+	 * By default, Setup() will have it call SetDrawPlaneFromWorldPos.
+	 */
+	TUniquePtr<FSelectClickedAction> SetPlaneCtrlClickBehaviorTarget;
+
+protected:
+	UPROPERTY()
+	TObjectPtr<USingleClickInputBehavior> ClickToSetPlaneBehavior;
+
+	void TransformChanged(UTransformProxy* Proxy, FTransform Transform);
+};
\ No newline at end of file
diff --git a/Engine/Plugins/Experimental/MeshModelingToolsetExp/Source/ModelingComponents/Public/Mechanics/SpaceCurveDeformationMechanic.h b/Engine/Plugins/Experimental/MeshModelingToolsetExp/Source/ModelingComponents/Public/Mechanics/SpaceCurveDeformationMechanic.h
new file mode 100644
--- /dev/null
+++ b/Engine/Plugins/Experimental/MeshModelingToolsetExp/Source/ModelingComponents/Public/Mechanics/SpaceCurveDeformationMechanic.h
@@ -0,0 +1,266 @@
+// Copyright Epic Games, Inc. All Rights Reserved.
+
+#pragma once
+
+#include "CoreMinimal.h"
+
+#include "BaseBehaviors/BehaviorTargetInterfaces.h"
+#include "FrameTypes.h"
+#include "InteractionMechanic.h"
+#include "InteractiveToolChange.h"
+#include "Snapping/PointPlanarSnapSolver.h"
+#include "Spatial/GeometrySet3.h"
+#include "ToolContextInterfaces.h" //FViewCameraState
+#include "VectorTypes.h"
+#include "VectorUtil.h"
+
+#include "SpaceCurveDeformationMechanic.generated.h"
+
+class APreviewGeometryActor;
+class ULineSetComponent;
+class UMouseHoverBehavior;
+class UPointSetComponent;
+class USingleClickInputBehavior;
+class UTransformGizmo;
+class UTransformProxy;
+
+
+
+
+class MODELINGCOMPONENTS_API FSpaceCurveSource
+{
+public:
+	TUniqueFunction<int32()> GetPointCount;
+	TUniqueFunction<UE::Geometry::FFrame3d(int32)> GetPoint;
+	TUniqueFunction<bool()> IsLoop;
+};
+
+
+
+
+
+UENUM()
+enum class ESpaceCurveControlPointTransformMode
+{
+	Shared,
+	PerVertex
+};
+
+
+UENUM()
+enum class ESpaceCurveControlPointOriginMode
+{
+	Shared,
+	First,
+	Last
+};
+
+UENUM()
+enum class ESpaceCurveControlPointFalloffType
+{
+	Linear,
+	Smooth
+};
+
+
+
+UCLASS()
+class MODELINGCOMPONENTS_API USpaceCurveDeformationMechanicPropertySet : public UInteractiveToolPropertySet
+{
+	GENERATED_BODY()
+
+public:
+	UPROPERTY(EditAnywhere, Category = Options)
+	ESpaceCurveControlPointTransformMode TransformMode = ESpaceCurveControlPointTransformMode::PerVertex;
+
+	UPROPERTY(EditAnywhere, Category = Options)
+	ESpaceCurveControlPointOriginMode TransformOrigin = ESpaceCurveControlPointOriginMode::First;
+
+	UPROPERTY(EditAnywhere, Category = Options, meta = (UIMin = "0", UIMax = "1"))
+	float Softness = 0.5;
+
+	UPROPERTY(EditAnywhere, Category = Options)
+	ESpaceCurveControlPointFalloffType SoftFalloff = ESpaceCurveControlPointFalloffType::Smooth;
+};
+
+
+
+
+
+/**
+
+ */
+UCLASS()
+class MODELINGCOMPONENTS_API USpaceCurveDeformationMechanic : public UInteractionMechanic, public IClickBehaviorTarget, public IHoverBehaviorTarget
+{
+	GENERATED_BODY()
+public:
+
+	// Behaviors used for moving points around and hovering them
+	UPROPERTY()
+	TObjectPtr<USingleClickInputBehavior> ClickBehavior = nullptr;
+	UPROPERTY()
+	TObjectPtr<UMouseHoverBehavior> HoverBehavior = nullptr;
+
+	// This delegate is called every time the control point sequence is altered.
+	DECLARE_MULTICAST_DELEGATE(OnPointsChangedEvent);
+	OnPointsChangedEvent OnPointsChanged;
+
+
+	void SetCurveSource(TSharedPtr<FSpaceCurveSource> CurveSource);
+	void ClearCurveSource();
+
+	void ClearSelection();
+	void SelectionGrowToNext();
+	void SelectionGrowToPrev();
+	void SelectionGrowToEnd();
+	void SelectionGrowToStart();
+	void SelectionFill();
+	void SelectionClear();
+
+	void GetCurrentCurvePoints(TArray<UE::Geometry::FFrame3d>& PointsOut) { PointsOut = CurvePoints; }
+
+
+	// Some other standard functions
+	virtual ~USpaceCurveDeformationMechanic();
+	virtual void Setup(UInteractiveTool* ParentTool) override;
+	virtual void Shutdown() override;
+	void SetWorld(UWorld* World);
+	virtual void Render(IToolsContextRenderAPI* RenderAPI) override;
+	virtual void Tick(float DeltaTime) override;
+
+	// IClickBehaviorTarget implementation
+	virtual FInputRayHit IsHitByClick(const FInputDeviceRay& ClickPos) override;
+	virtual void OnClicked(const FInputDeviceRay& ClickPos) override;
+
+	// IHoverBehaviorTarget implementation
+	virtual FInputRayHit BeginHoverSequenceHitTest(const FInputDeviceRay& PressPos) override;
+	virtual void OnBeginHover(const FInputDeviceRay& DevicePos) override;
+	virtual bool OnUpdateHover(const FInputDeviceRay& DevicePos) override;
+	virtual void OnEndHover() override;
+
+	// IModifierToggleBehaviorTarget implementation, inherited through IClickBehaviorTarget
+	virtual void OnUpdateModifierState(int ModifierID, bool bIsOn) override;
+
+public:
+	UPROPERTY()
+	TObjectPtr<USpaceCurveDeformationMechanicPropertySet> TransformProperties;
+
+protected:
+
+	TSharedPtr<FSpaceCurveSource> CurveSource;
+	TSharedPtr<FSpaceCurveSource> EmptyCurveSource;
+
+	TArray<UE::Geometry::FFrame3d> CurvePoints;
+
+	// Used for spatial queries
+	UE::Geometry::FGeometrySet3 GeometrySet;
+
+	bool bSpatialValid = false;
+	void UpdateSpatial();
+
+	/** Used for displaying points/segments */
+	UPROPERTY()
+	TObjectPtr<APreviewGeometryActor> PreviewGeometryActor;
+	UPROPERTY()
+	TObjectPtr<UPointSetComponent> RenderPoints;
+	UPROPERTY()
+	TObjectPtr<ULineSetComponent> RenderSegments;
+
+	bool bRenderGeometryValid = false;
+	void UpdateRenderGeometry();
+
+	// Variables for drawing
+	FColor NormalCurveColor;
+	FColor CurrentSegmentsColor;
+	FColor CurrentPointsColor;
+	float SegmentsThickness;
+	float PointsSize;
+	float DepthBias;
+	FColor PreviewColor;
+	FColor HoverColor;
+	FColor SelectedColor;
+
+	// Support for Shift and Ctrl toggles
+	bool bAddToSelectionToggle = false;
+	int32 ShiftModifierId = 1;
+	int32 CtrlModifierId = 2;
+
+	// Support for gizmo. Since the points aren't individual components, we don't actually use UTransformProxy
+	// for the transform forwarding- we just use it for the callbacks.
+	UPROPERTY()
+	TObjectPtr<UTransformProxy> PointTransformProxy;
+	UPROPERTY()
+	TObjectPtr<UTransformGizmo> PointTransformGizmo;
+
+	// Used to make it easy to tell whether the gizmo was moved by the user or by undo/redo or
+	// some other change that we shoulnd't respond to. Basing our movement undo/redo on the
+	// gizmo turns out to be quite a pain, though may someday be easier if the transform proxy
+	// is able to manage arbitrary objects.
+	bool bGizmoBeingDragged = false;
+
+	// Callbacks we'll receive from the gizmo proxy
+	void GizmoTransformChanged(UTransformProxy* Proxy, FTransform Transform);
+	void GizmoTransformStarted(UTransformProxy* Proxy);
+	void GizmoTransformEnded(UTransformProxy* Proxy);
+
+	// Support for hovering
+	FViewCameraState CameraState;
+	TFunction<bool(const FVector3d&, const FVector3d&)> GeometrySetToleranceTest;
+	int32 HoveredPointID = -1;
+	void ClearHover();
+
+	// Support for selection
+	TArray<int32> SelectedPointIDs;
+	// We need the selected point start positions so we can move multiple points appropriately.
+	TArray<UE::Geometry::FFrame3d> CurveStartPositions;
+	// The starting point of the gizmo is needed to determine the offset by which to move the points.
+	UE::Geometry::FFrame3d GizmoStartPosition;
+
+	// These issue undo/redo change objects, and must therefore not be called in undo/redo code.
+	void ChangeSelection(int32 NewPointID, bool AddToSelection);
+
+	// All of the following do not issue undo/redo change objects.
+	bool HitTest(const FInputDeviceRay& ClickPos, FInputRayHit& ResultOut);
+	void SelectPoint(int32 PointID);
+	bool DeselectPoint(int32 PointID);
+	void UpdateSelection(const TArray<int32>& NewSelection);
+	void UpdateGizmoLocation();
+	void UpdateCurve(const TArray<UE::Geometry::FFrame3d>& NewPositions);
+
+	// Used for expiring undo/redo changes, which compare this to their stored value and expire themselves if they do not match.
+	int32 CurrentChangeStamp = 0;
+
+	friend class FSpaceCurveDeformationMechanicSelectionChange;
+	friend class FSpaceCurveDeformationMechanicMovementChange;
+};
+
+
+// Undo/redo support:
+
+class MODELINGCOMPONENTS_API FSpaceCurveDeformationMechanicSelectionChange : public FToolCommandChange
+{
+public:
+	FSpaceCurveDeformationMechanicSelectionChange(const TArray<int32>& FromIDs, const TArray<int32>& ToIDs);
+	virtual void Apply(UObject* Object) override;
+	virtual void Revert(UObject* Object) override;
+	virtual FString ToString() const override;
+
+protected:
+	TArray<int32> From, To;
+};
+
+
+class MODELINGCOMPONENTS_API FSpaceCurveDeformationMechanicMovementChange : public FToolCommandChange
+{
+public:
+	FSpaceCurveDeformationMechanicMovementChange(const TArray<UE::Geometry::FFrame3d>& FromPositions, const TArray<UE::Geometry::FFrame3d>& ToPositions);
+
+	virtual void Apply(UObject* Object) override;
+	virtual void Revert(UObject* Object) override;
+	virtual FString ToString() const override;
+
+protected:
+	TArray<UE::Geometry::FFrame3d> From;
+	TArray<UE::Geometry::FFrame3d> To;
+};
diff --git a/Engine/Plugins/Experimental/MeshModelingToolsetExp/Source/ModelingComponents/Public/Mechanics/SpatialCurveDistanceMechanic.h b/Engine/Plugins/Experimental/MeshModelingToolsetExp/Source/ModelingComponents/Public/Mechanics/SpatialCurveDistanceMechanic.h
new file mode 100644
--- /dev/null
+++ b/Engine/Plugins/Experimental/MeshModelingToolsetExp/Source/ModelingComponents/Public/Mechanics/SpatialCurveDistanceMechanic.h
@@ -0,0 +1,49 @@
+// Copyright Epic Games, Inc. All Rights Reserved.
+
+#pragma once
+
+#include "CoreMinimal.h"
+#include "InteractionMechanic.h"
+#include "DynamicMesh/DynamicMesh3.h"
+#include "DynamicMesh/DynamicMeshAABBTree3.h"
+#include "SpatialCurveDistanceMechanic.generated.h"
+
+/**
+ *
+ */
+UCLASS()
+class MODELINGCOMPONENTS_API USpatialCurveDistanceMechanic : public UInteractionMechanic
+{
+	GENERATED_BODY()
+	using FTransform3d = UE::Geometry::FTransform3d;
+public:
+	/** If this function is set, the hit point will be passed in to this function for snapping. Return false to indicate no snapping occurred. */
+	TUniqueFunction<bool(const FVector3d&, FVector3d&)> WorldPointSnapFunc = nullptr;
+
+	/** Current distance */
+	double CurrentDistance = 0.0f;
+
+	FVector3d CurrentCurvePoint;
+	FVector3d CurrentSpacePoint;
+
+public:
+
+	virtual void Setup(UInteractiveTool* ParentTool) override;
+	virtual void Render(IToolsContextRenderAPI* RenderAPI) override;
+
+	/**
+	 */
+	virtual void InitializePolyCurve(const TArray<FVector3d>& CurvePoints, const FTransform3d& Transform);
+	virtual void InitializePolyLoop(const TArray<FVector3d>& CurvePoints, const FTransform3d& Transform);
+
+	/**
+	 * Update the current distance/height based on the input world ray
+	 */
+	virtual void UpdateCurrentDistance(const FRay& WorldRay);
+
+protected:
+	TArray<FVector3d> Curve;
+	UE::Geometry::FDynamicMesh3 TargetHitMesh;
+	UE::Geometry::FDynamicMeshAABBTree3 TargetHitMeshAABB;
+	FTransform3d Transform;
+};
