MeshModelingToolset: Split plugin into non-experimental and experimental

diff --git a/Engine/Plugins/Experimental/MeshModelingToolsetExp/Source/ModelingComponents/Public/Selection/GroupTopologySelector.h b/Engine/Plugins/Experimental/MeshModelingToolsetExp/Source/ModelingComponents/Public/Selection/GroupTopologySelector.h
new file mode 100644
--- /dev/null
+++ b/Engine/Plugins/Experimental/MeshModelingToolsetExp/Source/ModelingComponents/Public/Selection/GroupTopologySelector.h
@@ -0,0 +1,190 @@
+// Copyright Epic Games, Inc. All Rights Reserved.
+
+#pragma once
+
+#include "GroupTopology.h"
+#include "Spatial/GeometrySet3.h"
+#include "DynamicMesh/DynamicMeshAABBTree3.h"
+#include "RayTypes.h"
+
+struct FCameraRectangle;
+class FToolDataVisualizer;
+struct FViewCameraState;
+
+using UE::Geometry::FGroupTopologySelection;
+using UE::Geometry::FGroupTopology;
+using UE::Geometry::FDynamicMeshAABBTree3;
+using UE::Geometry::FDynamicMesh3;
+using UE::Geometry::FGeometrySet3;
+
+/**
+ * FGroupTopologySelector implements selection behavior for a FGroupTopology mesh.
+ * Groups, Group Edges, and Corners can be selected dependings on the settings passed in.
+ *
+ * Internally a FGeometrySet3 is constructed to support ray-hit testing against the edges and corners.
+ * 
+ * Note that to hit-test against the mesh you have to provide a your own FDynamicMeshAABBTree3.
+ * You do this by providing a callback via SetSpatialSource(). The reason for this is that
+ * (1) owners of FGroupTopologySelector likely already have a BVTree and (2) if the use case
+ * is deformation, we need to make sure the owner has recomputed the BVTree before we call functions 
+ * on it. The callback you provide should do that!
+ * 
+ * DrawSelection() can be used to visualize a selection via line/circle drawing.
+ * 
+ * @todo optionally have an internal mesh AABBTree that can be used when owner does not provide one?
+ */
+class MODELINGCOMPONENTS_API FGroupTopologySelector
+{
+	using FRay3d = UE::Geometry::FRay3d;
+	using FTransform3d = UE::Geometry::FTransform3d;
+public:
+
+	//
+	// Configuration variables
+	// 
+
+	/**
+	 * Determines the behavior of a FindSelectedElement() call.
+	 */
+	struct FSelectionSettings
+	{
+		bool bEnableFaceHits = true;
+		bool bEnableEdgeHits = true;
+		bool bEnableCornerHits = true;
+		
+		// The following are mainly useful for ortho viewport selection:
+
+		// Prefer an edge projected to a point rather than the point, and a face projected to an edge
+		// rather than the edge.
+		bool bPreferProjectedElement = false;
+
+		// If the first element is valid, select all elements behind it that are aligned with it.
+		bool bSelectDownRay = false;
+
+		//Do not check whether the closest element is occluded.
+		bool bIgnoreOcclusion = false;
+	};
+
+	/** 
+	 * This is the function we use to determine if a point on a corner/edge is close enough
+	 * to the hit-test ray to treat as a "hit". By default this is Euclidean distance with
+	 * a tolerance of 1.0. You probably need to replace this with your own function.
+	 */
+	TFunction<bool(const FVector3d&, const FVector3d&, double ToleranceScale)> PointsWithinToleranceTest;
+
+public:
+	FGroupTopologySelector();
+
+	/**
+	 * Initialize the selector with the given Mesh and Topology.
+	 * This does not create the internal data structures, this happens lazily on GetGeometrySet()
+	 */
+	void Initialize(const FDynamicMesh3* Mesh, const FGroupTopology* Topology);
+
+	/**
+	 * Provide a function that will return an AABBTree for the Mesh.
+	 * See class comment for why this is necessary.
+	 */
+	void SetSpatialSource(TFunction<FDynamicMeshAABBTree3*(void)> GetSpatialFunc)
+	{
+		GetSpatial = GetSpatialFunc;
+	}
+
+	/**
+	 * Notify the Selector that the mesh has changed. 
+	 * @param bTopologyDeformed if this is true, the mesh vertices have been moved so we need to update bounding boxes/etc
+	 * @param bTopologyModified if this is true, topology has changed and we need to rebuild spatial data structures from scratch
+	 */
+	void Invalidate(bool bTopologyDeformed, bool bTopologyModified);
+
+	/**
+	 * @return the internal GeometrySet. This does lazy updating of the GeometrySet, so this function may take some time.
+	 */
+	const FGeometrySet3& GetGeometrySet();
+
+	/**
+	 * Find which element was selected for a given ray.
+	 *
+	 * @param Settings settings that determine what can be selected.
+	 * @param Ray hit-test ray
+	 * @param ResultOut resulting selection. Will not be cleared before use. At most one of the Groups/Corners/Edges
+	 *  members will be added to.
+	 * @param SelectedPositionOut The point on the ray nearest to the selected element
+	 * @param SelectedNormalOut the normal at that point if ResultOut contains a selected face, otherwise uninitialized
+	 * @param EdgeSegmentIdOut When not null, if the selected element is a group edge, the segment id of the segment
+	 *   that was actually clicked within the edge polyline will be stored here.
+	 *
+	 * @return true if something was selected
+	 */
+	bool FindSelectedElement(const FSelectionSettings& Settings, const FRay3d& Ray, FGroupTopologySelection& ResultOut,
+		FVector3d& SelectedPositionOut, FVector3d& SelectedNormalOut, int32* EdgeSegmentIdOut = nullptr);
+
+	/**
+	 * Given a camera rectangle (from a marquee selection) gives the group corners or edges that are contained in
+	 * the rectangle. Currently does not select faces.
+	 * 
+	 * Uses the bEnableEdgeHits, bEnableCornerHits, and bIgnoreOcclusion members of the FSelectionSettings argument.
+	 * For group edges, occlusion is checked at endpoints, and if either is occluded, the edge is considered occluded.
+	 *
+	 * @param Settings Settings that determine what is selected and whether occlusion is tested.
+	 * @param CameraRectangle The camera rectangle.
+	 * @param TargetTransform Transform from points to world space, since the rectangle is in world space.
+	 * @param ResultOut Output selection (cleared before use).
+	 *
+	 * @return true if something was selected.
+	 */
+	bool FindSelectedElement(const FSelectionSettings& Settings, const FCameraRectangle& CameraRectangle, 
+		FTransform3d TargetTransform, FGroupTopologySelection& ResultOut);
+
+	/**
+	 * Using the edges in the given selection as starting points, add any "edge loops" containing the edges. An 
+	 * edge loop is a sequence of edges that passes through valence-4 corners through the opposite edge, and may
+	 * not actually form a complete loop if they hit a non-valence-4 corner.
+	 *
+	 * @param Selection Selection to expand.
+	 * @param bool true if selection was modified (i.e., were the already selected edges part of any edge loops whose
+	 *  member edges were not yet all selected).
+	 */
+	bool ExpandSelectionByEdgeLoops(FGroupTopologySelection& Selection);
+
+	/**
+	 * Using the edges in the given selection as starting points, add any "edge rings" containing the edges. An
+	 * edge ring is a sequence of edges that lie on opposite sides of quad-like faces, meaning faces that have four
+	 * corners.
+	 *
+	 * @param Selection Selection to expand.
+	 * @param bool true if selection was modified (i.e., were the already selected edges part of any edge rings whose
+	 *  member edges were not yet all selected).
+	 */
+	bool ExpandSelectionByEdgeRings(FGroupTopologySelection& Selection);
+
+	/**
+	 * Render the given selection with the default settings of the FToolDataVisualizer.
+	 * Selected edges are drawn as lines, and selected corners are drawn as small view-facing circles.
+	 * (Currently seleced faces are not draw)
+	 */
+	void DrawSelection(const FGroupTopologySelection& Selection, FToolDataVisualizer* Renderer, const FViewCameraState* CameraState);
+
+
+public:
+	// internal rendering parameters
+	float VisualAngleSnapThreshold = 0.5;
+
+
+protected:
+	const FDynamicMesh3* Mesh = nullptr;
+	const FGroupTopology* Topology = nullptr;
+
+	TFunction<FDynamicMeshAABBTree3*(void)> GetSpatial = nullptr;
+
+	bool bGeometryInitialized = false;
+	bool bGeometryUpToDate = false;
+	FGeometrySet3 GeometrySet;
+
+	bool DoCornerBasedSelection(const FSelectionSettings& Settings, const FRay3d& Ray,
+		FDynamicMeshAABBTree3* Spatial, const FGeometrySet3& TopoSpatial, 
+		FGroupTopologySelection& ResultOut, FVector3d& SelectedPositionOut, int32 *EdgeSegmentIdOut) const;
+	bool DoEdgeBasedSelection(const FSelectionSettings& Settings, const FRay3d& Ray,
+		FDynamicMeshAABBTree3* Spatial, const FGeometrySet3& TopoSpatial,
+		FGroupTopologySelection& ResultOut, FVector3d& SelectedPositionOut, int32 *EdgeSegmentIdOut) const;
+};
\ No newline at end of file
diff --git a/Engine/Plugins/Experimental/MeshModelingToolsetExp/Source/ModelingComponents/Public/Selection/PolygonSelectionMechanic.h b/Engine/Plugins/Experimental/MeshModelingToolsetExp/Source/ModelingComponents/Public/Selection/PolygonSelectionMechanic.h
new file mode 100644
--- /dev/null
+++ b/Engine/Plugins/Experimental/MeshModelingToolsetExp/Source/ModelingComponents/Public/Selection/PolygonSelectionMechanic.h
@@ -0,0 +1,421 @@
+// Copyright Epic Games, Inc. All Rights Reserved.
+
+#pragma once
+
+#include "CoreMinimal.h"
+#include "Drawing/PreviewGeometryActor.h"
+#include "Drawing/TriangleSetComponent.h"
+#include "InputBehavior.h"
+#include "InteractiveTool.h"
+#include "Mechanics/RectangleMarqueeMechanic.h"
+#include "Components/DynamicMeshComponent.h"
+#include "Selection/GroupTopologySelector.h"
+#include "TransformTypes.h"
+#include "ToolDataVisualizer.h"
+#include "InteractionMechanic.h"
+#include "PolygonSelectionMechanic.generated.h"
+
+class FPolygonSelectionMechanicSelectionChange;
+class UGroupTopologyStorableSelection;
+class UMouseHoverBehavior;
+class URectangleMarqueeMechanic;
+class USingleClickInputBehavior;
+
+using UE::Geometry::FDynamicMeshAABBTree3;
+PREDECLARE_USE_GEOMETRY_STRUCT(FCompactMaps);
+
+UCLASS()
+class MODELINGCOMPONENTS_API UPolygonSelectionMechanicProperties : public UInteractiveToolPropertySet
+{
+	GENERATED_BODY()
+
+public:
+	UPROPERTY(EditAnywhere, Category = SelectionFilter)
+	bool bSelectFaces = true;
+
+	UPROPERTY(EditAnywhere, Category = SelectionFilter)
+	bool bSelectEdges = true;
+
+	UPROPERTY(EditAnywhere, Category = SelectionFilter)
+	bool bSelectVertices = true;
+
+	/** When true, will select edge loops. Edge loops are paths along a string of valence-4 vertices. */
+	UPROPERTY(EditAnywhere, Category = SelectionFilter, meta = (EditCondition = "bSelectEdges"))
+	bool bSelectEdgeLoops = false;
+
+	/** When true, will select rings of edges that are opposite each other across a quad face. */
+	UPROPERTY(EditAnywhere, Category = SelectionFilter, meta = (EditCondition = "bSelectEdges"))
+	bool bSelectEdgeRings = false;
+
+	UPROPERTY(EditAnywhere, Category = SelectionFilter, AdvancedDisplay)
+	bool bEnableMarquee = true;
+
+	/** Determines whether vertices should be checked for occlusion in marquee select (Note: marquee select currently only works with edges and vertices) */
+	UPROPERTY(EditAnywhere, Category = SelectionFilter, meta = (EditCondition = "bEnableMarquee", EditConditionHides))
+	bool bMarqueeIgnoreOcclusion = true;
+
+	// The following were originally in their own category, all marked as AdvancedDisplay. However, since there wasn't a non-AdvancedDisplay
+	// property in the category, they started out as expanded and could not be collapsed.
+	// The alternative approach, used below, is to have them in a nested category, which starts out as collapsed. This works nicely.
+
+	/** Prefer to select an edge projected to a point rather than the point, or a face projected to an edge rather than the edge. */
+	UPROPERTY(EditAnywhere, Category = "SelectionFilter|Ortho Viewport Behavior")
+	bool bPreferProjectedElement = true;
+
+	/** If the closest element is valid, select other elements behind it that are aligned with it. */
+	UPROPERTY(EditAnywhere, Category = "SelectionFilter|Ortho Viewport Behavior")
+	bool bSelectDownRay = true;
+
+	/** Do not check whether the closest element is occluded from the current view. */
+	UPROPERTY(EditAnywhere, Category = "SelectionFilter|Ortho Viewport Behavior")
+	bool bIgnoreOcclusion = false;
+};
+
+
+
+/**
+ * UPolygonSelectionMechanic implements the interaction for selecting a set of faces/vertices/edges
+ * from a FGroupTopology on a UDynamicMeshComponent. 
+ */
+UCLASS()
+class MODELINGCOMPONENTS_API UPolygonSelectionMechanic : public UInteractionMechanic, public IClickBehaviorTarget, public IHoverBehaviorTarget
+{
+	GENERATED_BODY()
+	using FFrame3d = UE::Geometry::FFrame3d;
+	using FTransform3d = UE::Geometry::FTransform3d;	
+	using FAxisAlignedBox3d = UE::Geometry::FAxisAlignedBox3d;
+public:
+
+	virtual ~UPolygonSelectionMechanic();
+
+	// configuration variables that must be set before bSetup is called
+	bool bAddSelectionFilterPropertiesToParentTool = true;
+
+	virtual void Setup(UInteractiveTool* ParentTool) override;
+	virtual void Shutdown() override;
+
+	virtual void Render(IToolsContextRenderAPI* RenderAPI) override;
+	virtual void DrawHUD(FCanvas* Canvas, IToolsContextRenderAPI* RenderAPI);
+
+	/**
+	 * Initializes the mechanic.
+	 *
+	 * @param Mesh Mesh that we are operating on.
+	 * @param TargetTransform Transform of the mesh.
+	 * @param World World in which we are operating, used to add drawing components that draw highlighted edges.
+	 * @param Topology Group topology of the mesh.
+	 * @param GetSpatialSourceFunc Function that returns an AABB tree for the mesh.
+	 */
+	void Initialize(const FDynamicMesh3* Mesh,
+		FTransform TargetTransform,
+		UWorld* World,
+		const FGroupTopology* Topology,
+		TFunction<FDynamicMeshAABBTree3*()> GetSpatialSourceFunc
+		);
+
+	void Initialize(UDynamicMeshComponent* MeshComponent, const FGroupTopology* Topology,
+		TFunction<FDynamicMeshAABBTree3 * ()> GetSpatialSourceFunc
+	);
+
+	/**
+	 * Removes the mechanic's own click/hover handlers, which means that the parent tool
+	 * will need to call UpdateSelection(), UpdateHighlight(), ClearHighlight(), and 
+	 * ClearSelection() from its own hover/click handlers.
+	 *
+	 * Must be called during tool Setup() after calling Setup() on the mechanic to have an effect.
+	 *
+	 * @param ParentToolIn The parent tool, needed to be able to remove the behaviors.
+	 */
+	void DisableBehaviors(UInteractiveTool* ParentToolIn);
+
+	/**
+	 * Enable/disable the mechanic without permanently removing behaviors or shutting it down.
+	 */
+	void SetIsEnabled(bool bOn);
+
+	/**
+	 * Sets the base priority so that tools can make sure that their own behaviors are higher
+	 * priority. The mechanic will not use any priority value higher than this, but it may use
+	 * lower if it needs to stagger the priorities of behaviors it uses.
+	 * Can be called before or after Setup().
+	 */
+	void SetBasePriority(const FInputCapturePriority& Priority);
+
+	/**
+	 * Gets the current priority range used by behaviors in the mechanic. The returned pair will
+	 * have the base (highest) priority as the key, and the lowest priority as the value.
+	 */
+	TPair<FInputCapturePriority, FInputCapturePriority> GetPriorityRange() const;
+
+
+	void SetShouldSelectEdgeLoopsFunc(TFunction<bool(void)> Func)
+	{
+		ShouldSelectEdgeLoopsFunc = Func;
+	}
+
+	void SetShouldSelectEdgeRingsFunc(TFunction<bool(void)> Func)
+	{
+		ShouldSelectEdgeRingsFunc = Func;
+	}
+
+	/**
+	 * By default, the shift key will cause new clicks to add to the selection. However, this
+	 * can be changed by supplying a different function to check here.
+	 */
+	void SetShouldAddToSelectionFunc(TFunction<bool(void)> Func)
+	{
+		ShouldAddToSelectionFunc = Func;
+	}
+
+	/**
+	 * By default, the Ctrl key will cause new clicks to remove from the existing selection.
+	 * However, this can be changed by supplying a different function to check here.
+	 */
+	void SetShouldRemoveFromSelectionFunc(TFunction<bool(void)> Func)
+	{
+		ShouldRemoveFromSelectionFunc = Func;
+	}
+
+	/**
+	 * Notify internal data structures that the associated MeshComponent has been modified.
+	 * @param bTopologyModified if true, the underlying mesh topology has been changed. This clears the current selection.
+	 */
+	void NotifyMeshChanged(bool bTopologyModified);
+
+	/**
+	 * Perform a hit test on the topology using the current selection settings. In cases of hitting edges and
+	 * corners, OutHit contains the following:
+	 *   OutHit.FaceIndex: edge or corner id in the topology
+	 *   OutHit.ImpactPoint: closest point on the ray to the hit element (Note: not a point on the element!)
+	 *   OutHit.Distance: distance along the ray to ImpactPoint
+	 *   OutHit.Item: if hit item was an edge, index of the segment within the edge polyline. Otherwise undefined.
+	 *
+	 * @param bUseOrthoSettings If true, the ortho-relevant settings for selection are used (selecting down the view ray, etc)
+	 */
+	bool TopologyHitTest(const FRay& WorldRay, FHitResult& OutHit, FGroupTopologySelection& OutSelection, bool bUseOrthoSettings = false);
+	bool TopologyHitTest(const FRay& WorldRay, FHitResult& OutHit, bool bUseOrthoSettings = false);
+
+	TSharedPtr<FGroupTopologySelector, ESPMode::ThreadSafe> GetTopologySelector() { return TopoSelector; }
+
+	//
+	// Hover API
+	//
+
+	/**
+	 * Update the hover highlight based on the hit elements at the given World Ray
+	 * @return true if something was hit and is now being hovered
+	 */
+	bool UpdateHighlight(const FRay& WorldRay);
+
+	/**
+	 * Clear current hover-highlight
+	 */
+	void ClearHighlight();
+
+
+	//
+	// Selection API
+	//
+
+	/** 
+	 * Intersect the ray with the mesh and update the selection based on the hit element, modifier states, etc
+	 * @return true if selection was modified
+	 */
+	bool UpdateSelection(const FRay& WorldRay, FVector3d& LocalHitPositionOut, FVector3d& LocalHitNormalOut);
+
+	/**
+	 * Replace the current selection with an external selection. 
+	 * @warning does not check that the selection is valid!
+	 *
+	 * @param bBroadcast If true, issues an OnSelectionChanged delegate broadcast.
+	 */
+	void SetSelection(const FGroupTopologySelection& Selection, bool bBroadcast = true);
+
+	/**
+	 * Clear the current selection.
+	 */
+	void ClearSelection();
+
+	/** 
+	 * @return true if the current selection is non-empty 
+	 */
+	bool HasSelection() const;
+
+	/**
+	 * @return the current selection
+	 */
+	const FGroupTopologySelection& GetActiveSelection() const { return PersistentSelection; }
+
+	/**
+	 * Gives the current selection as a storable selection object. Can optionally apply the passed-in
+	 * compact maps to the object if the topology in the mechanic was not updated after compacting.
+	 */
+	void GetStorableSelection(UGroupTopologyStorableSelection& StorableSelectionOut, 
+		const FCompactMaps* CompactMapsToApply = nullptr) const;
+
+	/**
+	 * Sets the current selection using the given storable selection object. The topology in the
+	 * mechanic must already be initialized for this to work.
+	 */
+	void LoadStorableSelection(const UGroupTopologyStorableSelection& StorableSelectionIn);
+
+	/**
+	 * Can be used by in an OnSelectionChanged event to inspect the clicked location (i.e., the
+	 * values returned by the UpdateSelection() function when the click happened).
+	 */
+	void GetClickedHitPosition(FVector3d& PositionOut, FVector3d& NormalOut) const;
+
+	/**
+	 * @return The best-guess 3D frame for the current select
+	 * @param bWorld if true, local-to-world transform of the target MeshComponent is applied to the frame
+	 */
+	FFrame3d GetSelectionFrame(bool bWorld, FFrame3d* InitialLocalFrame = nullptr) const;
+
+	/**
+	 * @return Bounding box for the current selection
+	 * @param bWorld if true, the box is in world space, otherwise it is in local space of the target MeshComponent
+	 */
+	FAxisAlignedBox3d GetSelectionBounds(bool bWorld) const;
+
+	//
+	// Change Tracking
+	//
+
+	/**
+	 * Begin a change record. Internally creates a FCommandChange and initializes it with current state
+	 */
+	void BeginChange();
+
+	/**
+	 * End the active change and return it. Returns empty change if the selection was not modified!
+	 */
+	TUniquePtr<FToolCommandChange> EndChange();
+
+	/**
+	 * Ends the active change and emits it via the parent tool, if the selection has been modified.
+	 */
+	bool EndChangeAndEmitIfModified();
+
+	// IClickBehaviorTarget implementation
+	virtual FInputRayHit IsHitByClick(const FInputDeviceRay& ClickPos) override;
+	virtual void OnClicked(const FInputDeviceRay& ClickPos) override;
+
+	// IHoverBehaviorTarget implementation
+	virtual FInputRayHit BeginHoverSequenceHitTest(const FInputDeviceRay& PressPos) override;
+	virtual void OnBeginHover(const FInputDeviceRay& DevicePos) override;
+	virtual bool OnUpdateHover(const FInputDeviceRay& DevicePos) override;
+	virtual void OnEndHover() override;
+	virtual void OnUpdateModifierState(int ModifierID, bool bIsOn) override;
+
+protected:
+	// These get bound to marquee mechanic delegates.
+	virtual void OnDragRectangleStarted();
+	virtual void OnDragRectangleChanged(const FCameraRectangle& CurrentRectangle);
+	virtual void OnDragRectangleFinished();
+
+	virtual void UpdateMarqueeEnabled();
+
+public:
+	/** 
+	 * OnSelectionChanged is broadcast whenever the selection is modified (including by FChanges, which
+	 * means that called functions should not issue undo transactions. 
+	 */
+	FSimpleMulticastDelegate OnSelectionChanged;
+
+	// TODO: Is it worth issuing separate callbacks in normal selection changes and in FChange ones, to
+	// allow the user to bundle in some FChanges into the normal callback?
+
+	UPROPERTY()
+	TObjectPtr<UPolygonSelectionMechanicProperties> Properties;
+
+protected:
+	bool bIsEnabled = true;
+
+	const FDynamicMesh3* Mesh;
+	const FGroupTopology* Topology;
+	TFunction<FDynamicMeshAABBTree3*()> GetSpatialFunc;
+
+	UPROPERTY()
+	TObjectPtr<UMouseHoverBehavior> HoverBehavior;
+
+	UPROPERTY()
+	TObjectPtr<USingleClickInputBehavior> ClickBehavior;
+
+	UPROPERTY()
+	TObjectPtr<URectangleMarqueeMechanic> MarqueeMechanic;
+
+	FInputCapturePriority BasePriority = FInputCapturePriority(FInputCapturePriority::DEFAULT_TOOL_PRIORITY);
+
+	// When bSelectEdgeLoops is true, this function is tested to see if we should select edge loops,
+	// to allow edge loop selection to be toggled with some key (setting bSelectEdgeLoops to
+	// false overrides this function).
+	TFunction<bool(void)> ShouldSelectEdgeLoopsFunc = []() {return true; };
+
+	// When bSelectEdgeRings is true, this function is tested to see if we should select edge rings,
+	// to allow edge ring selection to be toggled with some key (setting bSelectEdgeRings to
+	// false overrides this function).
+	TFunction<bool(void)> ShouldSelectEdgeRingsFunc = []() {return true; };
+
+	TFunction<bool(void)> ShouldAddToSelectionFunc = [this]() {return bShiftToggle; };
+	TFunction<bool(void)> ShouldRemoveFromSelectionFunc = [this]() {return bCtrlToggle; };;
+
+	FTransform3d TargetTransform;
+
+	TSharedPtr<FGroupTopologySelector, ESPMode::ThreadSafe> TopoSelector;
+
+	/**
+	 * Get the topology selector settings to use given the current selection settings.
+	 * 
+	 * @param bUseOrthoSettings If true, the topology selector will be configured to use ortho settings,
+	 *  which are generally different to allow for selection of projected elements, etc.
+	 */
+	FGroupTopologySelector::FSelectionSettings GetTopoSelectorSettings(bool bUseOrthoSettings = false);
+
+	FGroupTopologySelection HilightSelection;
+	FGroupTopologySelection PersistentSelection;
+	int32 SelectionTimestamp = 0;
+	TUniquePtr<FPolygonSelectionMechanicSelectionChange> ActiveChange;
+
+	FGroupTopologySelection PreDragPersistentSelection; // Used for box selection
+	FVector3d LastClickedHitPosition;
+	FVector3d LastClickedHitNormal;
+
+	/** The actor we create internally to own the DrawnTriangleSetComponent */
+	UPROPERTY()
+	TObjectPtr<APreviewGeometryActor> PreviewGeometryActor;
+
+	UPROPERTY()
+	TObjectPtr<UTriangleSetComponent> DrawnTriangleSetComponent;
+
+	TSet<int> CurrentlyHighlightedGroups;
+
+	UPROPERTY()
+	TObjectPtr<UMaterialInterface> HighlightedFaceMaterial;
+
+	FViewCameraState CameraState;
+
+	bool bShiftToggle = false;
+	bool bCtrlToggle = false;
+	static const int32 ShiftModifierID = 1;
+	static const int32 CtrlModifierID = 2;
+public:
+	FToolDataVisualizer PolyEdgesRenderer;
+	FToolDataVisualizer HilightRenderer;
+	FToolDataVisualizer SelectionRenderer;
+
+	friend class FPolygonSelectionMechanicSelectionChange;
+};
+
+
+
+class MODELINGCOMPONENTS_API FPolygonSelectionMechanicSelectionChange : public FToolCommandChange
+{
+public:
+	FGroupTopologySelection Before;
+	FGroupTopologySelection After;
+	int32 Timestamp = 0;
+
+	virtual void Apply(UObject* Object) override;
+	virtual void Revert(UObject* Object) override;
+	virtual FString ToString() const override;
+};
\ No newline at end of file
diff --git a/Engine/Plugins/Experimental/MeshModelingToolsetExp/Source/ModelingComponents/Public/Transforms/MultiTransformer.h b/Engine/Plugins/Experimental/MeshModelingToolsetExp/Source/ModelingComponents/Public/Transforms/MultiTransformer.h
new file mode 100644
--- /dev/null
+++ b/Engine/Plugins/Experimental/MeshModelingToolsetExp/Source/ModelingComponents/Public/Transforms/MultiTransformer.h
@@ -0,0 +1,135 @@
+// Copyright Epic Games, Inc. All Rights Reserved.
+
+#pragma once
+
+#include "CoreMinimal.h"
+
+#include "ToolContextInterfaces.h"
+#include "ToolDataVisualizer.h"
+#include "FrameTypes.h"
+#include "InteractiveGizmo.h"
+#include "Mechanics/DragAlignmentMechanic.h"
+#include "InteractiveGizmoManager.h"
+
+#include "MultiTransformer.generated.h"
+
+
+class UTransformGizmo;
+class UTransformProxy;
+
+UENUM()
+enum class EMultiTransformerMode
+{
+	DefaultGizmo = 1,
+	QuickAxisTranslation = 2
+};
+
+
+/**
+ * UMultiTransformer abstracts both a default TRS Gizmo, and the "Quick" translate/rotate Gizmos.
+ * The "Quick" part is not yet implemented, and we might end up phasing out this class.
+ */
+UCLASS()
+class MODELINGCOMPONENTS_API UMultiTransformer : public UObject
+{
+	GENERATED_BODY()
+	using FFrame3d = UE::Geometry::FFrame3d;
+public:
+	
+	virtual void Setup(UInteractiveGizmoManager* GizmoManager, IToolContextTransactionProvider* TransactionProviderIn);
+	virtual void Shutdown();
+
+	virtual void Tick(float DeltaTime);
+
+	virtual void InitializeGizmoPositionFromWorldFrame(const FFrame3d& Frame, bool bResetScale = true);
+	virtual void UpdateGizmoPositionFromWorldFrame(const FFrame3d& Frame, bool bResetScale = true);
+
+	virtual void ResetScale();
+
+	virtual const FFrame3d& GetCurrentGizmoFrame() const { return ActiveGizmoFrame; }
+	virtual const FVector3d& GetCurrentGizmoScale() const { return ActiveGizmoScale; }
+	virtual bool InGizmoEdit() const { return bInGizmoEdit;	}
+
+	virtual EMultiTransformerMode GetMode() const { return ActiveMode; }
+	virtual void SetMode(EMultiTransformerMode NewMode);
+
+	virtual void SetGizmoVisibility(bool bVisible);
+
+	virtual void SetOverrideGizmoCoordinateSystem(EToolContextCoordinateSystem CoordSystem);
+
+	virtual void SetEnabledGizmoSubElements(ETransformGizmoSubElements EnabledSubElements);
+
+	virtual void SetGizmoRepositionable(bool bOn);
+
+	virtual EToolContextCoordinateSystem GetGizmoCoordinateSystem();
+
+	void SetSnapToWorldGridSourceFunc(TUniqueFunction<bool()> EnableSnapFunc);
+	void SetIsNonUniformScaleAllowedFunction(TFunction<bool()> IsNonUniformScaleAllowedIn);
+
+	void SetDisallowNegativeScaling(bool bDisallow);
+
+	void AddAlignmentMechanic(UDragAlignmentMechanic* AlignmentMechanic);
+
+public:
+	DECLARE_MULTICAST_DELEGATE(FMultiTransformerEvent);
+
+	// Note that the following delegates don't fire on pivot repositioning drags.
+
+	/** This delegate is fired when a drag is started */
+	FMultiTransformerEvent OnTransformStarted;
+
+	/** This delegate is fired when a drag is updated */
+	FMultiTransformerEvent OnTransformUpdated;
+
+	/** This delegate is fired when the drag is completed */
+	FMultiTransformerEvent OnTransformCompleted;
+
+	/** This delegate is fired when a repositioning drag is completed */
+	FMultiTransformerEvent OnEndPivotEdit;
+
+	// Note: we could have more pivot change delegates, but we don't yet need them,
+	// and we might phase out the UMultiTransformer.
+
+protected:
+	UPROPERTY()
+	TObjectPtr<UInteractiveGizmoManager> GizmoManager;
+
+	IToolContextTransactionProvider* TransactionProvider;
+
+	EMultiTransformerMode ActiveMode;
+
+	ETransformGizmoSubElements ActiveGizmoSubElements = ETransformGizmoSubElements::FullTranslateRotateScale;
+
+	EToolContextCoordinateSystem GizmoCoordSystem = EToolContextCoordinateSystem::World;
+	bool bForceGizmoCoordSystem = false;
+
+	bool bShouldBeVisible = true;
+	FFrame3d ActiveGizmoFrame;
+	FVector3d ActiveGizmoScale;
+
+	bool bRepositionableGizmo = false;
+
+	bool bDisallowNegativeScaling = false;
+
+	UPROPERTY()
+	TObjectPtr<UTransformGizmo> TransformGizmo;
+
+	UPROPERTY()
+	TObjectPtr<UTransformProxy> TransformProxy;
+
+	// We have to hold on to the mechanic only because the MultiTransformer has the capacity to delete and
+	// recreate its gizmo, in which case we'll need to attach the alignment mechanic again.
+	UPROPERTY()
+	TObjectPtr<UDragAlignmentMechanic> DragAlignmentMechanic = nullptr;
+
+	TUniqueFunction<bool()> EnableSnapToWorldGridFunc;
+	TFunction<bool()> IsNonUniformScaleAllowed;
+
+	// called on PlaneTransformProxy.OnTransformChanged
+	void OnProxyTransformChanged(UTransformProxy* Proxy, FTransform Transform);
+	void OnBeginProxyTransformEdit(UTransformProxy* Proxy);
+	void OnEndProxyTransformEdit(UTransformProxy* Proxy);
+	bool bInGizmoEdit = false;
+
+	void UpdateShowGizmoState(bool bNewVisibility);
+};
\ No newline at end of file
diff --git a/Engine/Plugins/Experimental/MeshModelingToolsetExp/Source/ModelingComponents/Public/Transforms/QuickAxisRotator.h b/Engine/Plugins/Experimental/MeshModelingToolsetExp/Source/ModelingComponents/Public/Transforms/QuickAxisRotator.h
new file mode 100644
--- /dev/null
+++ b/Engine/Plugins/Experimental/MeshModelingToolsetExp/Source/ModelingComponents/Public/Transforms/QuickAxisRotator.h
@@ -0,0 +1,117 @@
+// Copyright Epic Games, Inc. All Rights Reserved.
+
+#pragma once
+
+#include "CoreMinimal.h"
+#include "QuickTransformer.h"
+#include "Snapping/RaySpatialSnapSolver.h"
+
+namespace UE
+{
+namespace Geometry
+{
+
+/**
+
+ */
+class MODELINGCOMPONENTS_API FQuickAxisRotator : public FQuickTransformer
+{
+public:
+	virtual ~FQuickAxisRotator() {}
+
+	/**
+	 * Set up internal data structures
+	 */
+	virtual void Initialize() override;
+
+	/**
+	 * Set current transform frame to the unit axes at the given Origin
+	 */
+	virtual void SetActiveFrameFromWorldAxes(const FVector3d& Origin) override;
+
+	/**
+	 * Set current transform frame to the given frame
+	 */
+	virtual void SetActiveWorldFrame(const FFrame3d& Frame) override;
+
+	/**
+	 * @return the current transform frame
+	 */
+	const FFrame3d& GetActiveWorldFrame() const 
+	{
+		return AxisFrameWorld; 
+	}
+
+	/**
+	 * Set current snap-axis frame to a frame at the given Origin with Z aligned to the given Normal.
+	 * @param bAlignToWorldAxes if true, the frame is rotated around the Normal so that the X and Y axes are best-aligned to the World axes
+	 */
+	virtual void SetActiveFrameFromWorldNormal(const FVector3d& Origin, const FVector3d& Normal, bool bAlignToWorldAxes);
+
+	/**
+	 * Update the current snap-axis frame with a new origin
+	 */
+	virtual void UpdateActiveFrameOrigin(const FVector3d& NewOrigin) override;
+
+
+	/**
+	 * Try to find the best snap point for the given Ray, and store in SnapPointOut if found
+	 * @return true if found
+	 */
+	virtual bool UpdateSnap(const FRay3d& Ray, FVector3d& SnapPointOut);
+
+	/** @return true if there is an active snap */
+	virtual bool HaveActiveSnap() const;
+
+
+	virtual bool HaveActiveSnapRotation() const;
+
+	virtual FFrame3d GetActiveRotationFrame() const;
+
+	/** Reset transformer state */
+	virtual void Reset() override;
+
+
+	void ClearAxisLock();
+	void SetAxisLock();
+	bool GetHaveLockedToAxis() const { return bHaveLockedToAxis; }
+
+
+	/**
+	 * Update internal copy of camera state. You must call this for snapping to work!
+	 */
+	virtual void UpdateCameraState(const FViewCameraState& CameraState) override;
+
+	/**
+	 * Draw a visualization of the current snap axes and active snap point
+	 * @param RenderAPI provide access to context rendering info
+	 */
+	virtual void Render(IToolsContextRenderAPI* RenderAPI) override;
+
+	/**
+	 * Draw a visualization of the current snap axes and active snap point
+	 * @param RenderAPI provide access to context rendering info
+	 */
+	virtual void PreviewRender(IToolsContextRenderAPI* RenderAPI) override;
+
+
+protected:
+	// camera state saved at last Render()
+	FViewCameraState CameraState;
+
+	bool bHaveLockedToAxis = false;
+	FIndex3i IgnoredAxes;
+
+	FFrame3d AxisFrameWorld;
+
+	FRaySpatialSnapSolver MoveAxisSolver;
+	void UpdateSnapAxes();
+
+	FToolDataVisualizer QuickAxisRenderer;
+	TMap<int, FLinearColor> AxisColorMap;
+
+	FToolDataVisualizer QuickAxisPreviewRenderer;
+};
+
+} // end namespace UE::Geometry
+} // end namespace UE
\ No newline at end of file
diff --git a/Engine/Plugins/Experimental/MeshModelingToolsetExp/Source/ModelingOperators/ModelingOperators.Build.cs b/Engine/Plugins/Experimental/MeshModelingToolsetExp/Source/ModelingOperators/ModelingOperators.Build.cs
new file mode 100644
--- /dev/null
+++ b/Engine/Plugins/Experimental/MeshModelingToolsetExp/Source/ModelingOperators/ModelingOperators.Build.cs
@@ -0,0 +1,74 @@
+// Copyright Epic Games, Inc. All Rights Reserved.
+
+using UnrealBuildTool;
+
+public class ModelingOperators : ModuleRules
+{
+	public ModelingOperators(ReadOnlyTargetRules Target) : base(Target)
+	{
+		PCHUsage = ModuleRules.PCHUsageMode.UseExplicitOrSharedPCHs;
+		
+		PublicIncludePaths.AddRange(
+			new string[] {
+				// ... add public include paths required here ...
+			}
+			);
+				
+		
+		PrivateIncludePaths.AddRange(
+			new string[] {
+				// ... add other private include paths required here ...
+			}
+			);
+			
+		
+		PublicDependencyModuleNames.AddRange(
+			new string[]
+			{
+				// These dependencies were commented out when we split out ModelingOperatorsEditorOnly
+				// Many of them are editor only and should not be depended on here.  If you re-add any
+				// dependent modules, please confirm by "launching" or packaging that you are not introducing
+				// editor dependencies
+				//
+                //"ProxyLODMeshReduction", // currently required to be public due to IVoxelBasedCSG API
+                //"MeshDescription",
+
+				"Core",
+                "InteractiveToolsFramework",
+				"GeometryCore",
+				"DynamicMesh",
+				"MeshConversion",
+                "GeometryAlgorithms", // required for constrained Delaunay triangulation
+				// ... add other public dependencies that you statically link with here ...
+			}
+            );
+			
+		
+		PrivateDependencyModuleNames.AddRange(
+			new string[]
+			{
+				// These dependencies were commented out when we split out ModelingOperatorsEditorOnly
+				// Many of them are editor only and should not be depended on here.  If you re-add any
+				// dependent modules, please confirm by "launching" or packaging that you are not introducing
+				// editor dependencies
+				//
+				//"MeshBuilder",
+                //"MeshDescriptionOperations",
+				//"MeshUtilitiesCommon",
+				//"MeshReductionInterface", // for standard UE simplification
+
+				"CoreUObject",
+				"Engine",
+
+			}
+			);
+		
+		
+		DynamicallyLoadedModuleNames.AddRange(
+			new string[]
+			{
+				// ... add any modules that your module loads dynamically here ...
+			}
+			);
+	}
+}
diff --git a/Engine/Plugins/Experimental/MeshModelingToolsetExp/Source/ModelingOperators/Private/SmoothingOps/IterativeSmoothingOp.cpp b/Engine/Plugins/Experimental/MeshModelingToolsetExp/Source/ModelingOperators/Private/SmoothingOps/IterativeSmoothingOp.cpp
new file mode 100644
--- /dev/null
+++ b/Engine/Plugins/Experimental/MeshModelingToolsetExp/Source/ModelingOperators/Private/SmoothingOps/IterativeSmoothingOp.cpp
@@ -0,0 +1,73 @@
+// Copyright Epic Games, Inc. All Rights Reserved.
+
+#include "SmoothingOps/IterativeSmoothingOp.h"
+#include "Solvers/MeshSmoothing.h"
+#include "MeshCurvature.h"
+#include "Async/ParallelFor.h"
+
+#include "ExplicitUseGeometryMathTypes.h"		// using UE::Geometry::(math types)
+using namespace UE::Geometry;
+
+FIterativeSmoothingOp::FIterativeSmoothingOp(const FDynamicMesh3* Mesh, const FSmoothingOpBase::FOptions& OptionsIn) :
+	FSmoothingOpBase(Mesh, OptionsIn)
+{
+}
+
+void FIterativeSmoothingOp::CalculateResult(FProgressCancel* Progress)
+{
+	// Update the values in the position buffer with smoothed positions.
+	if (SmoothOptions.bUseImplicit)
+	{
+		if (SmoothOptions.bUniform)
+		{
+			Smooth_MeanValue();
+		}
+		else
+		{
+			Smooth_Implicit_Cotan();
+		}
+	}
+	else
+	{
+		Smooth_Forward(SmoothOptions.bUniform);
+	}
+
+	// Copy the results back into the result mesh and update normals
+	UpdateResultMesh();
+
+}
+
+
+double FIterativeSmoothingOp::GetSmoothAlpha(int32 VertexID, bool bIsBoundary)
+{
+	double UseAlpha = (bIsBoundary) ? SmoothOptions.BoundarySmoothAlpha : SmoothOptions.SmoothAlpha;
+	if (SmoothOptions.bUseWeightMap)
+	{
+		double t = FMathd::Clamp(SmoothOptions.WeightMap->GetValue(VertexID), 0.0, 1.0);
+		UseAlpha = FMathd::Lerp(SmoothOptions.WeightMapMinMultiplier * UseAlpha, UseAlpha, t);
+	}
+	return UseAlpha;
+}
+
+
+void FIterativeSmoothingOp::Smooth_Implicit_Cotan()
+{
+	double Intensity = 1.;
+	UE::MeshDeformation::ComputeSmoothing_BiHarmonic(ELaplacianWeightScheme::ClampedCotangent, 
+		*ResultMesh, SmoothOptions.SmoothAlpha, Intensity, SmoothOptions.Iterations, PositionBuffer);
+}
+
+
+void FIterativeSmoothingOp::Smooth_MeanValue()
+{
+	double Intensity = 1.;
+	UE::MeshDeformation::ComputeSmoothing_BiHarmonic(ELaplacianWeightScheme::MeanValue,
+		*ResultMesh, SmoothOptions.SmoothAlpha, Intensity, SmoothOptions.Iterations, PositionBuffer);
+}
+
+
+void FIterativeSmoothingOp::Smooth_Forward(bool bUniform)
+{
+	UE::MeshDeformation::ComputeSmoothing_Forward(bUniform, SmoothOptions.bSmoothBoundary,
+		*ResultMesh, [this](int VID, bool bBoundary) { return GetSmoothAlpha(VID, bBoundary); }, SmoothOptions.Iterations, PositionBuffer);
+}
diff --git a/Engine/Plugins/Experimental/MeshModelingToolsetExp/Source/ModelingOperators/Private/SmoothingOps/SmoothingOpBase.cpp b/Engine/Plugins/Experimental/MeshModelingToolsetExp/Source/ModelingOperators/Private/SmoothingOps/SmoothingOpBase.cpp
new file mode 100644
--- /dev/null
+++ b/Engine/Plugins/Experimental/MeshModelingToolsetExp/Source/ModelingOperators/Private/SmoothingOps/SmoothingOpBase.cpp
@@ -0,0 +1,53 @@
+// Copyright Epic Games, Inc. All Rights Reserved.
+
+#include "SmoothingOps/SmoothingOpBase.h"
+
+#include "DynamicMesh/DynamicMesh3.h"
+#include "DynamicMesh/MeshNormals.h"
+
+#include "ExplicitUseGeometryMathTypes.h"		// using UE::Geometry::(math types)
+using namespace UE::Geometry;
+
+FSmoothingOpBase::FSmoothingOpBase(const FDynamicMesh3* Mesh, const FOptions& OptionsIn) :
+	FDynamicMeshOperator(),
+	SmoothOptions(OptionsIn)
+{
+	// deep copy the src mesh into the result mesh.  This ResultMesh will be directly updated by the smoothing.
+	ResultMesh->Copy(*Mesh);
+	
+	PositionBuffer.SetNum(ResultMesh->MaxVertexID());
+
+	for (int vid : ResultMesh->VertexIndicesItr())
+	{
+		PositionBuffer[vid] = ResultMesh->GetVertex(vid);
+	}
+}
+
+void FSmoothingOpBase::SetTransform(const FTransform3d& XForm)
+{
+	ResultTransform = XForm;
+}
+
+
+void FSmoothingOpBase::UpdateResultMesh()
+{
+	// move the verts to the new location
+	for (int32 vid : ResultMesh->VertexIndicesItr())
+	{
+		const FVector3d Pos = PositionBuffer[vid];
+		ResultMesh->SetVertex(vid, Pos);
+	}
+
+	// recalculate normals
+	if (ResultMesh->HasAttributes())
+	{
+		FMeshNormals Normals(ResultMesh.Get());
+		FDynamicMeshNormalOverlay* NormalOverlay = ResultMesh->Attributes()->PrimaryNormals();
+		Normals.RecomputeOverlayNormals(NormalOverlay);
+		Normals.CopyToOverlay(NormalOverlay);
+	}
+	else
+	{
+		FMeshNormals::QuickComputeVertexNormals(*ResultMesh);
+	}
+}
diff --git a/Engine/Plugins/Experimental/MeshModelingToolsetExp/Source/ModelingOperators/Public/BaseOps/SimpleMeshProcessingBaseOp.h b/Engine/Plugins/Experimental/MeshModelingToolsetExp/Source/ModelingOperators/Public/BaseOps/SimpleMeshProcessingBaseOp.h
new file mode 100644
--- /dev/null
+++ b/Engine/Plugins/Experimental/MeshModelingToolsetExp/Source/ModelingOperators/Public/BaseOps/SimpleMeshProcessingBaseOp.h
@@ -0,0 +1,55 @@
+// Copyright Epic Games, Inc. All Rights Reserved.
+
+#pragma once
+
+#include "CoreMinimal.h"
+#include "Util/ProgressCancel.h"
+#include "ModelingOperators.h"
+#include "VectorTypes.h"
+
+namespace UE
+{
+namespace Geometry
+{
+
+class FDynamicMesh3;
+
+class MODELINGOPERATORS_API FMeshBoundaryCache
+{
+public:
+	TArray<bool> bIsBoundary;
+	TArray<int32> BoundaryVerts;
+	void Calculate(const FDynamicMesh3& Mesh);
+};
+
+
+/**
+ * Dynamic Mesh Operator that has a separate PositionBuffer that contains copy of input
+ * vertex positions. These can be modified by operator and then UpdateResultMeshPositions() can
+ * be copied to update ResultMesh. 
+ *
+ * UpdateResultMeshNormals() recalculates normals
+ */
+class MODELINGOPERATORS_API FSimpleMeshProcessingBaseOp : public FDynamicMeshOperator
+{
+public:
+	FSimpleMeshProcessingBaseOp(const FDynamicMesh3* Mesh);
+
+	virtual ~FSimpleMeshProcessingBaseOp() {}
+
+	// set ability on protected transform.
+	void SetTransform(const FTransform3d& XForm);
+
+	// copy the PositionBuffer locations back to the ResultMesh
+	void UpdateResultMeshPositions();
+
+	// recompute noramls of ResultMesh
+	void UpdateResultMeshNormals();
+
+protected:
+	TArray<FVector3d> PositionBuffer;
+};
+
+
+} // end namespace UE::Geometry
+} // end namespace UE
diff --git a/Engine/Plugins/Experimental/MeshModelingToolsetExp/Source/ModelingOperators/Public/CompositionOps/BooleanMeshesOp.h b/Engine/Plugins/Experimental/MeshModelingToolsetExp/Source/ModelingOperators/Public/CompositionOps/BooleanMeshesOp.h
new file mode 100644
--- /dev/null
+++ b/Engine/Plugins/Experimental/MeshModelingToolsetExp/Source/ModelingOperators/Public/CompositionOps/BooleanMeshesOp.h
@@ -0,0 +1,93 @@
+// Copyright Epic Games, Inc. All Rights Reserved.
+
+#pragma once
+
+#include "CoreMinimal.h"
+#include "Util/ProgressCancel.h"
+#include "ModelingOperators.h"
+
+
+#include "BooleanMeshesOp.generated.h"
+
+
+/** CSG operation types */
+UENUM()
+enum class ECSGOperation : uint8
+{
+	/** Subtracts the first object from the second */
+	DifferenceAB = 0 UMETA(DisplayName = "A - B"),
+
+	/** Subtracts the second object from the first */
+	DifferenceBA = 1 UMETA(DisplayName = "B - A"),
+
+	/** intersection of two objects */
+	Intersect = 2 UMETA(DisplayName = "Intersect"),
+
+	/** union of two objects */
+	Union = 3 UMETA(DisplayName = "Union"),
+};
+
+UENUM()
+enum class ETrimOperation : uint8
+{
+	/** Remove geometry from the first object using the second */
+	TrimA = 0 UMETA(DisplayName = "Trim A"),
+
+	/** Remove geometry from the second object using the first */
+	TrimB = 1 UMETA(DisplayName = "Trim B"),
+};
+
+
+UENUM()
+enum class ETrimSide : uint8
+{
+	RemoveInside = 0,
+	RemoveOutside = 1,
+};
+
+
+namespace UE
+{
+namespace Geometry
+{
+
+class MODELINGOPERATORS_API FBooleanMeshesOp : public FDynamicMeshOperator
+{
+public:
+	virtual ~FBooleanMeshesOp() {}
+
+	// inputs
+	ECSGOperation CSGOperation;
+	ETrimOperation TrimOperation;
+	ETrimSide TrimSide;
+	bool bTrimMode = false; // if true, do a trim operation instead of a boolean
+	TArray<TSharedPtr<const FDynamicMesh3, ESPMode::ThreadSafe>> Meshes;
+	TArray<FTransform> Transforms; // 1:1 with Meshes
+	bool bAttemptFixHoles = false;
+
+	/** If true, try to do edge-collapses along cut edges to remove unnecessary edges inserted by cut */
+	bool bTryCollapseExtraEdges = false;
+	/** Angle threshold in degrees used for testing if two triangles should be considered coplanar, or two lines collinear */
+	float TryCollapseExtraEdgesPlanarThresh = 0.01f;
+
+	void SetTransform(const FTransform& Transform);
+
+	//
+	// FDynamicMeshOperator implementation
+	// 
+
+	virtual void CalculateResult(FProgressCancel* Progress) override;
+
+	// IDs of any newly-created boundary edges in the result mesh
+	const TArray<int>& GetCreatedBoundaryEdges() const
+	{
+		return CreatedBoundaryEdges;
+	}
+
+private:
+	TArray<int> CreatedBoundaryEdges;
+};
+
+
+} // end namespace UE::Geometry
+} // end namespace UE
diff --git a/Engine/Plugins/Experimental/MeshModelingToolsetExp/Source/ModelingOperators/Public/CompositionOps/VoxelSolidifyMeshesOp.h b/Engine/Plugins/Experimental/MeshModelingToolsetExp/Source/ModelingOperators/Public/CompositionOps/VoxelSolidifyMeshesOp.h
new file mode 100644
--- /dev/null
+++ b/Engine/Plugins/Experimental/MeshModelingToolsetExp/Source/ModelingOperators/Public/CompositionOps/VoxelSolidifyMeshesOp.h
@@ -0,0 +1,45 @@
+// Copyright Epic Games, Inc. All Rights Reserved.
+
+#pragma once
+
+#include "CoreMinimal.h"
+#include "Util/ProgressCancel.h"
+#include "ModelingOperators.h"
+
+#include "BaseOps/VoxelBaseOp.h"
+
+namespace UE
+{
+namespace Geometry
+{
+
+class MODELINGOPERATORS_API FVoxelSolidifyMeshesOp : public FVoxelBaseOp
+{
+public:
+	virtual ~FVoxelSolidifyMeshesOp() {}
+
+	// inputs
+	TArray<TSharedPtr<const FDynamicMesh3, ESPMode::ThreadSafe>> Meshes;
+	TArray<FTransform> Transforms; // 1:1 with Meshes
+
+	double WindingThreshold = .5;
+	double ExtendBounds = 1;
+	bool bSolidAtBoundaries = true;
+	int SurfaceSearchSteps = 3;
+
+	bool bApplyThickenShells = false;
+	double ThickenShells = 5;
+
+	void SetTransform(const FTransform& Transform);
+
+	//
+	// FDynamicMeshOperator implementation
+	// 
+
+	virtual void CalculateResult(FProgressCancel* Progress) override;
+
+};
+
+
+} // end namespace UE::Geometry
+} // end namespace UE
\ No newline at end of file
diff --git a/Engine/Plugins/Experimental/MeshModelingToolsetExp/Source/ModelingOperators/Public/DeformationOps/LatticeDeformerOp.h b/Engine/Plugins/Experimental/MeshModelingToolsetExp/Source/ModelingOperators/Public/DeformationOps/LatticeDeformerOp.h
new file mode 100644
--- /dev/null
+++ b/Engine/Plugins/Experimental/MeshModelingToolsetExp/Source/ModelingOperators/Public/DeformationOps/LatticeDeformerOp.h
@@ -0,0 +1,45 @@
+// Copyright Epic Games, Inc. All Rights Reserved.
+
+#pragma once
+
+#include "BaseOps/SimpleMeshProcessingBaseOp.h"
+#include "DynamicMeshEditor.h"
+#include "DynamicMesh/MeshNormals.h"
+#include "MeshWeights.h"
+#include "MeshBoundaryLoops.h"
+#include "Operations/JoinMeshLoops.h"
+#include "Solvers/ConstrainedMeshDeformer.h"
+#include "Implicit/GridInterpolant.h"
+#include "Operations/FFDLattice.h"
+
+namespace UE
+{
+namespace Geometry
+{
+
+class MODELINGOPERATORS_API FLatticeDeformerOp : public FDynamicMeshOperator
+{
+public:
+
+	FLatticeDeformerOp(TSharedPtr<FDynamicMesh3, ESPMode::ThreadSafe> InOriginalMesh,
+					   TSharedPtr<UE::Geometry::FFFDLattice, ESPMode::ThreadSafe> InLattice,
+					   const TArray<FVector3d>& InLatticeControlPoints,
+					   UE::Geometry::ELatticeInterpolation InInterpolationType,
+					   bool bInDeformNormals);
+
+	// FDynamicMeshOperator implementation
+	void CalculateResult(FProgressCancel* Progress) override;
+
+protected:
+
+	// Inputs
+	const TSharedPtr<const UE::Geometry::FFFDLattice, ESPMode::ThreadSafe> Lattice;
+	const TSharedPtr<const FDynamicMesh3, ESPMode::ThreadSafe> OriginalMesh;
+	const TArray<FVector3d> LatticeControlPoints;
+	UE::Geometry::ELatticeInterpolation InterpolationType;
+	bool bDeformNormals = false;
+};
+
+} // end namespace UE::Geometry
+} // end namespace UE
+
diff --git a/Engine/Plugins/Experimental/MeshModelingToolsetExp/Source/ModelingOperators/Public/ParameterizationOps/UVProjectionOp.h b/Engine/Plugins/Experimental/MeshModelingToolsetExp/Source/ModelingOperators/Public/ParameterizationOps/UVProjectionOp.h
new file mode 100644
--- /dev/null
+++ b/Engine/Plugins/Experimental/MeshModelingToolsetExp/Source/ModelingOperators/Public/ParameterizationOps/UVProjectionOp.h
@@ -0,0 +1,77 @@
+// Copyright Epic Games, Inc. All Rights Reserved.
+
+#pragma once
+
+#include "CoreMinimal.h"
+#include "Util/ProgressCancel.h"
+#include "ModelingOperators.h"
+#include "OrientedBoxTypes.h"
+
+#include "UVProjectionOp.generated.h"
+
+UENUM()
+enum class EUVProjectionMethod : uint8
+{
+	/** Calculate UVs by assigning triangles to faces of a Box, and then apply per-box-face planar projection */
+	Box,
+	/** Calculate UVs by assigning triangles to Cylinder (radial projection) or Endcaps (planar projection) */
+	Cylinder,
+	/** Calculate UVs by projecting to a 3D plane */
+	Plane,
+	/** Calculate UVs by Exponential Map Projection centered at nearest surface point to input 3D plane */
+	ExpMap
+};
+
+namespace UE
+{
+namespace Geometry
+{
+
+class MODELINGOPERATORS_API FUVProjectionOp : public FDynamicMeshOperator
+{
+public:
+	virtual ~FUVProjectionOp() {}
+
+	// inputs
+	TSharedPtr<FDynamicMesh3, ESPMode::ThreadSafe> OriginalMesh;
+	TSharedPtr<TArray<int32>, ESPMode::ThreadSafe> TriangleROI;
+	int32 UseUVLayer = 0;
+
+	// transform from Mesh into space of Projection Geometry
+	FTransform3d MeshToProjectionSpace;
+
+	EUVProjectionMethod ProjectionMethod = EUVProjectionMethod::Box;
+	FOrientedBox3d ProjectionBox;
+	int32 MinRegionSize = 2;
+	float CylinderSplitAngle = 45.0f;
+
+	// ExpMap parameters for modifying normals
+	double BlendWeight = 0.0;
+	int32 SmoothingRounds = 0;
+	double SmoothingAlpha = 0.5;
+
+	// position to use as UV origin
+	FVector2f UVOrigin = FVector2f(0.5f,0.5f);
+	// rotation applied to computed UVs
+	float UVRotationAngleDeg = 0.0;
+	// scale applied to computed UVs
+	FVector2f UVScale = FVector2f::One();
+	// translation applied to computed UVs
+	FVector2f UVTranslate = FVector2f::Zero();
+
+
+	//
+	// FDynamicMeshOperator implementation
+	// 
+	virtual void CalculateResult(FProgressCancel* Progress) override;
+
+
+	virtual void CalculateResult_Box(FProgressCancel* Progress);
+	virtual void CalculateResult_Plane(FProgressCancel* Progress);
+	virtual void CalculateResult_ExpMap(FProgressCancel* Progress);
+	virtual void CalculateResult_Cylinder(FProgressCancel* Progress);
+};
+
+} // end namespace UE::Geometry
+} // end namespace UE
+
diff --git a/Engine/Plugins/Experimental/MeshModelingToolsetExp/Source/ModelingOperators/Public/SmoothingOps/CotanSmoothingOp.h b/Engine/Plugins/Experimental/MeshModelingToolsetExp/Source/ModelingOperators/Public/SmoothingOps/CotanSmoothingOp.h
new file mode 100644
--- /dev/null
+++ b/Engine/Plugins/Experimental/MeshModelingToolsetExp/Source/ModelingOperators/Public/SmoothingOps/CotanSmoothingOp.h
@@ -0,0 +1,30 @@
+// Copyright Epic Games, Inc. All Rights Reserved.
+
+#pragma once
+
+#include "SmoothingOpBase.h"
+#include "CoreMinimal.h"
+
+namespace UE
+{
+namespace Geometry
+{
+
+class MODELINGOPERATORS_API  FCotanSmoothingOp : public FSmoothingOpBase
+{
+public:
+	FCotanSmoothingOp(const FDynamicMesh3* Mesh, const FSmoothingOpBase::FOptions& OptionsIn);
+
+	~FCotanSmoothingOp() override {};
+
+	void CalculateResult(FProgressCancel* Progress) override;
+
+private:
+	// Compute the smoothed result by using Cotan Biharmonic
+	void Smooth();	
+
+	double GetSmoothPower(int32 VertexID, bool bIsBoundary);
+};
+
+} // end namespace UE::Geometry
+} // end namespace UE
diff --git a/Engine/Plugins/Experimental/MeshModelingToolsetExp/Source/ModelingOperators/Public/SmoothingOps/IterativeSmoothingOp.h b/Engine/Plugins/Experimental/MeshModelingToolsetExp/Source/ModelingOperators/Public/SmoothingOps/IterativeSmoothingOp.h
new file mode 100644
--- /dev/null
+++ b/Engine/Plugins/Experimental/MeshModelingToolsetExp/Source/ModelingOperators/Public/SmoothingOps/IterativeSmoothingOp.h
@@ -0,0 +1,39 @@
+// Copyright Epic Games, Inc. All Rights Reserved.
+
+#pragma once
+
+#include "SmoothingOpBase.h"
+#include "CoreMinimal.h"
+#include "VectorTypes.h"
+
+namespace UE
+{
+namespace Geometry
+{
+
+class MODELINGOPERATORS_API  FIterativeSmoothingOp : public FSmoothingOpBase
+{
+public:
+	FIterativeSmoothingOp(const FDynamicMesh3* Mesh, const FSmoothingOpBase::FOptions& OptionsIn);
+
+	~FIterativeSmoothingOp() override {};
+
+	// Apply smoothing. results in an updated ResultMesh
+	void CalculateResult(FProgressCancel* Progress) override;
+
+private:
+
+	double GetSmoothAlpha(int32 VertexID, bool bIsBoundary);
+
+	// uniform iterative smoothing
+	void Smooth_Forward(bool bUniform);
+
+	// cotan smoothing iterations
+	void Smooth_Implicit_Cotan();
+
+	// mean value smoothing iterations
+	void Smooth_MeanValue();
+};
+
+} // end namespace UE::Geometry
+} // end namespace UE
diff --git a/Engine/Plugins/Experimental/MeshModelingToolsetExp/Source/ModelingOperators/Public/SpaceDeformerOps/BendMeshOp.h b/Engine/Plugins/Experimental/MeshModelingToolsetExp/Source/ModelingOperators/Public/SpaceDeformerOps/BendMeshOp.h
new file mode 100644
--- /dev/null
+++ b/Engine/Plugins/Experimental/MeshModelingToolsetExp/Source/ModelingOperators/Public/SpaceDeformerOps/BendMeshOp.h
@@ -0,0 +1,25 @@
+// Copyright Epic Games, Inc. All Rights Reserved.
+
+#pragma once
+
+#include "CoreMinimal.h"
+#include "MeshSpaceDeformerOp.h"
+
+namespace UE
+{
+namespace Geometry
+{
+
+class MODELINGOPERATORS_API FBendMeshOp : public FMeshSpaceDeformerOp
+{
+public:
+	virtual void CalculateResult(FProgressCancel* Progress) override;
+
+	double BendDegrees = 90;
+	bool bLockBottom = false;
+
+protected:
+};
+
+} // end namespace UE::Geometry
+} // end namespace UE
diff --git a/Engine/Plugins/Experimental/MeshModelingToolsetExp/Source/ModelingOperators/Public/SpaceDeformerOps/FlareMeshOp.h b/Engine/Plugins/Experimental/MeshModelingToolsetExp/Source/ModelingOperators/Public/SpaceDeformerOps/FlareMeshOp.h
new file mode 100644
--- /dev/null
+++ b/Engine/Plugins/Experimental/MeshModelingToolsetExp/Source/ModelingOperators/Public/SpaceDeformerOps/FlareMeshOp.h
@@ -0,0 +1,34 @@
+// Copyright Epic Games, Inc. All Rights Reserved.
+
+#pragma once
+
+#include "CoreMinimal.h"
+#include "MeshSpaceDeformerOp.h"
+
+namespace UE
+{
+namespace Geometry
+{
+
+class MODELINGOPERATORS_API FFlareMeshOp : public FMeshSpaceDeformerOp
+{
+public:
+	virtual void CalculateResult(FProgressCancel* Progress) override;
+
+	/** 0% does nothing, 100% moves 2x away from Z axis at extremal point, -100% squishes down to Z axis at extremal point. */
+	double FlarePercentX = 100;
+	double FlarePercentY = 100;
+
+	/** 
+	 * Changes the type of flaring. When false, the flaring is the curve of sin(y) from 0 to Pi,
+	 * which makes the ends of the flaring sharp because the derivative is not 0. When true, the 
+	 * flaring is the curve of cos(y) + 1 from -Pi to Pi, which makes the ends smooth out back
+	 * into the shape.
+	 */
+	bool bSmoothEnds = false;
+protected:
+
+};
+
+} // end namespace UE::Geometry
+} // end namespace UE
diff --git a/Engine/Plugins/Experimental/MeshModelingToolsetExp/Source/ModelingOperators/Public/SpaceDeformerOps/TwistMeshOp.h b/Engine/Plugins/Experimental/MeshModelingToolsetExp/Source/ModelingOperators/Public/SpaceDeformerOps/TwistMeshOp.h
new file mode 100644
--- /dev/null
+++ b/Engine/Plugins/Experimental/MeshModelingToolsetExp/Source/ModelingOperators/Public/SpaceDeformerOps/TwistMeshOp.h
@@ -0,0 +1,25 @@
+// Copyright Epic Games, Inc. All Rights Reserved.
+
+#pragma once
+
+#include "CoreMinimal.h"
+#include "MeshSpaceDeformerOp.h"
+
+namespace UE
+{
+namespace Geometry
+{
+
+class MODELINGOPERATORS_API FTwistMeshOp : public FMeshSpaceDeformerOp
+{
+public:
+	virtual void CalculateResult(FProgressCancel* Progress) override;
+
+	double TwistDegrees = 180;
+	bool bLockBottom = false;
+protected:
+
+};
+
+} // end namespace UE::Geometry
+} // end namespace UE
diff --git a/Engine/Plugins/Experimental/MeshModelingToolsetExp/Source/ModelingOperatorsEditorOnly/Private/CompositionOps/VoxelBooleanMeshesOp.cpp b/Engine/Plugins/Experimental/MeshModelingToolsetExp/Source/ModelingOperatorsEditorOnly/Private/CompositionOps/VoxelBooleanMeshesOp.cpp
new file mode 100644
--- /dev/null
+++ b/Engine/Plugins/Experimental/MeshModelingToolsetExp/Source/ModelingOperatorsEditorOnly/Private/CompositionOps/VoxelBooleanMeshesOp.cpp
@@ -0,0 +1,241 @@
+// Copyright Epic Games, Inc. All Rights Reserved.
+
+#include "CompositionOps/VoxelBooleanMeshesOp.h"
+
+#include "Engine/StaticMesh.h"
+
+#include "MeshDescription.h"
+#include "StaticMeshAttributes.h"
+#include "MeshDescriptionToDynamicMesh.h"
+#include "MeshSimplification.h"
+#include "MeshConstraintsUtil.h"
+#include "CleaningOps/EditNormalsOp.h"
+
+
+// ProxyLOD currently only available on Windows. On other platforms we will make this a no-op
+#if PLATFORM_WINDOWS
+#define HAVE_PROXYLOD 1
+#else
+#define HAVE_PROXYLOD 0
+#endif
+
+#if HAVE_PROXYLOD
+#include "ProxyLODVolume.h"
+#endif
+
+#include "ExplicitUseGeometryMathTypes.h"		// using UE::Geometry::(math types)
+using namespace UE::Geometry;
+
+void FVoxelBooleanMeshesOp::CalculateResult(FProgressCancel* Progress)
+{
+#if HAVE_PROXYLOD
+	FMeshDescription ResultMeshDescription;
+	FStaticMeshAttributes Attributes(ResultMeshDescription);
+	Attributes.Register();
+
+	if (VoxelCount < 16)
+	{
+		VoxelCount = 16;
+	}
+
+	// Use the world space bounding box of each mesh to compute the voxel size
+	VoxelSizeD = ComputeVoxelSize();
+
+	// give this an absolute min since the user might scale both objects to zero..
+	VoxelSizeD = FMath::Max(VoxelSizeD, 0.001);
+
+	// world space units.
+	const double MaxNarrowBand = 3 * VoxelSizeD;
+	const double MaxIsoOffset = 2 * VoxelSizeD;
+
+	// IsoSurfaceD is in world units
+
+	const bool bShinkResult = (IsoSurfaceD < 0.);
+
+	// Note after the min/max operations of voxel-based CSG the interior distances or exterior 
+	// distances can be pretty f'd up so we need to do the initial post-CSG remesh with a "safe"
+	// value.
+	double CSGIsoSurface = 0.;
+
+	
+	// Create BooleanTool and Boolean the meshes.
+	TUniquePtr<IVoxelBasedCSG> VoxelCSGTool = IVoxelBasedCSG::CreateCSGTool(VoxelSizeD);
+	FVector MergedOrigin;
+
+	TArray<IVoxelBasedCSG::FPlacedMesh> PlacedMeshes;
+	for (const FInputMesh& InputMesh : InputMeshArray)
+	{
+		IVoxelBasedCSG::FPlacedMesh PlacedMesh;
+		PlacedMesh.Mesh = InputMesh.Mesh;
+		PlacedMesh.Transform = InputMesh.Transform;
+		PlacedMeshes.Add(PlacedMesh);
+	}
+
+	IVoxelBasedCSG::FPlacedMesh& A = PlacedMeshes[0];
+	IVoxelBasedCSG::FPlacedMesh& B = PlacedMeshes[1];
+
+	switch (Operation)
+	{
+	case EBooleanOperation::DifferenceAB:
+		CSGIsoSurface = 0.;
+		MergedOrigin = VoxelCSGTool->ComputeDifference(A, B, ResultMeshDescription, AdaptivityD, CSGIsoSurface);
+		break;
+	case EBooleanOperation::DifferenceBA:
+		CSGIsoSurface = 0.;
+		MergedOrigin = VoxelCSGTool->ComputeDifference(B, A, ResultMeshDescription, AdaptivityD, CSGIsoSurface);
+		break;
+	case EBooleanOperation::Intersect:
+		CSGIsoSurface = FMath::Clamp(IsoSurfaceD, -MaxIsoOffset, MaxIsoOffset);
+		MergedOrigin = VoxelCSGTool->ComputeIntersection(A, B, ResultMeshDescription, AdaptivityD, CSGIsoSurface);
+		break;
+	case EBooleanOperation::Union:
+		CSGIsoSurface = FMath::Clamp(IsoSurfaceD, 0., MaxIsoOffset);
+		MergedOrigin = VoxelCSGTool->ComputeUnion(A, B, ResultMeshDescription, AdaptivityD, CSGIsoSurface);
+		break;
+	default:
+		check(0);
+	}
+
+	ResultTransform = FTransform3d(MergedOrigin);
+
+	if (Progress && Progress->Cancelled())
+	{
+		return;
+	}
+
+	const bool bNeedAdditionalRemeshing = (CSGIsoSurface != IsoSurfaceD);
+
+	// iteratively expand (or contract) the surface
+	if ( bNeedAdditionalRemeshing )
+	{
+		// how additional remeshes do we need?
+		int32 NumRemeshes = FMath::CeilToInt( FMath::Abs(IsoSurfaceD - CSGIsoSurface) / MaxIsoOffset );
+
+		double DeltaIsoSurface = (IsoSurfaceD - CSGIsoSurface) / NumRemeshes;
+
+		for (int32 i = 0; i < NumRemeshes; ++i)
+		{
+			if (Progress && Progress->Cancelled())
+			{
+				return;
+			}
+
+			IVoxelBasedCSG::FPlacedMesh PlacedMesh;
+			PlacedMesh.Mesh = &ResultMeshDescription;
+			PlacedMesh.Transform = static_cast<FTransform>(ResultTransform);
+
+			TArray< IVoxelBasedCSG::FPlacedMesh> SourceMeshes;
+			SourceMeshes.Add(PlacedMesh);
+
+			// union of only one mesh. we are using this to voxelize and remesh with and offset.
+			MergedOrigin = VoxelCSGTool->ComputeUnion(SourceMeshes, ResultMeshDescription, AdaptivityD, DeltaIsoSurface);
+
+			ResultTransform = FTransform3d(MergedOrigin);
+		}
+	}
+
+	// Convert to dynamic mesh
+	FMeshDescriptionToDynamicMesh Converter;
+	Converter.Convert(&ResultMeshDescription, *ResultMesh);
+
+	if (Progress && Progress->Cancelled())
+	{
+		return;
+	}
+
+	// for now we automatically fix up the normals if simplificaiton was requested.
+	bool bFixNormals = bAutoSimplify;
+
+	if (bAutoSimplify)
+	{
+		FQEMSimplification Reducer(ResultMesh.Get());
+		FMeshConstraints constraints;
+		FMeshConstraintsUtil::ConstrainAllSeams(constraints, *ResultMesh, true, false);
+		Reducer.SetExternalConstraints(MoveTemp(constraints));
+		Reducer.Progress = Progress;
+
+		const double MaxDisplacementSqr = 3. * VoxelSizeD * VoxelSizeD;
+
+		Reducer.SimplifyToMaxError(MaxDisplacementSqr);
+
+	}
+
+
+	if (bFixNormals)
+	{
+		TSharedPtr<FDynamicMesh3, ESPMode::ThreadSafe> OpResultMesh(ExtractResult().Release()); // moved the unique pointer
+
+		// Recompute the normals
+		FEditNormalsOp EditNormalsOp;
+		EditNormalsOp.OriginalMesh = OpResultMesh; // the tool works on a deep copy of this mesh.
+		EditNormalsOp.bFixInconsistentNormals = true;
+		EditNormalsOp.bInvertNormals = false;
+		EditNormalsOp.bRecomputeNormals = true;
+		EditNormalsOp.NormalCalculationMethod = ENormalCalculationMethod::AreaAngleWeighting;
+		EditNormalsOp.SplitNormalMethod = ESplitNormalMethod::FaceNormalThreshold;
+		EditNormalsOp.bAllowSharpVertices = true;
+		EditNormalsOp.NormalSplitThreshold = 60.f;
+
+		EditNormalsOp.SetTransform(FTransform(ResultTransform));
+
+		EditNormalsOp.CalculateResult(Progress);
+
+		ResultMesh = EditNormalsOp.ExtractResult(); // return the edit normals operator copy to this tool.
+	}
+
+#else	// HAVE_PROXYLOD
+
+	FMeshDescriptionToDynamicMesh Converter;
+	Converter.Convert(InputMeshArray[0].Mesh, *ResultMesh);
+	ResultTransform = FTransform3d(InputMeshArray[0].Transform);
+
+#endif	// HAVE_PROXYLOD
+
+}
+
+float FVoxelBooleanMeshesOp::ComputeVoxelSize() const 
+{ 
+	float Size = 0.f;
+	for (int32 i = 0; i < InputMeshArray.Num(); ++i)
+	{
+		//Bounding box
+		FVector BBoxMin(FLT_MAX, FLT_MAX, FLT_MAX);
+		FVector BBoxMax(-FLT_MAX, -FLT_MAX, -FLT_MAX);
+
+		// Use the scale define by the transform.  We don't care about actual placement
+		// in the world for this.
+		FVector Scale = InputMeshArray[i].Transform.GetScale3D();
+		const FMeshDescription&  MeshDescription = *InputMeshArray[i].Mesh;
+
+		TArrayView<const FVector3f> VertexPositions = MeshDescription.GetVertexPositions().GetRawArray();
+		for (const FVertexID VertexID : MeshDescription.Vertices().GetElementIDs())
+		{
+			const FVector& Pos = VertexPositions[VertexID];
+            
+			BBoxMin.X = FMath::Min(BBoxMin.X, Pos.X);
+			BBoxMin.Y = FMath::Min(BBoxMin.Y, Pos.Y);
+			BBoxMin.Z = FMath::Min(BBoxMin.Z, Pos.Z);
+
+			BBoxMax.X = FMath::Max(BBoxMax.X, Pos.X);
+			BBoxMax.Y = FMath::Max(BBoxMax.Y, Pos.Y);
+			BBoxMax.Z = FMath::Max(BBoxMax.Z, Pos.Z);
+		}
+
+		// The size of the BBox in each direction
+		FVector Extents(BBoxMax.X - BBoxMin.X, BBoxMax.Y - BBoxMin.Y, BBoxMax.Z - BBoxMin.Z);
+		
+		// Scale with the local space scale.
+		Extents.X = Extents.X * FMath::Abs(Scale.X);
+		Extents.Y = Extents.Y * FMath::Abs(Scale.Y);
+		Extents.Z = Extents.Z * FMath::Abs(Scale.Z);
+
+		float MajorAxisSize = FMath::Max3(Extents.X, Extents.Y, Extents.Z);
+
+
+		Size = FMath::Max(MajorAxisSize / VoxelCount, Size);
+
+	}
+	
+	return Size;
+
+}
diff --git a/Engine/Plugins/Experimental/MeshModelingToolsetExp/Source/ModelingOperatorsEditorOnly/Private/CompositionOps/VoxelMergeMeshesOp.cpp b/Engine/Plugins/Experimental/MeshModelingToolsetExp/Source/ModelingOperatorsEditorOnly/Private/CompositionOps/VoxelMergeMeshesOp.cpp
new file mode 100644
--- /dev/null
+++ b/Engine/Plugins/Experimental/MeshModelingToolsetExp/Source/ModelingOperatorsEditorOnly/Private/CompositionOps/VoxelMergeMeshesOp.cpp
@@ -0,0 +1,204 @@
+// Copyright Epic Games, Inc. All Rights Reserved.
+
+#include "CompositionOps/VoxelMergeMeshesOp.h"
+
+#include "Engine/StaticMesh.h"
+
+#include "MeshDescription.h"
+#include "StaticMeshAttributes.h"
+#include "MeshDescriptionToDynamicMesh.h"
+#include "MeshSimplification.h"
+#include "MeshConstraintsUtil.h"
+#include "CleaningOps/EditNormalsOp.h"
+
+// ProxyLOD currently only available on Windows. On other platforms we will make this a no-op
+#if PLATFORM_WINDOWS
+#define HAVE_PROXYLOD 1
+#else
+#define HAVE_PROXYLOD 0
+#endif
+
+#if HAVE_PROXYLOD
+#include "ProxyLODVolume.h"
+#endif
+
+#include "ExplicitUseGeometryMathTypes.h"		// using UE::Geometry::(math types)
+using namespace UE::Geometry;
+
+void FVoxelMergeMeshesOp::CalculateResult(FProgressCancel* Progress)
+{
+#if HAVE_PROXYLOD
+
+	FMeshDescription MergedMeshesDescription;
+	FStaticMeshAttributes Attributes(MergedMeshesDescription);
+	Attributes.Register();
+
+	// Use the world space bounding box of each mesh to compute the voxel size
+	VoxelSizeD = ComputeVoxelSize();
+
+	// give this an absolute min since the user might scale both objects to zero..
+	VoxelSizeD = FMath::Max(VoxelSizeD, 0.001);
+
+	// Create CSGTool and merge the meshes.
+	TUniquePtr<IVoxelBasedCSG> VoxelCSGTool = IVoxelBasedCSG::CreateCSGTool(VoxelSizeD);
+
+	TArray<IVoxelBasedCSG::FPlacedMesh> PlacedMeshes;
+	for (const FInputMesh& InputMesh : InputMeshArray)
+	{
+		IVoxelBasedCSG::FPlacedMesh PlacedMesh;
+		PlacedMesh.Mesh = InputMesh.Mesh;
+		PlacedMesh.Transform = InputMesh.Transform;
+		PlacedMeshes.Add(PlacedMesh);
+	}
+
+	// world space units.
+	const double MaxNarrowBand = 3 * VoxelSizeD;
+	const double MaxIsoOffset = 2 * VoxelSizeD;
+	const double CSGIsoSurface = FMath::Clamp(IsoSurfaceD, 0., MaxIsoOffset); // the interior distance values maybe messed up when doing a union.
+	FVector MergedOrigin = VoxelCSGTool->ComputeUnion(PlacedMeshes, MergedMeshesDescription, AdaptivityD, CSGIsoSurface);
+	ResultTransform = FTransform3d(MergedOrigin);
+
+	if (Progress && Progress->Cancelled())
+	{
+		return;
+	}
+
+	const bool bNeedAdditionalRemeshing = (CSGIsoSurface != IsoSurfaceD);
+
+	// iteratively expand (or contract) the surface
+	if (bNeedAdditionalRemeshing)
+	{
+		// how additional remeshes do we need?
+		int32 NumRemeshes = FMath::CeilToInt(FMath::Abs(IsoSurfaceD - CSGIsoSurface) / MaxIsoOffset);
+
+		double DeltaIsoSurface = (IsoSurfaceD - CSGIsoSurface) / NumRemeshes;
+
+		for (int32 i = 0; i < NumRemeshes; ++i)
+		{
+			if (Progress && Progress->Cancelled())
+			{
+				return;
+			}
+
+			IVoxelBasedCSG::FPlacedMesh PlacedMesh;
+			PlacedMesh.Mesh = &MergedMeshesDescription;
+			PlacedMesh.Transform = static_cast<FTransform>(ResultTransform);
+
+			TArray< IVoxelBasedCSG::FPlacedMesh> SourceMeshes;
+			SourceMeshes.Add(PlacedMesh);
+
+			// union of only one mesh. we are using this to voxelize and remesh with and offset.
+			MergedOrigin = VoxelCSGTool->ComputeUnion(SourceMeshes, MergedMeshesDescription, AdaptivityD, DeltaIsoSurface);
+
+			ResultTransform = FTransform3d(MergedOrigin);
+		}
+	}
+
+	// Convert to dynamic mesh
+	FMeshDescriptionToDynamicMesh Converter;
+	Converter.Convert(&MergedMeshesDescription, *ResultMesh);
+
+	if (Progress && Progress->Cancelled())
+	{
+		return;
+	}
+
+	// for now we automatically fix up the normals if simplificaiton was requested.
+	bool bFixNormals = bAutoSimplify;
+
+	if (bAutoSimplify)
+	{
+		FQEMSimplification Reducer(ResultMesh.Get());
+		FMeshConstraints constraints;
+		FMeshConstraintsUtil::ConstrainAllSeams(constraints, *ResultMesh, true, false);
+		Reducer.SetExternalConstraints(MoveTemp(constraints));
+		Reducer.Progress = Progress;
+
+		const double MaxDisplacementSqr = 3. *VoxelSizeD * VoxelSizeD;
+
+		Reducer.SimplifyToMaxError(MaxDisplacementSqr);
+
+
+	}
+
+	if (bFixNormals)
+	{
+
+		TSharedPtr<FDynamicMesh3, ESPMode::ThreadSafe> OpResultMesh(ExtractResult().Release()); // moved the unique pointer
+
+		// Recompute the normals
+		FEditNormalsOp EditNormalsOp;
+		EditNormalsOp.OriginalMesh = OpResultMesh; // the tool works on a deep copy of this mesh.
+		EditNormalsOp.bFixInconsistentNormals = true;
+		EditNormalsOp.bInvertNormals = false;
+		EditNormalsOp.bRecomputeNormals = true;
+		EditNormalsOp.NormalCalculationMethod = ENormalCalculationMethod::AreaAngleWeighting;
+		EditNormalsOp.SplitNormalMethod = ESplitNormalMethod::FaceNormalThreshold;
+		EditNormalsOp.bAllowSharpVertices = true;
+		EditNormalsOp.NormalSplitThreshold = 60.f;
+
+		EditNormalsOp.SetTransform(FTransform(ResultTransform));
+
+		EditNormalsOp.CalculateResult(Progress);
+
+		ResultMesh = EditNormalsOp.ExtractResult(); // return the edit normals operator copy to this tool.
+	}
+
+
+#else	// HAVE_PROXYLOD
+
+	FMeshDescriptionToDynamicMesh Converter;
+	Converter.Convert(InputMeshArray[0].Mesh, *ResultMesh);
+	ResultTransform = FTransform3d(InputMeshArray[0].Transform);
+
+#endif	// HAVE_PROXYLOD
+
+}
+
+
+float FVoxelMergeMeshesOp::ComputeVoxelSize() const
+{
+	float Size = 0.f;
+	for (int32 i = 0; i < InputMeshArray.Num(); ++i)
+	{
+		//Bounding box
+		FVector BBoxMin(FLT_MAX, FLT_MAX, FLT_MAX);
+		FVector BBoxMax(-FLT_MAX, -FLT_MAX, -FLT_MAX);
+
+		// Use the scale define by the transform.  We don't care about actual placement
+		// in the world for this.
+		FVector Scale = InputMeshArray[i].Transform.GetScale3D();
+		const FMeshDescription&  MeshDescription = *InputMeshArray[i].Mesh;
+
+		TArrayView<const FVector3f> VertexPositions = MeshDescription.GetVertexPositions().GetRawArray();
+		for (const FVertexID VertexID : MeshDescription.Vertices().GetElementIDs())
+		{
+			const FVector& Pos = VertexPositions[VertexID];
+
+			BBoxMin.X = FMath::Min(BBoxMin.X, Pos.X);
+			BBoxMin.Y = FMath::Min(BBoxMin.Y, Pos.Y);
+			BBoxMin.Z = FMath::Min(BBoxMin.Z, Pos.Z);
+
+			BBoxMax.X = FMath::Max(BBoxMax.X, Pos.X);
+			BBoxMax.Y = FMath::Max(BBoxMax.Y, Pos.Y);
+			BBoxMax.Z = FMath::Max(BBoxMax.Z, Pos.Z);
+		}
+
+		// The size of the BBox in each direction
+		FVector Extents(BBoxMax.X - BBoxMin.X, BBoxMax.Y - BBoxMin.Y, BBoxMax.Z - BBoxMin.Z);
+
+		// Scale with the local space scale.
+		Extents.X = Extents.X * FMath::Abs(Scale.X);
+		Extents.Y = Extents.Y * FMath::Abs(Scale.Y);
+		Extents.Z = Extents.Z * FMath::Abs(Scale.Z);
+
+		float MajorAxisSize = FMath::Max3(Extents.X, Extents.Y, Extents.Z);
+
+
+		Size = FMath::Max(MajorAxisSize / VoxelCount, Size);
+
+	}
+
+	return Size;
+
+}
diff --git a/Engine/Plugins/Experimental/MeshModelingToolsetExp/Source/ModelingOperatorsEditorOnly/Private/ParameterizationOps/CalculateTangentsOp.cpp b/Engine/Plugins/Experimental/MeshModelingToolsetExp/Source/ModelingOperatorsEditorOnly/Private/ParameterizationOps/CalculateTangentsOp.cpp
new file mode 100644
--- /dev/null
+++ b/Engine/Plugins/Experimental/MeshModelingToolsetExp/Source/ModelingOperatorsEditorOnly/Private/ParameterizationOps/CalculateTangentsOp.cpp
@@ -0,0 +1,211 @@
+// Copyright Epic Games, Inc. All Rights Reserved.
+
+#include "ParameterizationOps/CalculateTangentsOp.h"
+#include "VectorUtil.h"
+
+#define WITH_MIKKTSPACE 1
+
+#if WITH_MIKKTSPACE
+#include "mikktspace.h"
+#endif //WITH_MIKKTSPACE
+
+#include "ExplicitUseGeometryMathTypes.h"		// using UE::Geometry::(math types)
+using namespace UE::Geometry;
+
+void FCalculateTangentsOp::CalculateResult(FProgressCancel* Progress)
+{
+	if (Progress && Progress->Cancelled())
+	{
+		return;
+	}
+
+	TUniquePtr<FMeshTangentsd> Tangents = MakeUnique<FMeshTangentsd>(SourceMesh.Get());
+
+	if (SourceMesh->HasAttributes() == false
+		|| SourceMesh->Attributes()->NumUVLayers() < 1
+		|| SourceMesh->Attributes()->NumNormalLayers() == 0)
+	{
+		bNoAttributesError = true;
+		Tangents->InitializeTriVertexTangents(true);
+		SetResult(MoveTemp(Tangents));
+		return;
+	}
+
+
+	switch (CalculationMethod)
+	{
+
+	default:
+	case EMeshTangentsType::FastMikkTSpace:
+		CalculateStandard(Progress, Tangents);
+		break;
+
+	case EMeshTangentsType::PerTriangle:
+		CalculateStandard(Progress, Tangents);
+		break;
+
+	case EMeshTangentsType::MikkTSpace:
+#if WITH_MIKKTSPACE
+		CalculateMikkT(Progress, Tangents);
+#else
+		CalculateStandard(Progress, Tangents);
+#endif
+		break;
+
+
+	case EMeshTangentsType::CopyExisting:
+		CopyFromSource(Progress, Tangents);
+		break;
+	}
+
+	SetResult(MoveTemp(Tangents));
+}
+
+
+void FCalculateTangentsOp::CalculateStandard(FProgressCancel* Progress, TUniquePtr<FMeshTangentsd>& Tangents)
+{
+	const FDynamicMeshNormalOverlay* NormalOverlay = SourceMesh->Attributes()->PrimaryNormals();
+	const FDynamicMeshUVOverlay* UVOverlay = SourceMesh->Attributes()->GetUVLayer(0);
+
+	FComputeTangentsOptions Options;
+	Options.bAveraged = (CalculationMethod != EMeshTangentsType::PerTriangle);
+
+	Tangents->ComputeTriVertexTangents(NormalOverlay, UVOverlay, Options);
+}
+
+
+void FCalculateTangentsOp::CopyFromSource(FProgressCancel* Progress, TUniquePtr<FMeshTangentsd>& Tangents)
+{
+	int32 NumTangents = SourceTangents->GetTangents().Num();
+	if ( ! ensure(NumTangents == SourceMesh->MaxTriangleID() * 3) )
+	{
+		bNoAttributesError = true;
+		Tangents->InitializeTriVertexTangents(true);
+		return;
+	}
+
+	Tangents->CopyTriVertexTangents(*SourceTangents);
+}
+
+
+
+#if WITH_MIKKTSPACE
+namespace DynamicMeshMikkTInterface
+{
+	struct FDynamicMeshInfo
+	{
+		const FDynamicMesh3* Mesh;
+		const FDynamicMeshNormalOverlay* NormalOverlay;
+		const FDynamicMeshUVOverlay* UVOverlay;
+		FMeshTangentsd* TangentsOut;
+	};
+
+
+	int MikkGetNumFaces(const SMikkTSpaceContext* Context)
+	{
+		const FDynamicMeshInfo* MeshInfo = (const FDynamicMeshInfo*)(Context->m_pUserData);
+		return MeshInfo->Mesh->TriangleCount();
+	}
+
+	int MikkGetNumVertsOfFace(const SMikkTSpaceContext* Context, const int FaceIdx)
+	{
+		const FDynamicMeshInfo* MeshInfo = (const FDynamicMeshInfo*)(Context->m_pUserData);
+		return MeshInfo->Mesh->IsTriangle(FaceIdx) ? 3 : 0;
+	}
+
+	void MikkGetPosition(const SMikkTSpaceContext* Context, float Position[3], const int FaceIdx, const int VertIdx)
+	{
+		const FDynamicMeshInfo* MeshInfo = (const FDynamicMeshInfo*)(Context->m_pUserData);
+		FIndex3i Triangle = MeshInfo->Mesh->GetTriangle(FaceIdx);
+		FVector3d VertexPos = MeshInfo->Mesh->GetVertex(Triangle[VertIdx]);
+		Position[0] = (float)VertexPos.X;
+		Position[1] = (float)VertexPos.Y;
+		Position[2] = (float)VertexPos.Z;
+	}
+
+	void MikkGetNormal(const SMikkTSpaceContext* Context, float Normal[3], const int FaceIdx, const int VertIdx)
+	{
+		const FDynamicMeshInfo* MeshInfo = (const FDynamicMeshInfo*)(Context->m_pUserData);
+		FIndex3i NormTriangle = MeshInfo->NormalOverlay->GetTriangle(FaceIdx);
+		FVector3f Normalf;
+		MeshInfo->NormalOverlay->GetElement(NormTriangle[VertIdx], Normalf);
+		Normal[0] = Normalf.X;
+		Normal[1] = Normalf.Y;
+		Normal[2] = Normalf.Z;
+	}
+
+	void MikkGetTexCoord(const SMikkTSpaceContext* Context, float UV[2], const int FaceIdx, const int VertIdx)
+	{
+		const FDynamicMeshInfo* MeshInfo = (const FDynamicMeshInfo*)(Context->m_pUserData);
+		FIndex3i UVTriangle = MeshInfo->UVOverlay->GetTriangle(FaceIdx);
+		FVector2f UVf;
+		MeshInfo->UVOverlay->GetElement(UVTriangle[VertIdx], UVf);
+		UV[0] = UVf.X;
+		UV[1] = UVf.Y;
+	}
+
+	void MikkSetTSpaceBasic(const SMikkTSpaceContext* Context, const float Tangent[3], const float BitangentSign, const int FaceIdx, const int VertIdx)
+	{
+		FDynamicMeshInfo* MeshInfo = (FDynamicMeshInfo*)(Context->m_pUserData);
+		FVector3d Tangentd(Tangent[0], Tangent[1], Tangent[2]);
+
+		FIndex3i NormTriangle = MeshInfo->NormalOverlay->GetTriangle(FaceIdx);
+		FVector3d Normald = (FVector3d)MeshInfo->NormalOverlay->GetElement(NormTriangle[VertIdx]);
+		FVector3d Bitangentd = VectorUtil::Bitangent(Normald, Tangentd, -(double)BitangentSign);
+
+		MeshInfo->TangentsOut->SetPerTriangleTangent(FaceIdx, VertIdx, Tangentd, Bitangentd);
+	}
+
+
+	void MikkSetTSpac(const SMikkTSpaceContext* Context, const float Tangent[3], const float BiTangent[3],
+		const float MagS, const float MagT, const tbool bIsOrientationPreserving,
+		const int FaceIdx, const int VertIdx)
+	{
+		FDynamicMeshInfo* MeshInfo = (FDynamicMeshInfo*)(Context->m_pUserData);
+		FVector3d Tangentd(Tangent[0], Tangent[1], Tangent[2]);
+		FVector3d Bitangentd(BiTangent[0], BiTangent[1], BiTangent[2]);
+		MeshInfo->TangentsOut->SetPerTriangleTangent(FaceIdx, VertIdx, Tangentd, Bitangentd);
+	}
+
+}
+
+
+
+void FCalculateTangentsOp::CalculateMikkT(FProgressCancel* Progress, TUniquePtr<FMeshTangentsd>& Tangents)
+{
+	Tangents->InitializeTriVertexTangents(true);
+
+	DynamicMeshMikkTInterface::FDynamicMeshInfo MeshInfo;
+	MeshInfo.Mesh = SourceMesh.Get();
+	MeshInfo.NormalOverlay = SourceMesh->Attributes()->PrimaryNormals();
+	MeshInfo.UVOverlay = SourceMesh->Attributes()->GetUVLayer(0);;
+	MeshInfo.TangentsOut = Tangents.Get();
+
+	SMikkTSpaceInterface MikkTInterface;
+	MikkTInterface.m_getNormal = DynamicMeshMikkTInterface::MikkGetNormal;
+	MikkTInterface.m_getNumFaces = DynamicMeshMikkTInterface::MikkGetNumFaces;
+	MikkTInterface.m_getNumVerticesOfFace = DynamicMeshMikkTInterface::MikkGetNumVertsOfFace;
+	MikkTInterface.m_getPosition = DynamicMeshMikkTInterface::MikkGetPosition;
+	MikkTInterface.m_getTexCoord = DynamicMeshMikkTInterface::MikkGetTexCoord;
+	MikkTInterface.m_setTSpaceBasic = DynamicMeshMikkTInterface::MikkSetTSpaceBasic;
+	MikkTInterface.m_setTSpace = nullptr;
+
+	SMikkTSpaceContext MikkTContext;
+	MikkTContext.m_pInterface = &MikkTInterface;
+	MikkTContext.m_pUserData = (void*)(&MeshInfo);
+	MikkTContext.m_bIgnoreDegenerates = false;
+	//MikkTContext.m_bIgnoreDegenerates = true;
+	//MikkTContext.m_bIgnoreDegenerates = bIgnoreDegenerateTriangles;
+
+	// execute mikkt
+	genTangSpaceDefault(&MikkTContext);
+}
+
+#else
+
+void FCalculateTangentsOp::CalculateMikkT(FProgressCancel* Progress, TUniquePtr<FMeshTangentsd>& Tangents)
+{
+	check(false);
+}
+
+#endif
\ No newline at end of file
diff --git a/Engine/Plugins/Experimental/MeshModelingToolsetExp/Source/ModelingOperatorsEditorOnly/Public/CompositionOps/VoxelBooleanMeshesOp.h b/Engine/Plugins/Experimental/MeshModelingToolsetExp/Source/ModelingOperatorsEditorOnly/Public/CompositionOps/VoxelBooleanMeshesOp.h
new file mode 100644
--- /dev/null
+++ b/Engine/Plugins/Experimental/MeshModelingToolsetExp/Source/ModelingOperatorsEditorOnly/Public/CompositionOps/VoxelBooleanMeshesOp.h
@@ -0,0 +1,60 @@
+// Copyright Epic Games, Inc. All Rights Reserved.
+
+#pragma once
+
+#include "CoreMinimal.h"
+#include "Util/ProgressCancel.h"
+#include "ModelingOperators.h"
+
+struct FMeshDescription;
+
+namespace UE
+{
+namespace Geometry
+{
+
+
+class MODELINGOPERATORSEDITORONLY_API FVoxelBooleanMeshesOp : public FDynamicMeshOperator
+{
+public:
+	virtual ~FVoxelBooleanMeshesOp() {}
+
+	enum class EBooleanOperation
+	{
+		DifferenceAB = 0,
+		DifferenceBA = 1,
+		Intersect = 2,
+		Union = 3
+	};
+
+	struct FInputMesh
+	{
+		const FMeshDescription* Mesh;
+		FTransform Transform;
+	};
+
+	// inputs
+	TArray<FInputMesh> InputMeshArray;
+	int32 VoxelCount = 128;
+	double VoxelSizeD = 1.0;
+	double AdaptivityD = 0;
+	double IsoSurfaceD = 0;
+	bool   bAutoSimplify = false;
+	EBooleanOperation Operation = EBooleanOperation::DifferenceBA;
+
+	//
+	// FDynamicMeshOperator implementation
+	// 
+
+	virtual void CalculateResult(FProgressCancel* Progress) override;
+
+private:
+
+	// compute the voxel size based on the voxel count and the input geometry
+	float ComputeVoxelSize() const;
+
+};
+
+} // end namespace UE::Geometry
+} // end namespace UE
+
diff --git a/Engine/Plugins/Experimental/MeshModelingToolsetExp/Source/ModelingOperatorsEditorOnly/Public/CompositionOps/VoxelMergeMeshesOp.h b/Engine/Plugins/Experimental/MeshModelingToolsetExp/Source/ModelingOperatorsEditorOnly/Public/CompositionOps/VoxelMergeMeshesOp.h
new file mode 100644
--- /dev/null
+++ b/Engine/Plugins/Experimental/MeshModelingToolsetExp/Source/ModelingOperatorsEditorOnly/Public/CompositionOps/VoxelMergeMeshesOp.h
@@ -0,0 +1,53 @@
+// Copyright Epic Games, Inc. All Rights Reserved.
+
+#pragma once
+
+#include "CoreMinimal.h"
+#include "Util/ProgressCancel.h"
+#include "ModelingOperators.h"
+
+struct FMeshDescription;
+
+namespace UE
+{
+namespace Geometry
+{
+
+
+class MODELINGOPERATORSEDITORONLY_API FVoxelMergeMeshesOp : public FDynamicMeshOperator
+{
+public:
+	virtual ~FVoxelMergeMeshesOp() {}
+
+	struct FInputMesh
+	{
+		const FMeshDescription* Mesh;
+		FTransform Transform;
+	};
+
+	// inputs
+	TArray<FInputMesh> InputMeshArray;
+	int32 VoxelCount = 128;
+	double VoxelSizeD = 1.0;
+	double AdaptivityD = 0;
+	double IsoSurfaceD = 0;
+	bool bAutoSimplify = false;
+
+	double FastCollapseVoxelSizeMultipler = 1.2;
+	int NumFastCollapsePasses = 10;
+
+	//
+	// FDynamicMeshOperator implementation
+	// 
+
+	virtual void CalculateResult(FProgressCancel* Progress) override;
+
+private:
+
+	// compute the voxel size based on the voxel count and the input geometry
+	float ComputeVoxelSize() const;
+};
+
+
+} // end namespace UE::Geometry
+} // end namespace UE
diff --git a/Engine/Plugins/Experimental/MeshModelingToolsetExp/Source/ModelingOperatorsEditorOnly/Public/ParameterizationOps/CalculateTangentsOp.h b/Engine/Plugins/Experimental/MeshModelingToolsetExp/Source/ModelingOperatorsEditorOnly/Public/ParameterizationOps/CalculateTangentsOp.h
new file mode 100644
--- /dev/null
+++ b/Engine/Plugins/Experimental/MeshModelingToolsetExp/Source/ModelingOperatorsEditorOnly/Public/ParameterizationOps/CalculateTangentsOp.h
@@ -0,0 +1,61 @@
+// Copyright Epic Games, Inc. All Rights Reserved.
+
+#pragma once
+
+#include "CoreMinimal.h"
+#include "Util/ProgressCancel.h"
+#include "ModelingOperators.h"
+#include "DynamicMesh/MeshTangents.h"
+#include "CalculateTangentsOp.generated.h"
+
+
+
+UENUM()
+enum class EMeshTangentsType : uint8
+{
+	/** Standard MikkTSpace Tangent Calculation */
+	MikkTSpace = 0,
+	/** MikkT-like Blended Per-Triangle Tangents, with Blending based on existing Mesh/UV/Normal Topology */
+	FastMikkTSpace = 1,
+	/** Project per-triangle Tangents projected onto Normals */
+	PerTriangle = 2,
+	/** Use Source Mesh Tangents */
+	CopyExisting = 3
+};
+
+namespace UE
+{
+namespace Geometry
+{
+
+
+class MODELINGOPERATORSEDITORONLY_API FCalculateTangentsOp : public TGenericDataOperator<UE::Geometry::FMeshTangentsd>
+{
+public:
+	virtual ~FCalculateTangentsOp() {}
+
+	// inputs
+	TSharedPtr<FDynamicMesh3, ESPMode::ThreadSafe> SourceMesh;
+	TSharedPtr<UE::Geometry::FMeshTangentsf, ESPMode::ThreadSafe> SourceTangents;
+
+	// parameters
+	EMeshTangentsType CalculationMethod;
+
+	// error flags
+	bool bNoAttributesError = false;
+
+	//
+	// TGenericDataOperator implementation
+	// 
+
+	virtual void CalculateResult(FProgressCancel* Progress) override;
+
+
+	virtual void CalculateStandard(FProgressCancel* Progress, TUniquePtr<UE::Geometry::FMeshTangentsd>& Tangents);
+	virtual void CalculateMikkT(FProgressCancel* Progress, TUniquePtr<UE::Geometry::FMeshTangentsd>& Tangents);
+	virtual void CopyFromSource(FProgressCancel* Progress, TUniquePtr<UE::Geometry::FMeshTangentsd>& Tangents);
+};
+
+} // end namespace UE::Geometry
+} // end namespace UE
+
diff --git a/Engine/Plugins/Runtime/MeshModelingToolset/Source/MeshModelingTools/MeshModelingTools.Build.cs b/Engine/Plugins/Runtime/MeshModelingToolset/Source/MeshModelingTools/MeshModelingTools.Build.cs
new file mode 100644
--- /dev/null
+++ b/Engine/Plugins/Runtime/MeshModelingToolset/Source/MeshModelingTools/MeshModelingTools.Build.cs
@@ -0,0 +1,68 @@
+// Copyright Epic Games, Inc. All Rights Reserved.
+
+using UnrealBuildTool;
+
+public class MeshModelingTools : ModuleRules
+{
+	public MeshModelingTools(ReadOnlyTargetRules Target) : base(Target)
+	{
+		PCHUsage = ModuleRules.PCHUsageMode.UseExplicitOrSharedPCHs;
+		
+		PublicIncludePaths.AddRange(
+			new string[] {
+				// ... add public include paths required here ...
+			}
+			);
+				
+		PublicIncludePathModuleNames.AddRange(
+			new string[] {
+				// "AnimationCore",			// For the BoneWeights.h include
+			}
+		);
+	
+		PublicDependencyModuleNames.AddRange(
+			new string[]
+			{
+				"Core",
+				// "Eigen",
+                "InteractiveToolsFramework",
+				"GeometryCore",
+				// "GeometryFramework",
+				// "GeometryAlgorithms",
+				"DynamicMesh",
+				// "MeshConversion",
+				// "MeshDescription",
+				// "StaticMeshDescription",
+				// "SkeletalMeshDescription",
+				"ModelingComponents",
+				// "ModelingOperators",
+
+				// ... add other public dependencies that you statically link with here ...
+			}
+            );
+			
+		
+		PrivateDependencyModuleNames.AddRange(
+			new string[]
+			{
+				"CoreUObject",
+				"Engine",
+				// "RenderCore",
+				"MeshModelingToolsExp",
+				// "ModelingOperators",
+				// "InputCore",
+				// "PhysicsCore",
+
+				// ... add private dependencies that you statically link with here ...
+			}
+			);
+		
+		
+		DynamicallyLoadedModuleNames.AddRange(
+			new string[]
+			{
+				// ... add any modules that your module loads dynamically here ...
+			}
+			);
+	}
+}
