ModelingTools: Stopgap fix for GetDynamicMeshCopy w/ tangents.

#rnx

diff --git a/Engine/Plugins/Experimental/MeshModelingToolsetExp/Source/ModelingComponents/Private/ModelingToolTargetUtil.cpp b/Engine/Plugins/Experimental/MeshModelingToolsetExp/Source/ModelingComponents/Private/ModelingToolTargetUtil.cpp
--- a/Engine/Plugins/Experimental/MeshModelingToolsetExp/Source/ModelingComponents/Private/ModelingToolTargetUtil.cpp
+++ b/Engine/Plugins/Experimental/MeshModelingToolsetExp/Source/ModelingComponents/Private/ModelingToolTargetUtil.cpp
@@ -166,8 +166,9 @@ FDynamicMesh3 UE::ToolTarget::GetDynamicMeshCopy(UToolTarget* Target, bool bWant
 		return Mesh;
 	}
 
+	// TODO: Handle tangent computation. For now skip if tangents requested.
 	IDynamicMeshProvider* DynamicMeshProvider = Cast<IDynamicMeshProvider>(Target);
-	if (DynamicMeshProvider)
+	if (DynamicMeshProvider && !bWantMeshTangents)
 	{
 		return *DynamicMeshProvider->GetDynamicMesh();
 	}
