fix enum comments on bake transform tool options

#rnx

diff --git a/Engine/Plugins/Experimental/MeshModelingToolsetExp/Source/MeshModelingToolsExp/Public/BakeTransformTool.h b/Engine/Plugins/Experimental/MeshModelingToolsetExp/Source/MeshModelingToolsExp/Public/BakeTransformTool.h
--- a/Engine/Plugins/Experimental/MeshModelingToolsetExp/Source/MeshModelingToolsExp/Public/BakeTransformTool.h
+++ b/Engine/Plugins/Experimental/MeshModelingToolsetExp/Source/MeshModelingToolsExp/Public/BakeTransformTool.h
@@ -37,9 +37,12 @@ protected:
 UENUM()
 enum class EBakeScaleMethod : uint8
 {
-	BakeFullScale,			// bake all scale information, so the component has scale of 1 on all axes
-	BakeNonuniformScale,	// bake the non-uniform scale, so the component has a uniform scale
-	DoNotBakeScale			// do not bake any scaling
+	// bake all scale information, so the component has scale of 1 on all axes
+	BakeFullScale,
+	// bake the non-uniform scale, so the component has a uniform scale
+	BakeNonuniformScale,
+	// do not bake any scaling
+	DoNotBakeScale
 };
 
 
