fix confusing tooltips on jacketing tool (due to code comments leaking through to tooltips)

#rnx

diff --git a/Engine/Plugins/Experimental/MeshModelingToolsetExp/Source/MeshModelingToolsExp/Public/RemoveOccludedTrianglesTool.h b/Engine/Plugins/Experimental/MeshModelingToolsetExp/Source/MeshModelingToolsExp/Public/RemoveOccludedTrianglesTool.h
--- a/Engine/Plugins/Experimental/MeshModelingToolsetExp/Source/MeshModelingToolsExp/Public/RemoveOccludedTrianglesTool.h
+++ b/Engine/Plugins/Experimental/MeshModelingToolsetExp/Source/MeshModelingToolsExp/Public/RemoveOccludedTrianglesTool.h
@@ -39,15 +39,20 @@ protected:
 UENUM()
 enum class EOcclusionTriangleSamplingUIMode : uint8
 {
-	// currently do not expose centroid-only option; it almost always looks bad
+	/** Test for occlusion at vertices */
 	Vertices,
+	/** Test for occlusion at vertices and triangle centroids */
 	VerticesAndCentroids
+	//~ currently do not expose centroid-only option; it almost always looks bad
 };
 
 UENUM()
 enum class EOcclusionCalculationUIMode : uint8
 {
-	GeneralizedWindingNumber, // maps to using fast winding number approximation
+	//~ GeneralizedWindingNumber maps to using fast winding number approximation
+	/** Test for occlusion by a 3D 'Winding Number' test (Note: Allows internal 'air pockets' to be considered 'not occluded') */
+	GeneralizedWindingNumber,
+	/** Test for occlusion by casting rays against the mesh */
 	RaycastOcclusionSamples
 };
 
