ModelingComponents: FMeshSceneAdapter improvements. Add ::AddComponents(), allows adding Components directly instead of Actors (currently results in multiple ActorAdapters for a single Actor in some cases). Add ::FastUpdateTransforms() which recomputes the transforms on all objects in the mesh scene, and rebuilds necessary spatial data structures, allowing for (eg) updates during user interaction/etc. Add ::FindNearestRayIntersection() that computes ray/mesh intersections with the mesh scene.
#rnx

diff --git a/Engine/Plugins/Experimental/MeshModelingToolsetExp/Source/ModelingComponents/Private/Scene/MeshSceneAdapter.cpp b/Engine/Plugins/Experimental/MeshModelingToolsetExp/Source/ModelingComponents/Private/Scene/MeshSceneAdapter.cpp
--- a/Engine/Plugins/Experimental/MeshModelingToolsetExp/Source/ModelingComponents/Private/Scene/MeshSceneAdapter.cpp
+++ b/Engine/Plugins/Experimental/MeshModelingToolsetExp/Source/ModelingComponents/Private/Scene/MeshSceneAdapter.cpp
@@ -347,6 +347,37 @@ public:
 		}
 	}
 
+
+	virtual bool RayIntersection(const FRay3d& WorldRay, const FTransformSequence3d& LocalToWorldTransform, FMeshSceneRayHit& WorldHitResultOut) override
+	{
+		FRay3d UseRay = WorldRay;
+		if (bHasBakedTransform == false)
+		{
+			FVector3d LocalOrigin = LocalToWorldTransform.InverseTransformPosition(WorldRay.Origin);
+			FVector3d LocalDir = Normalized(LocalToWorldTransform.InverseTransformPosition(WorldRay.PointAt(1.0)) - LocalOrigin);
+			UseRay = FRay3d(LocalOrigin, LocalDir);
+		}
+
+		double RayHitT; int32 HitTID;
+		if (AABBTree->FindNearestHitTriangle(UseRay, RayHitT, HitTID))
+		{
+			WorldHitResultOut.HitMeshTriIndex = HitTID;
+			WorldHitResultOut.HitMeshSpatialWrapper = this;
+			if (bHasBakedTransform)
+			{
+				WorldHitResultOut.RayDistance = RayHitT;
+			}
+			else
+			{
+				FVector3d WorldPos = LocalToWorldTransform.TransformPosition(UseRay.PointAt(RayHitT));
+				WorldHitResultOut.RayDistance = WorldRay.Project(WorldPos);
+			}
+			return true;
+		}
+		return false;
+	}
+
+
 	virtual void ProcessVerticesInWorld(TFunctionRef<void(const FVector3d&)> ProcessFunc, const FTransformSequence3d& LocalToWorldTransform) override
 	{
 		if (bHasBakedTransform)
@@ -589,6 +620,25 @@ public:
 		return (SourceMesh) ? FWNTree->FastWindingNumber(LocalToWorldTransform.InverseTransformPosition(P)) : 0.0;
 	}
 
+	virtual bool RayIntersection(const FRay3d& WorldRay, const FTransformSequence3d& LocalToWorldTransform, FMeshSceneRayHit& WorldHitResultOut) override
+	{
+		FVector3d LocalOrigin = LocalToWorldTransform.InverseTransformPosition(WorldRay.Origin);
+		FVector3d LocalDir = Normalized(LocalToWorldTransform.InverseTransformPosition(WorldRay.PointAt(1.0)) - LocalOrigin);
+		FRay3d LocalRay(LocalOrigin, LocalDir);
+
+		double LocalHitT; int32 HitTID;
+		if (AABBTree->FindNearestHitTriangle(LocalRay, LocalHitT, HitTID))
+		{
+			WorldHitResultOut.HitMeshTriIndex = HitTID;
+			WorldHitResultOut.HitMeshSpatialWrapper = this;
+			FVector3d WorldPos = LocalToWorldTransform.TransformPosition(LocalRay.PointAt(LocalHitT));
+			WorldHitResultOut.RayDistance = WorldRay.Project(WorldPos);
+			return true;
+		}
+		return false;
+	}
+
+
 	virtual void ProcessVerticesInWorld(TFunctionRef<void(const FVector3d&)> ProcessFunc, const FTransformSequence3d& LocalToWorldTransform) override
 	{
 		int32 NumVertices = (SourceMesh) ? Adapter->VertexCount() : 0;
@@ -659,7 +709,7 @@ static TUniquePtr<IMeshSpatialWrapper> SpatialWrapperFactory( const FMeshTypeCon
 
 
 
-static void CollectActorChildMeshes(AActor* Actor, UActorComponent* Component, FActorAdapter& Adapter)
+static void CollectComponentChildMeshes(UActorComponent* Component, FActorAdapter& Adapter)
 {
 	UStaticMeshComponent* StaticMeshComponent = Cast<UStaticMeshComponent>(Component);
 	if (StaticMeshComponent != nullptr)
@@ -723,6 +773,7 @@ void FMeshSceneAdapter::AddActors(const TArray<AActor*>& ActorsSetIn)
 	// build an FActorAdapter for each Actor, that contains all mesh Components we know 
 	// how to process, including those contained in ChildActorComponents
 	TArray<AActor*> ChildActors;
+	TArray<FActorAdapter*> NewActorAdapters;
 	for (AActor* Actor : ActorsSetIn)
 	{
 		TUniquePtr<FActorAdapter> Adapter = MakeUnique<FActorAdapter>();
@@ -730,7 +781,7 @@ void FMeshSceneAdapter::AddActors(const TArray<AActor*>& ActorsSetIn)
 
 		for (UActorComponent* Component : Actor->GetComponents())
 		{
-			CollectActorChildMeshes(Actor, Component, *Adapter);
+			CollectComponentChildMeshes(Component, *Adapter);
 		}
 
 		ChildActors.Reset();
@@ -739,17 +790,46 @@ void FMeshSceneAdapter::AddActors(const TArray<AActor*>& ActorsSetIn)
 		{
 			for (UActorComponent* Component : ChildActor->GetComponents())
 			{
-				CollectActorChildMeshes(ChildActor, Component, *Adapter);
+				CollectComponentChildMeshes(Component, *Adapter);
 			}
 		}
 
+		NewActorAdapters.Add(Adapter.Get());
 		SceneActors.Add(MoveTemp(Adapter));
 	}
 
+	InitializeSpatialWrappers(NewActorAdapters);
+}
+
 
+
+void FMeshSceneAdapter::AddComponents(const TArray<UActorComponent*>& ComponentSetIn)
+{
+	// Build an FActorAdapter for each Component. This may result in multiple FActorAdapters
+	// for a single AActor, this is currently not a problem
+	TArray<FActorAdapter*> NewActorAdapters;
+	for (UActorComponent* Component : ComponentSetIn)
+	{
+		TUniquePtr<FActorAdapter> Adapter = MakeUnique<FActorAdapter>();
+		Adapter->SourceActor = Component->GetOwner();
+		CollectComponentChildMeshes(Component, *Adapter);
+		if (Adapter->ChildMeshes.Num() > 0)
+		{
+			NewActorAdapters.Add(Adapter.Get());
+			SceneActors.Add(MoveTemp(Adapter));
+		}
+	}
+
+	InitializeSpatialWrappers(NewActorAdapters);
+}
+
+
+
+void FMeshSceneAdapter::InitializeSpatialWrappers(const TArray<FActorAdapter*>& NewItemsToProcess)
+{
 	// Find IMeshSpatialWrapper for each child mesh component. If one does not exist
 	// and we have not seen the underlying unique mesh (eg StaticMesh Asset, etc, construct a new one
-	for (TUniquePtr<FActorAdapter>& Actor : SceneActors)
+	for (FActorAdapter* Actor : NewItemsToProcess)
 	{
 		for (TUniquePtr<FActorChildMesh>& ChildMesh : Actor->ChildMeshes)
 		{
@@ -778,13 +858,9 @@ void FMeshSceneAdapter::AddActors(const TArray<AActor*>& ActorsSetIn)
 			ChildMesh->MeshSpatial = MeshInfo->SpatialWrapper.Get();
 		}
 	}
-
 }
 
 
-
-
-
 void FMeshSceneAdapter::Build(const FMeshSceneAdapterBuildOptions& BuildOptions)
 {
 	TRACE_CPUPROFILER_EVENT_SCOPE(MeshScene_Build);
@@ -1551,6 +1627,140 @@ void FMeshSceneAdapter::BuildSpatialEvaluationCache()
 
 
 
+bool FMeshSceneAdapter::FindNearestRayIntersection(const FRay3d& WorldRay, FMeshSceneRayHit& HitResultOut)
+{
+	check(bHaveSpatialEvaluationCache);		// must call BuildSpatialEvaluationCache() to build Octree
+
+	int32 HitObjectID = Octree->FindNearestHitObject(WorldRay,
+		[&](int32 idx) { return SortedSpatials[idx].Bounds; },
+		[&](int idx, const FRay3d& WorldRay) {
+			FMeshSceneRayHit LocalHitResult;
+			FSpatialCacheInfo& CacheInfo = SortedSpatials[idx];
+			if (CacheInfo.Spatial->RayIntersection(WorldRay, CacheInfo.ChildMesh->WorldTransform, LocalHitResult))
+			{
+				return LocalHitResult.RayDistance;
+			}
+			return TNumericLimits<double>::Max();
+		});
+
+	if (HitObjectID >= 0)
+	{
+		FSpatialCacheInfo& CacheInfo = SortedSpatials[HitObjectID];
+		bool bHit = CacheInfo.Spatial->RayIntersection(WorldRay, CacheInfo.ChildMesh->WorldTransform, HitResultOut);
+		if (ensure(bHit))
+		{
+			HitResultOut.HitActor = CacheInfo.Actor->SourceActor;
+			HitResultOut.HitComponent = CacheInfo.ChildMesh->SourceComponent;
+			HitResultOut.HitComponentElementIndex = CacheInfo.ChildMesh->ComponentIndex;
+			HitResultOut.Ray = WorldRay;
+			HitResultOut.LocalToWorld = CacheInfo.ChildMesh->WorldTransform;
+			return true;
+		}
+	}
+
+	return false;
+}
+
+
+
+void FMeshSceneAdapter::FastUpdateTransforms(bool bRebuildSpatialCache)
+{
+	for (TUniquePtr<FActorAdapter>& Actor : SceneActors)
+	{
+		for (TUniquePtr<FActorChildMesh>& ChildMesh : Actor->ChildMeshes)
+		{
+			ChildMesh->WorldTransform = FTransformSequence3d();
+
+			if ((ChildMesh->ComponentType == EActorMeshComponentType::InstancedStaticMesh) ||
+				(ChildMesh->ComponentType == EActorMeshComponentType::HierarchicalInstancedStaticMesh))
+			{
+				UInstancedStaticMeshComponent* ISMComponent = Cast<UInstancedStaticMeshComponent>(ChildMesh->SourceComponent);
+				if (ISMComponent->IsValidInstance(ChildMesh->ComponentIndex))
+				{
+					FTransform InstanceTransform;
+					ISMComponent->GetInstanceTransform(ChildMesh->ComponentIndex, InstanceTransform, true);
+					ChildMesh->WorldTransform.Append(InstanceTransform);
+					ChildMesh->bIsNonUniformScaled = ChildMesh->WorldTransform.HasNonUniformScale();
+				}
+			}
+			else if (ChildMesh->ComponentType == EActorMeshComponentType::StaticMesh)
+			{
+				UStaticMeshComponent* StaticMeshComponent = Cast<UStaticMeshComponent>(ChildMesh->SourceComponent);
+				ChildMesh->WorldTransform.Append(StaticMeshComponent->GetComponentTransform());
+				ChildMesh->bIsNonUniformScaled = ChildMesh->WorldTransform.HasNonUniformScale();
+			}
+		}
+	}
+
+	// this cache is invalid now
+	bHaveSpatialEvaluationCache = false;
+
+	// update bounding boxes for actors
+	{
+		TRACE_CPUPROFILER_EVENT_SCOPE(MeshScene_FastUpdateTransforms_ActorBounds);
+		ParallelFor(SceneActors.Num(), [&](int32 i)
+		{
+			UpdateActorBounds(*SceneActors[i]);
+		});
+	}
+
+	if (bRebuildSpatialCache)
+	{
+		BuildSpatialEvaluationCache();
+	}
+
+}
+
+
+
+void FMeshSceneAdapter::GetMeshBoundingBoxes(TArray<FAxisAlignedBox3d>& Bounds)
+{
+	if (bHaveSpatialEvaluationCache)
+	{
+		for (FSpatialCacheInfo& CacheInfo : SortedSpatials)
+		{
+			Bounds.Add(CacheInfo.Bounds);
+		}
+	}
+	else
+	{
+		for (const TUniquePtr<FActorAdapter>& Actor : SceneActors)
+		{
+			for (const TUniquePtr<FActorChildMesh>& ChildMesh : Actor->ChildMeshes)
+			{
+				if (ChildMesh->MeshSpatial != nullptr)
+				{
+					FAxisAlignedBox3d ChildBounds = ChildMesh->MeshSpatial->GetWorldBounds(
+						[&](const FVector3d& P) { return ChildMesh->WorldTransform.TransformPosition(P); });
+					Bounds.Add(ChildBounds);
+				}
+			}
+		}
+	}
+}
+
+
+FAxisAlignedBox3d FMeshSceneAdapter::GetMeshBoundingBox(UActorComponent* Component, int32 ComponentIndex)
+{
+	// implementation below depends on the spatial cache
+	check(bHaveSpatialEvaluationCache);
+
+	for (FSpatialCacheInfo& CacheInfo : SortedSpatials)
+	{
+		if (CacheInfo.ChildMesh->SourceComponent == Component)
+		{
+			if (ComponentIndex == -1 || ComponentIndex == CacheInfo.ChildMesh->ComponentIndex )
+			{
+				return CacheInfo.Bounds;
+			}
+		}
+	}
+
+	return FAxisAlignedBox3d(FVector3d::Zero(), 1.0);
+}
+
+
+
 void FMeshSceneAdapter::GetAccumulatedMesh(FDynamicMesh3& AccumMesh)
 {
 	for (const TUniquePtr<FActorAdapter>& Actor : SceneActors)
diff --git a/Engine/Plugins/Experimental/MeshModelingToolsetExp/Source/ModelingComponents/Public/Scene/MeshSceneAdapter.h b/Engine/Plugins/Experimental/MeshModelingToolsetExp/Source/ModelingComponents/Public/Scene/MeshSceneAdapter.h
--- a/Engine/Plugins/Experimental/MeshModelingToolsetExp/Source/ModelingComponents/Public/Scene/MeshSceneAdapter.h
+++ b/Engine/Plugins/Experimental/MeshModelingToolsetExp/Source/ModelingComponents/Public/Scene/MeshSceneAdapter.h
@@ -81,8 +81,34 @@ struct FMeshSceneAdapterBuildOptions
 };
 
 
+class IMeshSpatialWrapper;
 
 
+/**
+ * FMeshSceneRayHit is returned by various ray-intersection functions below
+ */
+struct MODELINGCOMPONENTS_API FMeshSceneRayHit
+{
+	// world ray that was cast, stored here for convenience
+	FRay3d Ray;
+	// distance along ray that intersection occurred at
+	double RayDistance = -1.0;
+	// Actor that was hit, if available
+	AActor* HitActor = nullptr;
+	// Component that was hit, if available
+	UActorComponent* HitComponent = nullptr;
+	// Element Index on hit Component, if available (eg Instance Index on an InstancedStaticMesh)
+	int32 HitComponentElementIndex = -1;
+	// Triangle Index/ID on mesh of hit Component
+	int32 HitMeshTriIndex = -1;
+
+	// SpatialWrapper for the mesh geometry that was hit, if available. This is a pointer to data owned by the FMeshSceneAdapter that was queried
+	IMeshSpatialWrapper* HitMeshSpatialWrapper = nullptr;
+
+	// LocalToWorld Transform on the IMeshSpatialWrapper (for convenience)
+	FTransformSequence3d LocalToWorld;
+};
+
 
 /**
  * Abstract interface to a spatial data structure for a mesh
@@ -106,6 +132,9 @@ public:
 	/** Calculate the mesh winding number at the given Position. Must be callable in parallel from any thread.  */
 	virtual double FastWindingNumber(const FVector3d& P, const FTransformSequence3d& LocalToWorldTransform) = 0;
 
+	/** Find the nearest ray-intersection with the mesh. Must be callable in parallel from any thread. */
+	virtual bool RayIntersection(const FRay3d& WorldRay, const FTransformSequence3d& LocalToWorldTransform, FMeshSceneRayHit& WorldHitResultOut) = 0;
+
 	/** Collect a set of seed points from this Mesh, mapped through LocalToWorldFunc to world space. Must be callable in parallel from any thread. */
 	virtual void CollectSeedPoints(TArray<FVector3d>& WorldPoints, TFunctionRef<FVector3d(const FVector3d&)> LocalToWorldFunc) = 0;
 
@@ -184,15 +213,15 @@ public:
 };
 
 
+
 /**
  * FMeshSceneAdapter creates an internal representation of an Actor/Component/Asset hierarchy,
  * so that a minimal set of Mesh data structures can be constructed for the unique Meshes (generally Assets).
  * This allows queries against the Actor set to be computed without requiring mesh copies or 
  * duplicates of the mesh data structures (ie, saving memory, at the cost of some computation overhead).
  * 
- * Currenly this builds an AABBTree and FastWindingTree for each unique Mesh,
- * and only FastWinding and "Seed Point" queries are (currently) exposed.
- * (SeedPoint query is necessary for mesh the scalar Winding field)
+ * Currently this builds an AABBTree and FastWindingTree for each unique Mesh, and
+ * supports various queries based on those data structures.
  */
 class MODELINGCOMPONENTS_API FMeshSceneAdapter
 {
@@ -203,11 +232,20 @@ public:
 	FMeshSceneAdapter(const FMeshSceneAdapter&) = delete;		// must delete this due to TArray<TUniquePtr> member
 	FMeshSceneAdapter(FMeshSceneAdapter&&) = delete;
 
-	/** Add the given Actors to our Actor set. */
+	//
+	// Mesh Scene Setup etc
+	//
+
+	/** 
+	 * Add the given Actors to the Mesh Scene 
+	 */
 	void AddActors(const TArray<AActor*>& ActorsSetIn);
 
-	/** Build */
-	void Build(const FMeshSceneAdapterBuildOptions& BuildOptions);
+	/** 
+	 * Add the given Components to the Mesh Scene.
+	 * If multiple Components from the same Actor are passed in, duplicate FActorAdapters will be created
+	 */
+	void AddComponents(const TArray<UActorComponent*>& ComponentSetIn);
 
 	/**
 	 * Generate a new mesh that "caps" the mesh scene on the bottom. This can be used in cases where
@@ -220,7 +258,33 @@ public:
 	void GenerateBaseClosingMesh(double BaseHeight = 1.0, double ExtrudeHeight = 0.0);
 
 	/**
-	 * Statistics about the mesh scene returned by GetGeometryStatistics()
+	 * Build the Mesh Scene. This must be called after the functions above, and before any queries below can be run.
+	 */
+	void Build(const FMeshSceneAdapterBuildOptions& BuildOptions);
+
+	/** 
+	 * Precompute data structures that accelerate spatial evaluation queries. Precondition for calling FastWindingNumber() and FindNearestRayIntersection().
+	 * This can be called after Build() has completed
+	 */
+	virtual void BuildSpatialEvaluationCache();
+
+	/**
+	 * Update the transforms on all the Components in the current Mesh Scene. Requires that Build() has been called.
+	 * This will invalidate any existing SpatialEvaluationCache
+	 * @param bRebuildSpatialCache if true, rebuild the SpatialEvaluationCache
+	 */
+	void FastUpdateTransforms(bool bRebuildSpatialCache);
+
+
+	//
+	// Mesh Scene Queries 	  
+	//
+
+	/** @return bounding box for the Actor set */
+	virtual FAxisAlignedBox3d GetBoundingBox();
+
+	/**
+	 * Statistics about the Mesh Scene returned by GetGeometryStatistics()
 	 */
 	struct FStatistics
 	{
@@ -234,21 +298,33 @@ public:
 	/** @return bounding box for the Actor set */
 	virtual void GetGeometryStatistics(FStatistics& StatsOut);
 
-	/** @return bounding box for the Actor set */
-	virtual FAxisAlignedBox3d GetBoundingBox();
-
 	/** @return a set of points on the surface of the meshes, can be used to initialize the MarchingCubes mesher */
 	virtual void CollectMeshSeedPoints(TArray<FVector3d>& PointsOut);
 
-	/** Precompute data structures that accelerate spatial evaluation queries. Precondition for calling FastWindingNumber() */
-	virtual void BuildSpatialEvaluationCache();
+	/** 
+	 * @return FastWindingNumber computed at WorldPoint across all mesh Actors/Components 
+	 * @param bFastEarlyOutIfPossible if true, then if any Mesh Scene Element returns a Winding Number of 1, the function returns immediately
+	 */
+	virtual double FastWindingNumber(const FVector3d& WorldPoint, bool bFastEarlyOutIfPossible = true);
 
-	/** @return FastWindingNumber computed across all mesh Actors/Components */
-	virtual double FastWindingNumber(const FVector3d& P, bool bFastEarlyOutIfPossible = true);
+	/** 
+	 * Intersect the given Ray with the MeshScene and find the nearest mesh hit
+	 * @param HitResultOut ray-intersection information will be returned here
+	 * @return true if hit was found and HitResultOut is initialized
+	 */
+	virtual bool FindNearestRayIntersection(const FRay3d& WorldRay, FMeshSceneRayHit& HitResultOut);
 
-	/** Append all instance triangles to a single mesh. May be very large. */
+	/** Append all instance triangles to a single mesh. May be **very** large. */
 	virtual void GetAccumulatedMesh(FDynamicMesh3& AccumMesh);
 
+	/** Get all world-space bounding boxes for all Scene Meshes */
+	virtual void GetMeshBoundingBoxes(TArray<FAxisAlignedBox3d>& Bounds);
+
+	/** @return world-space bounding box for the Scene Mesh of the specified Component and ComponentIndex. Requires that SpatialEvaluationCache is available. */
+	virtual FAxisAlignedBox3d GetMeshBoundingBox(UActorComponent* Component, int32 ComponentIndex = -1);
+
+
+
 protected:
 	// top-level list of ActorAdapters, which represent each Actor and set of Components
 	TArray<TUniquePtr<FActorAdapter>> SceneActors;
@@ -264,6 +340,8 @@ protected:
 	// Unique set of spatial data structure query interfaces, one for each Mesh object, which is identified by void* pointer
 	TMap<void*, TSharedPtr<FSpatialWrapperInfo>> SpatialAdapters;
 
+	void InitializeSpatialWrappers(const TArray<FActorAdapter*>& NewItemsToProcess);
+
 	bool bSceneIsAllSolids = false;
 
 	void UpdateActorBounds(FActorAdapter& Actor);
