MeshModelinTools: Tweak BSP conversion tool to make the mesh boolean path prioritize later brushes in cases where faces are coplanar. Add simplification for the boolean operations.

#rnx

diff --git a/Engine/Plugins/Experimental/MeshModelingToolsetExp/Source/MeshModelingToolsEditorOnly/Private/BspConversionTool.cpp b/Engine/Plugins/Experimental/MeshModelingToolsetExp/Source/MeshModelingToolsEditorOnly/Private/BspConversionTool.cpp
--- a/Engine/Plugins/Experimental/MeshModelingToolsetExp/Source/MeshModelingToolsEditorOnly/Private/BspConversionTool.cpp
+++ b/Engine/Plugins/Experimental/MeshModelingToolsetExp/Source/MeshModelingToolsEditorOnly/Private/BspConversionTool.cpp
@@ -511,16 +511,26 @@ bool UBspConversionTool::ConvertThenCombine(FText *ErrorMessage)
 		}
 
 		// Apply the boolean operation
-		ApplyStaticMeshBooleanOperation(
-			*InputMesh, *InputTransform, *InputMaterials,
-			NextMesh, NextTransform, NextMaterials,
-			*OutputMesh, *OutputTransform, *OutputMaterials, Operation);
-		
 		if (IsSubtractiveBrush)
 		{
+			ApplyStaticMeshBooleanOperation(
+				*InputMesh, *InputTransform, *InputMaterials,
+				NextMesh, NextTransform, NextMaterials,
+				*OutputMesh, *OutputTransform, *OutputMaterials, Operation);
+		
 			// Undo the change in brush type that we did before.
 			NextBrush->BrushType = EBrushType::Brush_Subtract;
 		}
+		else
+		{
+			// For union, we actually swap the order of meshes in hopes of better lining up
+			// with BSP brush priority in coplanar places (brushes added later have priority,
+			// whereas for our boolean operations, first mesh has priority)
+			ApplyStaticMeshBooleanOperation(
+				NextMesh, NextTransform, NextMaterials,
+				*InputMesh, *InputTransform, *InputMaterials,
+				*OutputMesh, *OutputTransform, *OutputMaterials, Operation);
+		}
 	}//end converting other brushes
 
 	// If there was only a single additive brush, we should keep its pivot in the original location.
@@ -632,6 +642,7 @@ void ApplyStaticMeshBooleanOperation(
 
 	// Perform the actual boolean operation.
 	FMeshBoolean BooleanOperation(&MeshA, TransformA, &MeshB, TransformB, &OutputMesh, Operation);
+	BooleanOperation.bSimplifyAlongNewEdges = true;
 	BooleanOperation.Compute();
 	OutputTransform = BooleanOperation.ResultTransform;
 }
