improve language of polycut tool, and remove fix-open-boundary option for trim operations (where open bouandaries are intended)

#rnx

diff --git a/Engine/Plugins/Experimental/MeshModelingToolsetExp/Source/MeshModelingToolsEditorOnly/Private/PolygonOnMeshTool.cpp b/Engine/Plugins/Experimental/MeshModelingToolsetExp/Source/MeshModelingToolsEditorOnly/Private/PolygonOnMeshTool.cpp
--- a/Engine/Plugins/Experimental/MeshModelingToolsetExp/Source/MeshModelingToolsEditorOnly/Private/PolygonOnMeshTool.cpp
+++ b/Engine/Plugins/Experimental/MeshModelingToolsetExp/Source/MeshModelingToolsEditorOnly/Private/PolygonOnMeshTool.cpp
@@ -256,7 +256,10 @@ TUniquePtr<FDynamicMeshOperator> UPolygonOnMeshTool::MakeNewOperator()
 	EmbedOp->bDiscardAttributes = false;
 	EmbedOp->Operation = BasicProperties->Operation;
 	EmbedOp->bCutWithBoolean = BasicProperties->bCutWithBoolean;
-	EmbedOp->bAttemptFixHolesOnBoolean = BasicProperties->bAttemptFixHoles;
+	bool bOpLeavesOpenBoundaries =
+		BasicProperties->Operation == EEmbeddedPolygonOpMethod::TrimInside ||
+		BasicProperties->Operation == EEmbeddedPolygonOpMethod::TrimOutside;
+	EmbedOp->bAttemptFixHolesOnBoolean = !bOpLeavesOpenBoundaries && BasicProperties->bTryToFixCracks;
 
 	FFrame3d LocalFrame = DrawPlaneWorld;
 	UE::Geometry::FTransform3d ToLocal = WorldTransform.Inverse();
diff --git a/Engine/Plugins/Experimental/MeshModelingToolsetExp/Source/MeshModelingToolsEditorOnly/Public/PolygonOnMeshTool.h b/Engine/Plugins/Experimental/MeshModelingToolsetExp/Source/MeshModelingToolsEditorOnly/Public/PolygonOnMeshTool.h
--- a/Engine/Plugins/Experimental/MeshModelingToolsetExp/Source/MeshModelingToolsEditorOnly/Public/PolygonOnMeshTool.h
+++ b/Engine/Plugins/Experimental/MeshModelingToolsetExp/Source/MeshModelingToolsEditorOnly/Public/PolygonOnMeshTool.h
@@ -76,9 +76,9 @@ public:
 	UPROPERTY(EditAnywhere, Category = Operation)
 	bool bCutWithBoolean = true;
 
-	/** Automatically attempt to fill any holes left by CSG (e.g. due to numerical errors) */
-	UPROPERTY(EditAnywhere, Category = Operation, meta = (EditCondition = "bCutWithBoolean"))
-	bool bAttemptFixHoles = true;
+	/** Automatically attempt to fill any open boundaries left by CSG (e.g. due to numerical errors) */
+	UPROPERTY(EditAnywhere, Category = Operation, meta = (EditCondition = "bCutWithBoolean && Operation == EEmbeddedPolygonOpMethod::CutThrough || bCutWithBoolean && Operation == EEmbeddedPolygonOpMethod::InsertPolygon", EditConditionHides))
+	bool bTryToFixCracks = true;
 
 	// TODO: re-add if/when extrude is added as a supported operation
 	///** Amount to extrude, if extrude is enabled */
diff --git a/Engine/Plugins/Experimental/MeshModelingToolsetExp/Source/ModelingOperatorsEditorOnly/Private/CuttingOps/EmbedPolygonsOp.cpp b/Engine/Plugins/Experimental/MeshModelingToolsetExp/Source/ModelingOperatorsEditorOnly/Private/CuttingOps/EmbedPolygonsOp.cpp
--- a/Engine/Plugins/Experimental/MeshModelingToolsetExp/Source/ModelingOperatorsEditorOnly/Private/CuttingOps/EmbedPolygonsOp.cpp
+++ b/Engine/Plugins/Experimental/MeshModelingToolsetExp/Source/ModelingOperatorsEditorOnly/Private/CuttingOps/EmbedPolygonsOp.cpp
@@ -176,6 +176,7 @@ void FEmbedPolygonsOp::BooleanPath(FProgressCancel* Progress)
 	int MaxGroupID = ResultMesh->MaxGroupID();
 	FMeshBoolean Boolean(ResultMesh.Get(), &ExtrudePoly, ResultMesh.Get(), BoolOp);
 	Boolean.Progress = Progress;
+	Boolean.bSimplifyAlongNewEdges = true;
 	Boolean.bPutResultInInputSpace = true;
 	bool bBoolSuccess = Boolean.Compute();
 
