MeshModelingToolset: Fix tooltip for bAllowSharpVertices

#rnx

diff --git a/Engine/Plugins/Experimental/MeshModelingToolsetExp/Source/MeshModelingToolsExp/Public/EditNormalsTool.h b/Engine/Plugins/Experimental/MeshModelingToolsetExp/Source/MeshModelingToolsExp/Public/EditNormalsTool.h
--- a/Engine/Plugins/Experimental/MeshModelingToolsetExp/Source/MeshModelingToolsExp/Public/EditNormalsTool.h
+++ b/Engine/Plugins/Experimental/MeshModelingToolsetExp/Source/MeshModelingToolsExp/Public/EditNormalsTool.h
@@ -77,7 +77,7 @@ public:
 	UPROPERTY(EditAnywhere, Category = NormalsTopology, meta = (UIMin = "0.0", UIMax = "180.0", ClampMin = "0.0", ClampMax = "180.0", EditCondition = "SplitNormalMethod == ESplitNormalMethod::FaceNormalThreshold"))
 	float SharpEdgeAngleThreshold;
 
-	/** Assign separate normals at 'sharp' vertices -- for example, at the tip of a cone */
+	/** Assign separate normals at 'sharp' vertices, for example, at the tip of a cone */
 	UPROPERTY(EditAnywhere, Category = NormalsTopology, meta = (EditCondition = "SplitNormalMethod == ESplitNormalMethod::FaceNormalThreshold"))
 	bool bAllowSharpVertices;
 };
