MeshModelingToolset: Fix tooltip for ESplitNormalMethod::PerVertex

diff --git a/Engine/Plugins/Experimental/MeshModelingToolsetExp/Source/ModelingOperators/Public/CleaningOps/EditNormalsOp.h b/Engine/Plugins/Experimental/MeshModelingToolsetExp/Source/ModelingOperators/Public/CleaningOps/EditNormalsOp.h
--- a/Engine/Plugins/Experimental/MeshModelingToolsetExp/Source/ModelingOperators/Public/CleaningOps/EditNormalsOp.h
+++ b/Engine/Plugins/Experimental/MeshModelingToolsetExp/Source/ModelingOperators/Public/CleaningOps/EditNormalsOp.h
@@ -29,7 +29,7 @@ enum class ESplitNormalMethod : uint8
 	FaceGroupID,
 	/** Set each triangle-vertex to have the face normal of that triangle's plane */
 	PerTriangle,
-	/** Set each vertex to have a fully shared normal, ie no split normals  */
+	/** Set each vertex to have a fully shared normal, i.e. no split normals  */
 	PerVertex
 };
 
