clean up ersatz em dashes in modeling tools tooltips

#rnx

diff --git a/Engine/Plugins/Experimental/MeshModelingToolsetExp/Source/MeshModelingToolsExp/Public/CSGMeshesTool.h b/Engine/Plugins/Experimental/MeshModelingToolsetExp/Source/MeshModelingToolsExp/Public/CSGMeshesTool.h
--- a/Engine/Plugins/Experimental/MeshModelingToolsetExp/Source/MeshModelingToolsExp/Public/CSGMeshesTool.h
+++ b/Engine/Plugins/Experimental/MeshModelingToolsetExp/Source/MeshModelingToolsExp/Public/CSGMeshesTool.h
@@ -28,7 +28,7 @@ public:
 	UPROPERTY(EditAnywhere, Category = Operation)
 	ECSGOperation Operation = ECSGOperation::DifferenceAB;
 
-	/** Show boundary edges created by the CSG operation -- often due to numerical error */
+	/** Show boundary edges created by the CSG operation (often due to numerical error) */
 	UPROPERTY(EditAnywhere, Category = Operation)
 	bool bShowNewBoundaryEdges = true;
 
diff --git a/Engine/Plugins/Experimental/MeshModelingToolsetExp/Source/MeshModelingToolsExp/Public/CutMeshWithMeshTool.h b/Engine/Plugins/Experimental/MeshModelingToolsetExp/Source/MeshModelingToolsExp/Public/CutMeshWithMeshTool.h
--- a/Engine/Plugins/Experimental/MeshModelingToolsetExp/Source/MeshModelingToolsExp/Public/CutMeshWithMeshTool.h
+++ b/Engine/Plugins/Experimental/MeshModelingToolsetExp/Source/MeshModelingToolsExp/Public/CutMeshWithMeshTool.h
@@ -23,7 +23,7 @@ class MESHMODELINGTOOLSEXP_API UCutMeshWithMeshToolProperties : public UInteract
 {
 	GENERATED_BODY()
 public:
-	/** Show boundary edges created by the Boolean operations -- often due to numerical error */
+	/** Show boundary edges created by the Boolean operations (often due to numerical error) */
 	UPROPERTY(EditAnywhere, Category = Options)
 	bool bShowNewBoundaryEdges = true;
 
diff --git a/Engine/Plugins/Experimental/MeshModelingToolsetExp/Source/MeshModelingToolsExp/Public/SelfUnionMeshesTool.h b/Engine/Plugins/Experimental/MeshModelingToolsetExp/Source/MeshModelingToolsExp/Public/SelfUnionMeshesTool.h
--- a/Engine/Plugins/Experimental/MeshModelingToolsetExp/Source/MeshModelingToolsExp/Public/SelfUnionMeshesTool.h
+++ b/Engine/Plugins/Experimental/MeshModelingToolsetExp/Source/MeshModelingToolsExp/Public/SelfUnionMeshesTool.h
@@ -31,7 +31,7 @@ public:
 	UPROPERTY(EditAnywhere, Category = Options)
 	bool bCollapseExtraEdges = true;
 
-	/** Show boundary edges created by the union operation -- often due to numerical error */
+	/** Show boundary edges created by the union operation (often due to numerical error) */
 	UPROPERTY(EditAnywhere, Category = Options)
 	bool bShowNewBoundaryEdges = true;
 
