MeshModelingTools: Add ability to put accept/cancel tools into UV editor and pull a version of the UV Layout tool into it.

#rnx

diff --git a/Engine/Plugins/Experimental/MeshModelingToolsetExp/Source/MeshModelingToolsEditorOnly/Private/UVLayoutTool.cpp b/Engine/Plugins/Experimental/MeshModelingToolsetExp/Source/MeshModelingToolsEditorOnly/Private/UVLayoutTool.cpp
deleted file mode 100644
diff --git a/Engine/Plugins/Experimental/MeshModelingToolsetExp/Source/MeshModelingToolsEditorOnly/Public/UVLayoutTool.h b/Engine/Plugins/Experimental/MeshModelingToolsetExp/Source/MeshModelingToolsEditorOnly/Public/UVLayoutTool.h
deleted file mode 100644
diff --git a/Engine/Plugins/Experimental/MeshModelingToolsetExp/Source/MeshModelingToolsExp/Private/UVLayoutTool.cpp b/Engine/Plugins/Experimental/MeshModelingToolsetExp/Source/MeshModelingToolsExp/Private/UVLayoutTool.cpp
new file mode 100644
--- /dev/null
+++ b/Engine/Plugins/Experimental/MeshModelingToolsetExp/Source/MeshModelingToolsExp/Private/UVLayoutTool.cpp
@@ -0,0 +1,346 @@
+// Copyright Epic Games, Inc. All Rights Reserved.
+
+#include "UVLayoutTool.h"
+#include "InteractiveToolManager.h"
+#include "ToolBuilderUtil.h"
+
+#include "ToolSetupUtil.h"
+
+#include "DynamicMesh/DynamicMesh3.h"
+#include "BaseBehaviors/MultiClickSequenceInputBehavior.h"
+#include "Selection/SelectClickedAction.h"
+
+#include "MeshDescriptionToDynamicMesh.h"
+#include "DynamicMeshToMeshDescription.h"
+#include "Materials/MaterialInstanceDynamic.h"
+
+#include "ParameterizationOps/UVLayoutOp.h"
+#include "Properties/UVLayoutProperties.h"
+
+#include "TargetInterfaces/MaterialProvider.h"
+#include "TargetInterfaces/MeshDescriptionCommitter.h"
+#include "TargetInterfaces/MeshDescriptionProvider.h"
+#include "TargetInterfaces/PrimitiveComponentBackedTarget.h"
+#include "ToolTargetManager.h"
+
+#include "ExplicitUseGeometryMathTypes.h"		// using UE::Geometry::(math types)
+using namespace UE::Geometry;
+
+#define LOCTEXT_NAMESPACE "UUVLayoutTool"
+
+/*
+ * ToolBuilder
+ */
+
+
+const FToolTargetTypeRequirements& UUVLayoutToolBuilder::GetTargetRequirements() const
+{
+	static FToolTargetTypeRequirements TypeRequirements({
+		UMaterialProvider::StaticClass(),
+		UMeshDescriptionCommitter::StaticClass(),
+		UMeshDescriptionProvider::StaticClass(),
+		UPrimitiveComponentBackedTarget::StaticClass()
+		});
+	return TypeRequirements;
+}
+
+bool UUVLayoutToolBuilder::CanBuildTool(const FToolBuilderState& SceneState) const
+{
+	return SceneState.TargetManager->CountSelectedAndTargetable(SceneState, GetTargetRequirements()) >= 1;
+}
+
+UInteractiveTool* UUVLayoutToolBuilder::BuildTool(const FToolBuilderState& SceneState) const
+{
+	UUVLayoutTool* NewTool = NewObject<UUVLayoutTool>(SceneState.ToolManager);
+
+	TArray<TObjectPtr<UToolTarget>> Targets = SceneState.TargetManager->BuildAllSelectedTargetable(SceneState, GetTargetRequirements());
+	NewTool->SetTargets(MoveTemp(Targets));
+	NewTool->SetWorld(SceneState.World, SceneState.GizmoManager);
+
+	return NewTool;
+}
+
+
+
+/*
+ * Tool
+ */
+
+UUVLayoutTool::UUVLayoutTool()
+{
+}
+
+void UUVLayoutTool::SetWorld(UWorld* World, UInteractiveGizmoManager* GizmoManagerIn)
+{
+	this->TargetWorld = World;
+}
+
+void UUVLayoutTool::Setup()
+{
+	UInteractiveTool::Setup();
+
+	// hide input StaticMeshComponent
+	for (int32 ComponentIdx = 0; ComponentIdx < Targets.Num(); ComponentIdx++)
+	{
+		TargetComponentInterface(ComponentIdx)->SetOwnerVisibility(false);
+	}
+
+	// if we only have one object, add ability to set UV channel
+	if (Targets.Num() == 1)
+	{
+		UVChannelProperties = NewObject<UMeshUVChannelProperties>(this);
+		UVChannelProperties->RestoreProperties(this);
+		UVChannelProperties->Initialize(TargetMeshProviderInterface(0)->GetMeshDescription(), false);
+		UVChannelProperties->ValidateSelection(true);
+		AddToolPropertySource(UVChannelProperties);
+		UVChannelProperties->WatchProperty(UVChannelProperties->UVChannel, [this](const FString& NewValue)
+		{
+			MaterialSettings->UVChannel = UVChannelProperties->GetSelectedChannelIndex(true);
+		});
+	}
+
+	BasicProperties = NewObject<UUVLayoutProperties>(this);
+	BasicProperties->RestoreProperties(this);
+	AddToolPropertySource(BasicProperties);
+
+	MaterialSettings = NewObject<UExistingMeshMaterialProperties>(this);
+	MaterialSettings->RestoreProperties(this);
+	AddToolPropertySource(MaterialSettings);
+
+	// if we only have one object, add optional UV layout view
+	if (Targets.Num() == 1)
+	{
+		UVLayoutView = NewObject<UUVLayoutPreview>(this);
+		UVLayoutView->CreateInWorld(TargetWorld);
+
+		FComponentMaterialSet MaterialSet;
+		TargetMaterialInterface(0)->GetMaterialSet(MaterialSet);
+		UVLayoutView->SetSourceMaterials(MaterialSet);
+
+		UVLayoutView->SetSourceWorldPosition(
+			TargetComponentInterface(0)->GetOwnerActor()->GetTransform(),
+			TargetComponentInterface(0)->GetOwnerActor()->GetComponentsBoundingBox());
+
+		UVLayoutView->Settings->RestoreProperties(this);
+		AddToolPropertySource(UVLayoutView->Settings);
+	}
+
+	UpdateVisualization();
+
+	SetToolDisplayName(LOCTEXT("ToolName", "UV Layout"));
+	GetToolManager()->DisplayMessage(LOCTEXT("OnStartUVLayoutTool", "Transform/Rotate/Scale existing UV Charts using various strategies"),
+		EToolMessageLevel::UserNotification);
+}
+
+
+void UUVLayoutTool::UpdateNumPreviews()
+{
+	int32 CurrentNumPreview = Previews.Num();
+	int32 TargetNumPreview = Targets.Num();
+	if (TargetNumPreview < CurrentNumPreview)
+	{
+		for (int32 PreviewIdx = CurrentNumPreview - 1; PreviewIdx >= TargetNumPreview; PreviewIdx--)
+		{
+			Previews[PreviewIdx]->Cancel();
+		}
+		Previews.SetNum(TargetNumPreview);
+		OriginalDynamicMeshes.SetNum(TargetNumPreview);
+	}
+	else
+	{
+		OriginalDynamicMeshes.SetNum(TargetNumPreview);
+		for (int32 PreviewIdx = CurrentNumPreview; PreviewIdx < TargetNumPreview; PreviewIdx++)
+		{
+			OriginalDynamicMeshes[PreviewIdx] = MakeShared<FDynamicMesh3, ESPMode::ThreadSafe>();
+			FMeshDescriptionToDynamicMesh Converter;
+			Converter.Convert(TargetMeshProviderInterface(PreviewIdx)->GetMeshDescription(), *OriginalDynamicMeshes[PreviewIdx]);
+
+			UUVLayoutOperatorFactory* OpFactory = NewObject<UUVLayoutOperatorFactory>();
+			OpFactory->OriginalMesh = OriginalDynamicMeshes[PreviewIdx];
+			OpFactory->Settings = BasicProperties;
+			OpFactory->TargetTransform = TargetComponentInterface(PreviewIdx)->GetWorldTransform();
+			OpFactory->GetSelectedUVChannel = [this]() { return GetSelectedUVChannel(); };
+
+			UMeshOpPreviewWithBackgroundCompute* Preview = Previews.Add_GetRef(NewObject<UMeshOpPreviewWithBackgroundCompute>(OpFactory, "Preview"));
+			Preview->Setup(this->TargetWorld, OpFactory);
+
+			FComponentMaterialSet MaterialSet;
+			TargetMaterialInterface(PreviewIdx)->GetMaterialSet(MaterialSet);
+			Preview->ConfigureMaterials(MaterialSet.Materials,
+				ToolSetupUtil::GetDefaultWorkingMaterial(GetToolManager())
+			);
+			Preview->PreviewMesh->UpdatePreview(OriginalDynamicMeshes[PreviewIdx].Get());
+			Preview->PreviewMesh->SetTransform(TargetComponentInterface(PreviewIdx)->GetWorldTransform());
+
+			Preview->OnMeshUpdated.AddLambda([this](UMeshOpPreviewWithBackgroundCompute* Compute)
+			{
+				OnPreviewMeshUpdated(Compute);
+			});
+
+			Preview->SetVisibility(true);
+		}
+	}
+}
+
+
+void UUVLayoutTool::Shutdown(EToolShutdownType ShutdownType)
+{
+	if (UVLayoutView)
+	{
+		UVLayoutView->Settings->SaveProperties(this);
+		UVLayoutView->Disconnect();
+	}
+
+	BasicProperties->SaveProperties(this);
+	MaterialSettings->SaveProperties(this);
+
+	// Restore (unhide) the source meshes
+	for (int32 ComponentIdx = 0; ComponentIdx < Targets.Num(); ComponentIdx++)
+	{
+		TargetComponentInterface(ComponentIdx)->SetOwnerVisibility(true);
+	}
+
+	TArray<FDynamicMeshOpResult> Results;
+	for (UMeshOpPreviewWithBackgroundCompute* Preview : Previews)
+	{
+		Results.Emplace(Preview->Shutdown());
+	}
+	if (ShutdownType == EToolShutdownType::Accept)
+	{
+		GenerateAsset(Results);
+	}
+}
+
+int32 UUVLayoutTool::GetSelectedUVChannel() const
+{
+	return UVChannelProperties ? UVChannelProperties->GetSelectedChannelIndex(true) : 0;
+}
+
+
+void UUVLayoutTool::Render(IToolsContextRenderAPI* RenderAPI)
+{
+	GetToolManager()->GetContextQueriesAPI()->GetCurrentViewState(CameraState);
+
+	if (UVLayoutView)
+	{
+		UVLayoutView->Render(RenderAPI);
+	}
+
+}
+
+void UUVLayoutTool::OnTick(float DeltaTime)
+{
+	for (UMeshOpPreviewWithBackgroundCompute* Preview : Previews)
+	{
+		Preview->Tick(DeltaTime);
+	}
+
+	if (UVLayoutView)
+	{
+		UVLayoutView->OnTick(DeltaTime);
+	}
+
+
+}
+
+
+
+void UUVLayoutTool::OnPropertyModified(UObject* PropertySet, FProperty* Property)
+{
+	if (PropertySet == BasicProperties || PropertySet == UVChannelProperties)
+	{
+		UpdateNumPreviews();
+		for (UMeshOpPreviewWithBackgroundCompute* Preview : Previews)
+		{
+			Preview->InvalidateResult();
+		}
+	}
+	else if (PropertySet == MaterialSettings)
+	{
+		// if we don't know what changed, or we know checker density changed, update checker material
+		UpdateVisualization();
+	}
+}
+
+
+void UUVLayoutTool::OnPreviewMeshUpdated(UMeshOpPreviewWithBackgroundCompute* Compute)
+{
+	if (UVLayoutView)
+	{
+		FDynamicMesh3 ResultMesh;
+		if (Compute->GetCurrentResultCopy(ResultMesh, false) == false)
+		{
+			return;
+		}
+		UVLayoutView->UpdateUVMesh(&ResultMesh, GetSelectedUVChannel());
+	}
+
+}
+
+
+void UUVLayoutTool::UpdateVisualization()
+{
+	MaterialSettings->UpdateMaterials();
+	UpdateNumPreviews();
+	for (int PreviewIdx = 0; PreviewIdx < Previews.Num(); PreviewIdx++)
+	{
+		UMeshOpPreviewWithBackgroundCompute* Preview = Previews[PreviewIdx];
+		Preview->OverrideMaterial = MaterialSettings->GetActiveOverrideMaterial();
+	}
+
+	for (UMeshOpPreviewWithBackgroundCompute* Preview : Previews)
+	{
+		Preview->InvalidateResult();
+	}
+}
+
+bool UUVLayoutTool::CanAccept() const
+{
+	for (UMeshOpPreviewWithBackgroundCompute* Preview : Previews)
+	{
+		if (!Preview->HaveValidResult())
+		{
+			return false;
+		}
+	}
+	return Super::CanAccept();
+}
+
+
+void UUVLayoutTool::GenerateAsset(const TArray<FDynamicMeshOpResult>& Results)
+{
+	GetToolManager()->BeginUndoTransaction(LOCTEXT("UVLayoutToolTransactionName", "UV Layout Tool"));
+
+	check(Results.Num() == Targets.Num());
+	
+	for (int32 ComponentIdx = 0; ComponentIdx < Targets.Num(); ComponentIdx++)
+	{
+		check(Results[ComponentIdx].Mesh.Get() != nullptr);
+		TargetMeshCommitterInterface(ComponentIdx)->CommitMeshDescription([&Results, &ComponentIdx, this](const IMeshDescriptionCommitter::FCommitterParams& CommitParams)
+		{
+			FDynamicMesh3* DynamicMesh = Results[ComponentIdx].Mesh.Get();
+			FMeshDescription* MeshDescription = CommitParams.MeshDescriptionOut;
+			
+			bool bVerticesOnly = false;
+			bool bAttributesOnly = true;
+			if (FDynamicMeshToMeshDescription::HaveMatchingElementCounts(DynamicMesh, MeshDescription, bVerticesOnly, bAttributesOnly))
+			{
+				FDynamicMeshToMeshDescription Converter;
+				Converter.UpdateAttributes(DynamicMesh, *MeshDescription, false, false, true/*update uvs*/);
+			}
+			else
+			{
+				// must have been duplicate tris in the mesh description; we can't count on 1-to-1 mapping of TriangleIDs.  Just convert
+				FDynamicMeshToMeshDescription Converter;
+				Converter.Convert(DynamicMesh, *MeshDescription);
+			}
+		});
+	}
+
+	GetToolManager()->EndUndoTransaction();
+}
+
+
+
+
+#undef LOCTEXT_NAMESPACE
diff --git a/Engine/Plugins/Experimental/MeshModelingToolsetExp/Source/MeshModelingToolsExp/Public/UVLayoutTool.h b/Engine/Plugins/Experimental/MeshModelingToolsetExp/Source/MeshModelingToolsExp/Public/UVLayoutTool.h
new file mode 100644
--- /dev/null
+++ b/Engine/Plugins/Experimental/MeshModelingToolsetExp/Source/MeshModelingToolsExp/Public/UVLayoutTool.h
@@ -0,0 +1,108 @@
+// Copyright Epic Games, Inc. All Rights Reserved.
+
+#pragma once
+
+#include "CoreMinimal.h"
+#include "UObject/NoExportTypes.h"
+#include "MultiSelectionTool.h"
+#include "GeometryBase.h"
+#include "InteractiveToolBuilder.h"
+#include "MeshOpPreviewHelpers.h"
+#include "ToolDataVisualizer.h"
+#include "DynamicMesh/DynamicMesh3.h"
+#include "BaseTools/SingleClickTool.h"
+#include "Properties/MeshMaterialProperties.h"
+#include "Properties/MeshUVChannelProperties.h"
+#include "Drawing/UVLayoutPreview.h"
+
+#include "UVLayoutTool.generated.h"
+
+
+// predeclarations
+struct FMeshDescription;
+class UDynamicMeshComponent;
+class UUVLayoutProperties;
+PREDECLARE_GEOMETRY(class FDynamicMesh3);
+
+/**
+ *
+ */
+UCLASS()
+class MESHMODELINGTOOLSEXP_API UUVLayoutToolBuilder : public UInteractiveToolBuilder
+{
+	GENERATED_BODY()
+
+public:
+	virtual bool CanBuildTool(const FToolBuilderState& SceneState) const override;
+	virtual UInteractiveTool* BuildTool(const FToolBuilderState& SceneState) const override;
+
+protected:
+	virtual const FToolTargetTypeRequirements& GetTargetRequirements() const override;
+};
+
+
+/**
+ * The level editor version of the UV layout tool.
+ */
+UCLASS()
+class MESHMODELINGTOOLSEXP_API UUVLayoutTool : public UMultiSelectionTool
+{
+	GENERATED_BODY()
+
+public:
+
+	UUVLayoutTool();
+
+	virtual void Setup() override;
+	virtual void Shutdown(EToolShutdownType ShutdownType) override;
+
+	virtual void SetWorld(UWorld* World, UInteractiveGizmoManager* GizmoManagerIn);
+
+	virtual void OnTick(float DeltaTime) override;
+	virtual void Render(IToolsContextRenderAPI* RenderAPI) override;
+
+	virtual bool HasCancel() const override { return true; }
+	virtual bool HasAccept() const override { return true; }
+	virtual bool CanAccept() const override;
+
+	virtual void OnPropertyModified(UObject* PropertySet, FProperty* Property) override;
+
+
+public:
+	int32 GetSelectedUVChannel() const;
+
+
+protected:
+
+	UPROPERTY()
+	TObjectPtr<UMeshUVChannelProperties> UVChannelProperties = nullptr;
+
+	UPROPERTY()
+	TObjectPtr<UUVLayoutProperties> BasicProperties = nullptr;
+
+	UPROPERTY()
+	TObjectPtr<UExistingMeshMaterialProperties> MaterialSettings = nullptr;
+
+	UPROPERTY()
+	TArray<TObjectPtr<UMeshOpPreviewWithBackgroundCompute>> Previews;
+
+protected:
+	TArray<TSharedPtr<UE::Geometry::FDynamicMesh3, ESPMode::ThreadSafe>> OriginalDynamicMeshes;
+
+	UWorld* TargetWorld = nullptr;
+
+	FViewCameraState CameraState;
+
+	void UpdateNumPreviews();
+
+	void UpdateVisualization();
+
+	void OnPreviewMeshUpdated(UMeshOpPreviewWithBackgroundCompute* Compute);
+
+	void GenerateAsset(const TArray<FDynamicMeshOpResult>& Results);
+
+
+protected:
+	UPROPERTY()
+	TObjectPtr<UUVLayoutPreview> UVLayoutView = nullptr;
+};
diff --git a/Engine/Plugins/Experimental/MeshModelingToolsetExp/Source/ModelingOperators/Private/ParameterizationOps/UVLayoutOp.cpp b/Engine/Plugins/Experimental/MeshModelingToolsetExp/Source/ModelingOperators/Private/ParameterizationOps/UVLayoutOp.cpp
new file mode 100644
--- /dev/null
+++ b/Engine/Plugins/Experimental/MeshModelingToolsetExp/Source/ModelingOperators/Private/ParameterizationOps/UVLayoutOp.cpp
@@ -0,0 +1,122 @@
+// Copyright Epic Games, Inc. All Rights Reserved.
+
+#include "ParameterizationOps/UVLayoutOp.h"
+
+#include "DynamicMesh/DynamicMesh3.h"
+#include "DynamicMesh/DynamicMeshAttributeSet.h"
+#include "Parameterization/MeshUVPacking.h"
+#include "Properties/UVLayoutProperties.h"
+#include "Selections/MeshConnectedComponents.h"
+
+
+#include "ExplicitUseGeometryMathTypes.h"		// using UE::Geometry::(math types)
+using namespace UE::Geometry;
+
+void FUVLayoutOp::SetTransform(const FTransform& Transform) {
+	ResultTransform = (FTransform3d)Transform;
+}
+
+
+void FUVLayoutOp::CalculateResult(FProgressCancel* Progress)
+{
+	if (Progress && Progress->Cancelled())
+	{
+		return;
+	}
+	ResultMesh->Copy(*OriginalMesh, true, true, true, true);
+
+	if (!ensureMsgf(ResultMesh->HasAttributes(), TEXT("Attributes not found on mesh? Conversion should always create them, so this operator should not need to do so.")))
+	{
+		ResultMesh->EnableAttributes();
+	}
+
+	if (Progress && Progress->Cancelled())
+	{
+		return;
+	}
+
+
+	int UVLayerInput = UVLayerIndex, UVLayerOutput = UVLayerIndex;
+	FDynamicMeshUVOverlay* UseUVLayer = ResultMesh->Attributes()->GetUVLayer(UVLayerInput);
+	
+
+	bool bWillRepackIslands = (UVLayoutMode != EUVLayoutOpLayoutModes::TransformOnly);
+
+	// split bowties so that we can process islands independently
+	if (bWillRepackIslands || bAlwaysSplitBowties)
+	{
+		UseUVLayer->SplitBowties();
+	}
+
+	if (Progress && Progress->Cancelled())
+	{
+		return;
+	}
+
+	FDynamicMeshUVPacker Packer(UseUVLayer);
+	Packer.TextureResolution = this->TextureResolution;
+	Packer.GutterSize = this->GutterSize;
+	Packer.bAllowFlips = this->bAllowFlips;
+
+	if (UVLayoutMode == EUVLayoutOpLayoutModes::RepackToUnitRect)
+	{
+		if (Packer.StandardPack() == false)
+		{
+			// failed... what to do?
+			return;
+		}
+	}
+	else if (UVLayoutMode == EUVLayoutOpLayoutModes::StackInUnitRect)
+	{
+		if (Packer.StackPack() == false)
+		{
+			// failed... what to do?
+			return;
+		}
+	}
+
+	if (Progress && Progress->Cancelled())
+	{
+		return;
+	}
+
+	if (UVScaleFactor != 1.0 || UVTranslation != FVector2f::Zero() )
+	{
+		for (int ElementID : UseUVLayer->ElementIndicesItr())
+		{
+			FVector2f UV = UseUVLayer->GetElement(ElementID);
+			UV = (UV * UVScaleFactor) + UVTranslation;
+			UseUVLayer->SetElement(ElementID, UV);
+		}
+	}
+
+}
+
+TUniquePtr<FDynamicMeshOperator> UUVLayoutOperatorFactory::MakeNewOperator()
+{
+	TUniquePtr<FUVLayoutOp> Op = MakeUnique<FUVLayoutOp>();
+
+	Op->OriginalMesh = OriginalMesh;
+
+	switch (Settings->LayoutType)
+	{
+	case EUVLayoutType::Transform:
+		Op->UVLayoutMode = EUVLayoutOpLayoutModes::TransformOnly;
+		break;
+	case EUVLayoutType::Stack:
+		Op->UVLayoutMode = EUVLayoutOpLayoutModes::StackInUnitRect;
+		break;
+	case EUVLayoutType::Repack:
+		Op->UVLayoutMode = EUVLayoutOpLayoutModes::RepackToUnitRect;
+		break;
+	}
+
+	Op->UVLayerIndex = GetSelectedUVChannel();
+	Op->TextureResolution = Settings->TextureResolution;
+	Op->bAllowFlips = Settings->bAllowFlips;
+	Op->UVScaleFactor = Settings->UVScaleFactor;
+	Op->UVTranslation = FVector2f(Settings->UVTranslate);
+	Op->SetTransform(TargetTransform);
+
+	return Op;
+}
\ No newline at end of file
diff --git a/Engine/Plugins/Experimental/MeshModelingToolsetExp/Source/ModelingOperators/Public/ParameterizationOps/UVLayoutOp.h b/Engine/Plugins/Experimental/MeshModelingToolsetExp/Source/ModelingOperators/Public/ParameterizationOps/UVLayoutOp.h
new file mode 100644
--- /dev/null
+++ b/Engine/Plugins/Experimental/MeshModelingToolsetExp/Source/ModelingOperators/Public/ParameterizationOps/UVLayoutOp.h
@@ -0,0 +1,75 @@
+// Copyright Epic Games, Inc. All Rights Reserved.
+
+#pragma once
+
+#include "CoreMinimal.h"
+#include "Util/ProgressCancel.h"
+#include "ModelingOperators.h"
+#include "GeometryBase.h"
+
+#include "UVLayoutOp.generated.h"
+
+class UUVLayoutProperties;
+
+namespace UE
+{
+namespace Geometry
+{
+	class FDynamicMesh3;
+
+enum class EUVLayoutOpLayoutModes
+{
+	TransformOnly = 0,
+	RepackToUnitRect = 1,
+	StackInUnitRect = 2
+};
+
+class MODELINGOPERATORS_API FUVLayoutOp : public FDynamicMeshOperator
+{
+public:
+	virtual ~FUVLayoutOp() {}
+
+	// inputs
+	TSharedPtr<FDynamicMesh3, ESPMode::ThreadSafe> OriginalMesh;
+
+	EUVLayoutOpLayoutModes UVLayoutMode = EUVLayoutOpLayoutModes::RepackToUnitRect;
+
+	int UVLayerIndex = 0;
+	int TextureResolution = 128;
+	bool bAllowFlips = false;
+	bool bAlwaysSplitBowties = true;
+	float UVScaleFactor = 1.0;
+	float GutterSize = 1.0;
+	FVector2f UVTranslation = FVector2f::Zero();
+
+	void SetTransform(const FTransform& Transform);
+
+	//
+	// FDynamicMeshOperator implementation
+	// 
+
+	virtual void CalculateResult(FProgressCancel* Progress) override;
+};
+
+} // end namespace UE::Geometry
+} // end namespace UE
+
+/**
+ * Can be hooked up to a UMeshOpPreviewWithBackgroundCompute to perform UV Layout operations.
+ */
+UCLASS()
+class MODELINGOPERATORS_API UUVLayoutOperatorFactory : public UObject, public UE::Geometry::IDynamicMeshOperatorFactory
+{
+	GENERATED_BODY()
+
+public:
+	// IDynamicMeshOperatorFactory API
+	virtual TUniquePtr<UE::Geometry::FDynamicMeshOperator> MakeNewOperator() override;
+
+	UPROPERTY()
+	TObjectPtr<UUVLayoutProperties> Settings;
+
+	TSharedPtr<UE::Geometry::FDynamicMesh3, ESPMode::ThreadSafe> OriginalMesh;
+	TUniqueFunction<int32()> GetSelectedUVChannel = []() { return 0; };
+	FTransform TargetTransform;
+};
\ No newline at end of file
diff --git a/Engine/Plugins/Experimental/MeshModelingToolsetExp/Source/ModelingOperators/Public/Properties/UVLayoutProperties.h b/Engine/Plugins/Experimental/MeshModelingToolsetExp/Source/ModelingOperators/Public/Properties/UVLayoutProperties.h
new file mode 100644
--- /dev/null
+++ b/Engine/Plugins/Experimental/MeshModelingToolsetExp/Source/ModelingOperators/Public/Properties/UVLayoutProperties.h
@@ -0,0 +1,49 @@
+// Copyright Epic Games, Inc. All Rights Reserved.
+
+#pragma once
+
+#include "CoreMinimal.h"
+
+#include "InteractiveTool.h"
+
+#include "UVLayoutProperties.generated.h"
+
+UENUM()
+enum class EUVLayoutType
+{
+	Transform,
+	Stack,
+	Repack
+};
+
+
+/**
+ * Standard properties
+ */
+UCLASS()
+class MODELINGOPERATORS_API UUVLayoutProperties : public UInteractiveToolPropertySet
+{
+	GENERATED_BODY()
+
+public:
+	/** Type of transformation to apply to input UV islands */
+	UPROPERTY(EditAnywhere, Category = UVLayout)
+	EUVLayoutType LayoutType = EUVLayoutType::Repack;
+
+	/** Expected resolution of output textures; controls spacing left between charts */
+	UPROPERTY(EditAnywhere, Category = UVLayout, meta = (UIMin = "64", UIMax = "2048", ClampMin = "2", ClampMax = "4096"))
+	int TextureResolution = 1024;
+
+	/** Apply this uniform scaling to the UVs after any layout recalculation */
+	UPROPERTY(EditAnywhere, Category = UVLayout, meta = (UIMin = "0.1", UIMax = "5.0", ClampMin = "0.0001", ClampMax = "10000") )
+	float UVScaleFactor = 1;
+
+	/** Apply this 2D translation to the UVs after any layout recalculation, and after scaling */
+	UPROPERTY(EditAnywhere, Category = UVLayout)
+	FVector2D UVTranslate = FVector2D(0,0);
+
+	/** Allow the packer to flip the orientation of UV islands if it save space. May cause problems for downstream operations, not recommended. */
+	UPROPERTY(EditAnywhere, Category = UVLayout, AdvancedDisplay)
+	bool bAllowFlips = false;
+
+};
\ No newline at end of file
diff --git a/Engine/Plugins/Experimental/MeshModelingToolsetExp/Source/ModelingOperatorsEditorOnly/Private/ParameterizationOps/UVLayoutOp.cpp b/Engine/Plugins/Experimental/MeshModelingToolsetExp/Source/ModelingOperatorsEditorOnly/Private/ParameterizationOps/UVLayoutOp.cpp
deleted file mode 100644
diff --git a/Engine/Plugins/Experimental/MeshModelingToolsetExp/Source/ModelingOperatorsEditorOnly/Public/ParameterizationOps/UVLayoutOp.h b/Engine/Plugins/Experimental/MeshModelingToolsetExp/Source/ModelingOperatorsEditorOnly/Public/ParameterizationOps/UVLayoutOp.h
deleted file mode 100644
diff --git a/Engine/Plugins/Experimental/UVEditor/Source/UVEditor/Private/SUVEditor2DViewport.cpp b/Engine/Plugins/Experimental/UVEditor/Source/UVEditor/Private/SUVEditor2DViewport.cpp
new file mode 100644
--- /dev/null
+++ b/Engine/Plugins/Experimental/UVEditor/Source/UVEditor/Private/SUVEditor2DViewport.cpp
@@ -0,0 +1,20 @@
+// Copyright Epic Games, Inc. All Rights Reserved.
+
+#include "SUVEditor2DViewport.h"
+
+#define LOCTEXT_NAMESPACE "SUVEditor2DViewport"
+
+void SUVEditor2DViewport::AddOverlayWidget(TSharedRef<SWidget> OverlaidWidget)
+{
+	ViewportOverlay->AddSlot()
+	[
+		OverlaidWidget
+	];
+}
+
+void SUVEditor2DViewport::RemoveOverlayWidget(TSharedRef<SWidget> OverlaidWidget)
+{
+	ViewportOverlay->RemoveSlot(OverlaidWidget);
+}
+
+#undef LOCTEXT_NAMESPACE
diff --git a/Engine/Plugins/Experimental/UVEditor/Source/UVEditor/Private/UVEditorCommands.cpp b/Engine/Plugins/Experimental/UVEditor/Source/UVEditor/Private/UVEditorCommands.cpp
--- a/Engine/Plugins/Experimental/UVEditor/Source/UVEditor/Private/UVEditorCommands.cpp
+++ b/Engine/Plugins/Experimental/UVEditor/Source/UVEditor/Private/UVEditorCommands.cpp
@@ -25,6 +25,11 @@ void FUVEditorCommands::RegisterCommands()
 	// These get linked to various tool buttons.
 	UI_COMMAND(BeginSelectTool, "Select", "Switch to select tool", EUserInterfaceActionType::ToggleButton, FInputChord());
 	UI_COMMAND(BeginTransformTool, "Transform", "Switch to transform tool", EUserInterfaceActionType::ToggleButton, FInputChord());
+	UI_COMMAND(BeginLayoutTool, "Layout", "Switch to layout tool", EUserInterfaceActionType::ToggleButton, FInputChord());
+
+	// These allow us to link up to pressed keys
+	UI_COMMAND(AcceptActiveTool, "Accept", "Accept the active tool", EUserInterfaceActionType::Button, FInputChord(EKeys::Enter));
+	UI_COMMAND(CancelActiveTool, "Cancel", "Cancel the active tool", EUserInterfaceActionType::Button, FInputChord(EKeys::Escape));
 }
 
 #undef LOCTEXT_NAMESPACE
diff --git a/Engine/Plugins/Experimental/UVEditor/Source/UVEditor/Private/UVEditorMode.cpp b/Engine/Plugins/Experimental/UVEditor/Source/UVEditor/Private/UVEditorMode.cpp
--- a/Engine/Plugins/Experimental/UVEditor/Source/UVEditor/Private/UVEditorMode.cpp
+++ b/Engine/Plugins/Experimental/UVEditor/Source/UVEditor/Private/UVEditorMode.cpp
@@ -6,6 +6,7 @@
 #include "ContextObjectStore.h"
 #include "Drawing/MeshElementsVisualizer.h"
 #include "EdModeInteractiveToolsContext.h" //ToolsContext
+#include "Framework/Commands/UICommandList.h"
 #include "InteractiveTool.h"
 #include "MeshOpPreviewHelpers.h" // UMeshOpPreviewWithBackgroundCompute
 #include "PreviewMesh.h"
@@ -18,13 +19,13 @@
 #include "ToolTargetManager.h"
 #include "ToolTargets/UVEditorToolMeshInput.h"
 #include "UVEditorCommands.h"
+#include "UVEditorLayoutTool.h"
 #include "UVSelectTool.h"
 #include "UVEditorModeToolkit.h"
 #include "UVEditorSubsystem.h"
 #include "UVEditorToolUtil.h"
 #include "UVToolContextObjects.h"
 #include "UVEditorBackgroundPreview.h"
-#include "Editor.h"
 
 #define LOCTEXT_NAMESPACE "UUVEditorMode"
 
@@ -79,14 +80,18 @@ void UUVEditorMode::RegisterTools()
 	const FUVEditorCommands& CommandInfos = FUVEditorCommands::Get();
 
 	// TODO: Other tool registrations go here
-	auto UVSelectToolBuilder = NewObject<UUVSelectToolBuilder>();
+	UUVSelectToolBuilder* UVSelectToolBuilder = NewObject<UUVSelectToolBuilder>();
 	UVSelectToolBuilder->Targets = &ToolInputObjects;
 	RegisterTool(CommandInfos.BeginSelectTool, TEXT("UVSelectTool"), UVSelectToolBuilder);
 
-	auto UVTransformToolBuilder = NewObject<UUVSelectToolBuilder>();
+	UUVSelectToolBuilder* UVTransformToolBuilder = NewObject<UUVSelectToolBuilder>();
 	UVTransformToolBuilder->Targets = &ToolInputObjects;
 	UVTransformToolBuilder->bGizmoEnabled = true;
 	RegisterTool(CommandInfos.BeginTransformTool, TEXT("UVTransformTool"), UVTransformToolBuilder);
+
+	UUVEditorLayoutToolBuilder* UVEditorLayoutToolBuilder = NewObject<UUVEditorLayoutToolBuilder>();
+	UVEditorLayoutToolBuilder->Targets = &ToolInputObjects;
+	RegisterTool(CommandInfos.BeginLayoutTool, TEXT("UVLayoutTool"), UVEditorLayoutToolBuilder);
 }
 
 void UUVEditorMode::CreateToolkit()
@@ -102,12 +107,41 @@ void UUVEditorMode::ActivateDefaultTool()
 
 void UUVEditorMode::BindCommands()
 {
-	// We currently don't have mode-level commands (rather than tool level or asset
-	// editor level), but presumably they would go here if we had them.
+	const FUVEditorCommands& CommandInfos = FUVEditorCommands::Get();
+	const TSharedRef<FUICommandList>& CommandList = Toolkit->GetToolkitCommands();
+
+	// Hook up to Enter/Esc key presses
+	CommandList->MapAction(
+		CommandInfos.AcceptActiveTool,
+		FExecuteAction::CreateLambda([this]() { 
+			ToolsContext->EndTool(EToolShutdownType::Accept); 
+			ActivateDefaultTool();
+			}),
+		FCanExecuteAction::CreateLambda([this]() { return ToolsContext->CanAcceptActiveTool(); }),
+		FGetActionCheckState(),
+		FIsActionButtonVisible::CreateLambda([this]() {return ToolsContext->ActiveToolHasAccept(); }),
+		EUIActionRepeatMode::RepeatDisabled
+	);
+
+	CommandList->MapAction(
+		CommandInfos.CancelActiveTool,
+		FExecuteAction::CreateLambda([this]() { 
+			ToolsContext->EndTool(EToolShutdownType::Cancel);
+			ActivateDefaultTool();
+			}),
+		FCanExecuteAction::CreateLambda([this]() { return ToolsContext->CanCancelActiveTool(); }),
+		FGetActionCheckState(),
+		FIsActionButtonVisible::CreateLambda([this]() {return ToolsContext->ActiveToolHasAccept(); }),
+		EUIActionRepeatMode::RepeatDisabled
+	);
 }
 
 void UUVEditorMode::Exit()
 {
+	// ToolsContext->EndTool only shuts the tool on the next tick, and ToolsContext->DeactivateActiveTool is
+	// inaccessible, so we end up having to do this to force the shutdown right now.
+	ToolsContext->ToolManager->DeactivateTool(EToolSide::Mouse, EToolShutdownType::Cancel);
+
 	for (TObjectPtr<UUVEditorToolMeshInput> ToolInput : ToolInputObjects)
 	{
 		ToolInput->Shutdown();
diff --git a/Engine/Plugins/Experimental/UVEditor/Source/UVEditor/Private/UVEditorModeToolkit.cpp b/Engine/Plugins/Experimental/UVEditor/Source/UVEditor/Private/UVEditorModeToolkit.cpp
--- a/Engine/Plugins/Experimental/UVEditor/Source/UVEditor/Private/UVEditorModeToolkit.cpp
+++ b/Engine/Plugins/Experimental/UVEditor/Source/UVEditor/Private/UVEditorModeToolkit.cpp
@@ -8,11 +8,11 @@
 #include "Framework/MultiBox/MultiBoxDefs.h"
 #include "IDetailsView.h"
 #include "Tools/UEdMode.h"
+#include "UVEditorBackgroundPreview.h"
 #include "UVEditorCommands.h"
+#include "UVEditorMode.h"
 #include "Widgets/Layout/SBorder.h"
 #include "Widgets/Text/STextBlock.h"
-#include "UVEditorMode.h"
-#include "UVEditorBackgroundPreview.h"
 
 #include "EdModeInteractiveToolsContext.h"
 
@@ -135,6 +135,7 @@ void FUVEditorModeToolkit::Init(const TSharedPtr<IToolkitHost>& InitToolkitHost,
 	// TODO: Add more tools here
 	ToolbarBuilder.AddToolBarButton(CommandInfos.BeginSelectTool);
 	ToolbarBuilder.AddToolBarButton(CommandInfos.BeginTransformTool);
+	ToolbarBuilder.AddToolBarButton(CommandInfos.BeginLayoutTool);
 
 	// Hook in the tool palette
 	ToolButtonsContainer->SetContent(ToolbarBuilder.MakeWidget());
@@ -164,6 +165,97 @@ void FUVEditorModeToolkit::Init(const TSharedPtr<IToolkitHost>& InitToolkitHost,
 		BackgroundDetailsContainer->SetContent(BackgroundDetailsView.ToSharedRef());
 	}
 	
+	// Set up the overlay. Larely copied from ModelingToolsEditorModeToolkit, except
+	// that we don't have a "complete" option since we don't have a non-tool state.
+	// TODO: We could put some of the shared code in some common place.
+	SAssignNew(ViewportOverlayWidget, SHorizontalBox)
+
+	+SHorizontalBox::Slot()
+	.HAlign(HAlign_Center)
+	.VAlign(VAlign_Bottom)
+	.Padding(FMargin(0.0f, 0.0f, 0.f, 15.f))
+	[
+		SNew(SBorder)
+		.BorderImage(FAppStyle::Get().GetBrush("EditorViewport.OverlayBrush"))
+		.Padding(8.f)
+		[
+			SNew(SHorizontalBox)
+
+			// TODO: Need to add a tool icon here the way we do in modeling mode.
+			//+SHorizontalBox::Slot()
+			//.AutoWidth()
+			//.VAlign(VAlign_Center)
+			//.Padding(FMargin(0.f, 0.f, 8.f, 0.f))
+			//[
+			//	SNew(SImage)
+			//	.Image_Lambda([this] () { return ActiveToolIcon; })
+			//]
+
+			+SHorizontalBox::Slot()
+			.AutoWidth()
+			.VAlign(VAlign_Center)
+			.Padding(FMargin(0.f, 0.f, 8.f, 0.f))
+			[
+				SNew(STextBlock)
+				.Text(this, &FUVEditorModeToolkit::GetActiveToolDisplayName)
+			]
+
+			+SHorizontalBox::Slot()
+			.AutoWidth()
+			.Padding(FMargin(0.0, 0.f, 2.f, 0.f))
+			[
+				SNew(SButton)
+				.ButtonStyle(FAppStyle::Get(), "PrimaryButton")
+				.TextStyle( FAppStyle::Get(), "DialogButtonText" )
+				.Text(LOCTEXT("OverlayAccept", "Accept"))
+				.ToolTipText(LOCTEXT("OverlayAcceptTooltip", "Accept/Commit the results of the active Tool [Enter]"))
+				.HAlign(HAlign_Center)
+				.OnClicked_Lambda([this]() { 
+					GetScriptableEditorMode()->GetInteractiveToolsContext()->EndTool(EToolShutdownType::Accept);
+					Cast<UUVEditorMode>(GetScriptableEditorMode())->ActivateDefaultTool();
+					return FReply::Handled(); 
+					})
+				.IsEnabled_Lambda([this]() { return GetScriptableEditorMode()->GetInteractiveToolsContext()->CanAcceptActiveTool(); })
+				.Visibility_Lambda([this]() { return GetScriptableEditorMode()->GetInteractiveToolsContext()->ActiveToolHasAccept() ? EVisibility::Visible : EVisibility::Collapsed; })
+			]
+
+			+SHorizontalBox::Slot()
+			.AutoWidth()
+			.Padding(FMargin(2.0, 0.f, 0.f, 0.f))
+			[
+				SNew(SButton)
+				.ButtonStyle(FAppStyle::Get(), "Button")
+				.TextStyle( FAppStyle::Get(), "DialogButtonText" )
+				.Text(LOCTEXT("OverlayCancel", "Cancel"))
+				.ToolTipText(LOCTEXT("OverlayCancelTooltip", "Cancel the active Tool [Esc]"))
+				.HAlign(HAlign_Center)
+				.OnClicked_Lambda([this]() { 
+					GetScriptableEditorMode()->GetInteractiveToolsContext()->EndTool(EToolShutdownType::Cancel); 
+					Cast<UUVEditorMode>(GetScriptableEditorMode())->ActivateDefaultTool();
+					return FReply::Handled(); 
+					})
+				.IsEnabled_Lambda([this]() { return GetScriptableEditorMode()->GetInteractiveToolsContext()->CanCancelActiveTool(); })
+				.Visibility_Lambda([this]() { return GetScriptableEditorMode()->GetInteractiveToolsContext()->ActiveToolHasAccept() ? EVisibility::Visible : EVisibility::Collapsed; })
+			]
+
+			// Currently haven't needed a complete button, but this is here in case that changes.
+			//+SHorizontalBox::Slot()
+			//.AutoWidth()
+			//.Padding(FMargin(2.0, 0.f, 0.f, 0.f))
+			//[
+			//	SNew(SButton)
+			//	.ButtonStyle(FAppStyle::Get(), "PrimaryButton")
+			//	.TextStyle( FAppStyle::Get(), "DialogButtonText" )
+			//	.Text(LOCTEXT("OverlayComplete", "Complete"))
+			//	.ToolTipText(LOCTEXT("OverlayCompleteTooltip", "Exit the active Tool [Enter]"))
+			//	.HAlign(HAlign_Center)
+			//	.OnClicked_Lambda([this]() { GetScriptableEditorMode()->GetInteractiveToolsContext()->EndTool(EToolShutdownType::Completed); return FReply::Handled(); })
+			//	.IsEnabled_Lambda([this]() { return GetScriptableEditorMode()->GetInteractiveToolsContext()->CanCompleteActiveTool(); })
+			//	.Visibility_Lambda([this]() { return GetScriptableEditorMode()->GetInteractiveToolsContext()->CanCompleteActiveTool() ? EVisibility::Visible : EVisibility::Collapsed; })
+			//]
+		]	
+	];
+
 }
 
 FName FUVEditorModeToolkit::GetToolkitFName() const
@@ -183,4 +275,30 @@ void FUVEditorModeToolkit::SetBackgroundSettings(UObject* InSettingsObject)
 }
 
 
+void FUVEditorModeToolkit::OnToolStarted(UInteractiveToolManager* Manager, UInteractiveTool* Tool)
+{
+	FModeToolkit::OnToolStarted(Manager, Tool);
+	
+	ActiveToolName = Tool->GetToolInfo().ToolDisplayName;
+
+	if (GetScriptableEditorMode()->GetInteractiveToolsContext()->ActiveToolHasAccept())
+	{
+		GetToolkitHost()->AddViewportOverlayWidget(ViewportOverlayWidget.ToSharedRef());
+	}
+	
+
+}
+
+void FUVEditorModeToolkit::OnToolEnded(UInteractiveToolManager* Manager, UInteractiveTool* Tool)
+{
+	FModeToolkit::OnToolEnded(Manager, Tool);
+
+	ActiveToolName = FText::GetEmpty();
+
+	if (IsHosted())
+	{
+		GetToolkitHost()->RemoveViewportOverlayWidget(ViewportOverlayWidget.ToSharedRef());
+	}
+}
+
 #undef LOCTEXT_NAMESPACE
diff --git a/Engine/Plugins/Experimental/UVEditor/Source/UVEditor/Private/UVEditorToolkit.cpp b/Engine/Plugins/Experimental/UVEditor/Source/UVEditor/Private/UVEditorToolkit.cpp
--- a/Engine/Plugins/Experimental/UVEditor/Source/UVEditor/Private/UVEditorToolkit.cpp
+++ b/Engine/Plugins/Experimental/UVEditor/Source/UVEditor/Private/UVEditorToolkit.cpp
@@ -9,6 +9,7 @@
 #include "Misc/MessageDialog.h"
 #include "PreviewScene.h"
 #include "SAssetEditorViewport.h"
+#include "SUVEditor2DViewport.h"
 #include "ThumbnailRendering/SceneThumbnailInfo.h"
 #include "ToolMenus.h"
 #include "UVEditor.h"
@@ -19,6 +20,8 @@
 #include "UVToolContextObjects.h"
 #include "Widgets/Docking/SDockTab.h"
 
+#include "SLevelViewport.h"
+
 #include "EdModeInteractiveToolsContext.h"
 
 #define LOCTEXT_NAMESPACE "UVEditorToolkit"
@@ -212,6 +215,19 @@ bool FUVEditorToolkit::OnRequestClose()
 	return FAssetEditorToolkit::OnRequestClose();
 }
 
+// These get called indirectly (via toolkit host) from the mode toolkit when the mode starts or ends a tool,
+// in order to add or remove an accept/cancel overlay.
+void FUVEditorToolkit::AddViewportOverlayWidget(TSharedRef<SWidget> InViewportOverlayWidget) 
+{
+	TSharedPtr<SUVEditor2DViewport> ViewportWidget = StaticCastSharedPtr<SUVEditor2DViewport>(ViewportTabContent->GetFirstViewport());
+	ViewportWidget->AddOverlayWidget(InViewportOverlayWidget);
+}
+void FUVEditorToolkit::RemoveViewportOverlayWidget(TSharedRef<SWidget> InViewportOverlayWidget)
+{
+	TSharedPtr<SUVEditor2DViewport> ViewportWidget = StaticCastSharedPtr<SUVEditor2DViewport>(ViewportTabContent->GetFirstViewport());
+	ViewportWidget->RemoveOverlayWidget(InViewportOverlayWidget);
+}
+
 // We override the "Save" button behavior slightly to apply our changes before saving the asset.
 void FUVEditorToolkit::SaveAsset_Execute()
 {
@@ -273,7 +289,7 @@ AssetEditorViewportFactoryFunction FUVEditorToolkit::GetViewportDelegate()
 {
 	AssetEditorViewportFactoryFunction TempViewportDelegate = [this](FAssetEditorViewportConstructionArgs InArgs)
 	{
-		return SNew(SAssetEditorViewport, InArgs)
+		return SNew(SUVEditor2DViewport, InArgs)
 			.EditorViewportClient(ViewportClient);
 	};
 
diff --git a/Engine/Plugins/Experimental/UVEditor/Source/UVEditor/Public/SUVEditor2DViewport.h b/Engine/Plugins/Experimental/UVEditor/Source/UVEditor/Public/SUVEditor2DViewport.h
new file mode 100644
--- /dev/null
+++ b/Engine/Plugins/Experimental/UVEditor/Source/UVEditor/Public/SUVEditor2DViewport.h
@@ -0,0 +1,19 @@
+// Copyright Epic Games, Inc. All Rights Reserved.
+
+#pragma once
+
+#include "CoreMinimal.h"
+
+#include "SAssetEditorViewport.h"
+
+
+
+class UVEDITOR_API SUVEditor2DViewport : public SAssetEditorViewport
+{
+public:
+
+	// These allow the toolkit to add an accept/cancel overlay when needed. PopulateViewportOverlays
+	// is not helpful here because that gets called just once.
+	virtual void AddOverlayWidget(TSharedRef<SWidget> OverlaidWidget);
+	virtual void RemoveOverlayWidget(TSharedRef<SWidget> OverlaidWidget);
+};
diff --git a/Engine/Plugins/Experimental/UVEditor/Source/UVEditor/Public/UVEditorCommands.h b/Engine/Plugins/Experimental/UVEditor/Source/UVEditor/Public/UVEditorCommands.h
--- a/Engine/Plugins/Experimental/UVEditor/Source/UVEditor/Public/UVEditorCommands.h
+++ b/Engine/Plugins/Experimental/UVEditor/Source/UVEditor/Public/UVEditorCommands.h
@@ -20,4 +20,8 @@ public:
 
 	TSharedPtr<FUICommandInfo> BeginSelectTool;
 	TSharedPtr<FUICommandInfo> BeginTransformTool;
+	TSharedPtr<FUICommandInfo> BeginLayoutTool;
+
+	TSharedPtr<FUICommandInfo> AcceptActiveTool;
+	TSharedPtr<FUICommandInfo> CancelActiveTool;
 };
\ No newline at end of file
diff --git a/Engine/Plugins/Experimental/UVEditor/Source/UVEditor/Public/UVEditorMode.h b/Engine/Plugins/Experimental/UVEditor/Source/UVEditor/Public/UVEditorMode.h
--- a/Engine/Plugins/Experimental/UVEditor/Source/UVEditor/Public/UVEditorMode.h
+++ b/Engine/Plugins/Experimental/UVEditor/Source/UVEditor/Public/UVEditorMode.h
@@ -70,6 +70,11 @@ public:
 	virtual void Enter() override;
 	virtual void Exit() override;
 	virtual void ModeTick(float DeltaTime) override;
+	// We're changing visibility of this one to public here so that we can call it from the toolkit
+	// when clicking accept/cancel buttons. We don't want to friend the toolkit because we don't want
+	// it to later get (accidentally) more entangled with mode internals. At the same time, we're not
+	// sure whether we want to make the UEdMode one public. So this is the minimal-impact tweak.
+	virtual void ActivateDefaultTool() override;
 
 	// We don't actually override MouseEnter, etc, because things get forwarded to the input
 	// router via FEditorModeTools, and we don't have any additional input handling to do at the mode level.
@@ -77,14 +82,13 @@ public:
 
 	// Holds the background visualiztion
 	UPROPERTY()
-		TObjectPtr<UUVEditorBackgroundPreview> BackgroundVisualization;
+	TObjectPtr<UUVEditorBackgroundPreview> BackgroundVisualization;
 
 
 protected:
 
 	// UEdMode overrides
 	virtual void CreateToolkit() override;
-	virtual void ActivateDefaultTool() override;
 	// Not sure whether we need these yet
 	virtual void BindCommands() override;
 	virtual void OnToolStarted(UInteractiveToolManager* Manager, UInteractiveTool* Tool) {}
diff --git a/Engine/Plugins/Experimental/UVEditor/Source/UVEditor/Public/UVEditorModeToolkit.h b/Engine/Plugins/Experimental/UVEditor/Source/UVEditor/Public/UVEditorModeToolkit.h
--- a/Engine/Plugins/Experimental/UVEditor/Source/UVEditor/Public/UVEditorModeToolkit.h
+++ b/Engine/Plugins/Experimental/UVEditor/Source/UVEditor/Public/UVEditorModeToolkit.h
@@ -31,6 +31,10 @@ public:
 
 	virtual void SetBackgroundSettings(UObject* InSettingsObject);
 
+	virtual FText GetActiveToolDisplayName() const override { return ActiveToolName; }
+	virtual void OnToolStarted(UInteractiveToolManager* Manager, UInteractiveTool* Tool) override;
+	virtual void OnToolEnded(UInteractiveToolManager* Manager, UInteractiveTool* Tool) override;
+
 protected:
 	// The mode's entire toolbox, which gets returned by GetInlineContent()
 	TSharedPtr<SWidget> ToolkitWidget;
@@ -54,4 +58,9 @@ protected:
 
 	// A place for tools to write out any instructions
 	TSharedPtr<STextBlock> ToolMessageArea;
+
+	/** Contains the widget container for the Accept/Cancel buttons for tools. */
+	TSharedPtr<SWidget> ViewportOverlayWidget;
+
+	FText ActiveToolName;
 };
\ No newline at end of file
diff --git a/Engine/Plugins/Experimental/UVEditor/Source/UVEditor/Public/UVEditorToolkit.h b/Engine/Plugins/Experimental/UVEditor/Source/UVEditor/Public/UVEditorToolkit.h
--- a/Engine/Plugins/Experimental/UVEditor/Source/UVEditor/Public/UVEditorToolkit.h
+++ b/Engine/Plugins/Experimental/UVEditor/Source/UVEditor/Public/UVEditorToolkit.h
@@ -11,6 +11,7 @@ class FUVEditorPreviewModeManager;
 class IDetailsView;
 class SDockTab;
 class SBorder;
+class SWidget;
 class UInteractiveToolsContext;
 class UInputRouter;
 
@@ -37,6 +38,8 @@ public:
 	virtual void CreateWidgets() override;
 
 	// FAssetEditorToolkit
+	virtual void AddViewportOverlayWidget(TSharedRef<SWidget> InViewportOverlayWidget) override;
+	virtual void RemoveViewportOverlayWidget(TSharedRef<SWidget> InViewportOverlayWidget) override;
 	virtual void CreateEditorModeManager() override;
 	virtual FText GetToolkitName() const override;
 	virtual FName GetToolkitFName() const override;
diff --git a/Engine/Plugins/Experimental/UVEditor/Source/UVEditorTools/Private/ToolTargets/UVEditorToolMeshInput.cpp b/Engine/Plugins/Experimental/UVEditor/Source/UVEditorTools/Private/ToolTargets/UVEditorToolMeshInput.cpp
--- a/Engine/Plugins/Experimental/UVEditor/Source/UVEditorTools/Private/ToolTargets/UVEditorToolMeshInput.cpp
+++ b/Engine/Plugins/Experimental/UVEditor/Source/UVEditorTools/Private/ToolTargets/UVEditorToolMeshInput.cpp
@@ -192,7 +192,7 @@ void UUVEditorToolMeshInput::Shutdown()
 	OriginalAsset = nullptr;
 }
 
-void UUVEditorToolMeshInput::UpdateUnwrapPreviewOverlay(const TArray<int32>* ChangedVids, const TArray<int32>* ChangedConnectivityTids, const TArray<int32>* FastRenderUpdateTids)
+void UUVEditorToolMeshInput::UpdateUnwrapPreviewOverlayFromPositions(const TArray<int32>* ChangedVids, const TArray<int32>* ChangedConnectivityTids, const TArray<int32>* FastRenderUpdateTids)
 {
 	UnwrapPreview->PreviewMesh->DeferredEditMesh([this, ChangedVids, ChangedConnectivityTids](FDynamicMesh3& Mesh)
 	{
@@ -217,7 +217,7 @@ void UUVEditorToolMeshInput::UpdateUnwrapPreviewOverlay(const TArray<int32>* Cha
 	}
 }
 
-void UUVEditorToolMeshInput::UpdateUnwrapCanonicalOverlay(const TArray<int32>* ChangedVids, const TArray<int32>* ChangedConnectivityTids)
+void UUVEditorToolMeshInput::UpdateUnwrapCanonicalOverlayFromPositions(const TArray<int32>* ChangedVids, const TArray<int32>* ChangedConnectivityTids)
 {
 	FDynamicMeshUVOverlay* DestOverlay = UnwrapCanonical->Attributes()->PrimaryUV();
 	UVEditorToolUtil::UpdateUVOverlayFromUnwrapMesh(*UnwrapCanonical, *DestOverlay, VertPositionToUV,
@@ -295,112 +295,72 @@ void UUVEditorToolMeshInput::UpdateCanonicalFromPreviews(const TArray<int32>* Ch
 	CopyMeshOverlay(*SourceOverlay, *DestOverlay, ChangedVids, ChangedConnectivityTids);
 }
 
-void UUVEditorToolMeshInput::UpdateAllFromUnwrapPreview(
-	const TArray<int32>* ChangedVids, const TArray<int32>* ChangedConnectivityTids, const TArray<int32>* FastRenderUpdateTids)
+void UUVEditorToolMeshInput::UpdatePreviewsFromCanonical(const TArray<int32>* ChangedVids, const TArray<int32>* ChangedConnectivityTids, const TArray<int32>* FastRenderUpdateTids)
 {
-	const FDynamicMesh3* ReferenceMesh = UnwrapPreview->PreviewMesh->GetMesh();
-
-	// Update UnwrapCanonical
-	UpdateOtherUnwrap(*ReferenceMesh, *UnwrapCanonical, ChangedVids, ChangedConnectivityTids);
-
-	// Update the applied meshes
-	const FDynamicMeshUVOverlay* SourceOverlay = ReferenceMesh->Attributes()->PrimaryUV();
-	UpdateAppliedOverlays(*SourceOverlay, ChangedVids, ChangedConnectivityTids, FastRenderUpdateTids);
-}
+	using namespace UVEditorToolMeshInputLocals;
 
-void UUVEditorToolMeshInput::UpdateAllFromUnwrapCanonical(
-	const TArray<int32>* ChangedVids, const TArray<int32>* ChangedConnectivityTids, const TArray<int32>* FastRenderUpdateTids)
-{
-	// Update UnwrapPreview
+	// Update UnwrapPreview from UnwrapCanonical
 	UnwrapPreview->PreviewMesh->DeferredEditMesh([this, ChangedVids, ChangedConnectivityTids](FDynamicMesh3& Mesh)
 	{
 		UpdateOtherUnwrap(*UnwrapCanonical, Mesh, ChangedVids, ChangedConnectivityTids);
 	}, false);
-
 	if (FastRenderUpdateTids) {
 		UnwrapPreview->PreviewMesh->NotifyRegionDeferredEditCompleted(*FastRenderUpdateTids,
-			EMeshRenderAttributeFlags::Positions | EMeshRenderAttributeFlags::VertexUVs); // Assuming user changed positions and didn't notify yet
+			EMeshRenderAttributeFlags::Positions | EMeshRenderAttributeFlags::VertexUVs);
 	}
 	else
 	{
 		UnwrapPreview->PreviewMesh->NotifyDeferredEditCompleted(UPreviewMesh::ERenderUpdateMode::FastUpdate,
-			EMeshRenderAttributeFlags::Positions | EMeshRenderAttributeFlags::VertexUVs, true); // Assuming user changed positions and didn't notify yet
+			EMeshRenderAttributeFlags::Positions | EMeshRenderAttributeFlags::VertexUVs, true);
 	}
 
-	// Update the applied meshes
-	FDynamicMeshUVOverlay* SourceOverlay = UnwrapCanonical->Attributes()->PrimaryUV();
-	UpdateAppliedOverlays(*SourceOverlay, ChangedVids, ChangedConnectivityTids, FastRenderUpdateTids);
-
 	if (WireframeDisplay)
 	{
 		WireframeDisplay->NotifyMeshChanged();
 	}
-}
 
-void UUVEditorToolMeshInput::UpdateAllFromAppliedPreview(
-	const TArray<int32>* ChangedElementIDs, const TArray<int32>* ChangedConnectivityTids, const TArray<int32>* FastRenderUpdateTids)
-{
-	using namespace UVEditorToolMeshInputLocals;
-
-	// Update AppliedCanonical
-	const FDynamicMeshUVOverlay* SourceOverlay = AppliedPreview->PreviewMesh->GetMesh()
-		->Attributes()->GetUVLayer(UVLayerIndex);
-	FDynamicMeshUVOverlay* DestOverlay = AppliedCanonical->Attributes()->GetUVLayer(UVLayerIndex);
-	CopyMeshOverlay(*SourceOverlay, *DestOverlay, ChangedElementIDs, ChangedConnectivityTids);
-
-	// Update UnwrapCanonical
-	UVEditorToolUtil::UpdateUVUnwrapMesh(*SourceOverlay, *UnwrapCanonical, 
-		UVToVertPosition, ChangedElementIDs, ChangedConnectivityTids);
-
-	// Update UnwrapPreview from UnwrapCanonical
-	UnwrapPreview->PreviewMesh->DeferredEditMesh([this, ChangedElementIDs, ChangedConnectivityTids](FDynamicMesh3& Mesh) {
-		UpdateOtherUnwrap(*UnwrapCanonical, Mesh, ChangedElementIDs, ChangedConnectivityTids);
+	// Update AppliedPreview from AppliedCanonical
+	AppliedPreview->PreviewMesh->DeferredEditMesh([this, ChangedVids, ChangedConnectivityTids](FDynamicMesh3& Mesh)
+	{
+		const FDynamicMeshUVOverlay* SourceOverlay = AppliedCanonical->Attributes()->GetUVLayer(UVLayerIndex);
+		FDynamicMeshUVOverlay* DestOverlay = Mesh.Attributes()->GetUVLayer(UVLayerIndex);
+		CopyMeshOverlay(*SourceOverlay, *DestOverlay, ChangedVids, ChangedConnectivityTids);
 	}, false);
-
 	if (FastRenderUpdateTids) {
-		UnwrapPreview->PreviewMesh->NotifyRegionDeferredEditCompleted(*FastRenderUpdateTids,
-			EMeshRenderAttributeFlags::Positions | EMeshRenderAttributeFlags::VertexUVs); // Assuming user changed positions and didn't notify yet
+		AppliedPreview->PreviewMesh->NotifyRegionDeferredEditCompleted(*FastRenderUpdateTids,
+			EMeshRenderAttributeFlags::VertexUVs);
 	}
 	else
 	{
-		UnwrapPreview->PreviewMesh->NotifyDeferredEditCompleted(UPreviewMesh::ERenderUpdateMode::FastUpdate,
-			EMeshRenderAttributeFlags::Positions | EMeshRenderAttributeFlags::VertexUVs, true); // Assuming user changed positions and didn't notify yet
+		AppliedPreview->PreviewMesh->NotifyDeferredEditCompleted(UPreviewMesh::ERenderUpdateMode::FastUpdate,
+			EMeshRenderAttributeFlags::VertexUVs, true);
 	}
+}
 
-	if (WireframeDisplay)
-	{
-		WireframeDisplay->NotifyMeshChanged();
-	}
+void UUVEditorToolMeshInput::UpdateAllFromUnwrapPreview(const TArray<int32>* ChangedVids, const TArray<int32>* ChangedConnectivityTids, const TArray<int32>* FastRenderUpdateTids)
+{
+	UpdateAppliedPreviewFromUnwrapPreview(ChangedVids, ChangedConnectivityTids, FastRenderUpdateTids);
+	UpdateCanonicalFromPreviews(ChangedVids, ChangedConnectivityTids);
 }
 
-/**
- * Helper function. Updates the UV overlays of AppliedCanonical and AppliedPreview.
- */
-void UUVEditorToolMeshInput::UpdateAppliedOverlays(const FDynamicMeshUVOverlay& SourceOverlay,
+void UUVEditorToolMeshInput::UpdateAllFromUnwrapCanonical(
 	const TArray<int32>* ChangedVids, const TArray<int32>* ChangedConnectivityTids, const TArray<int32>* FastRenderUpdateTids)
 {
 	using namespace UVEditorToolMeshInputLocals;
 
-	// Update UnwrappedMeshCanonical
+	// Update AppliedCanonical
+	FDynamicMeshUVOverlay* SourceOverlay = UnwrapCanonical->Attributes()->PrimaryUV();
 	FDynamicMeshUVOverlay* DestOverlay = AppliedCanonical->Attributes()->GetUVLayer(UVLayerIndex);
-	CopyMeshOverlay(SourceOverlay, *DestOverlay, ChangedVids, ChangedConnectivityTids);
+	CopyMeshOverlay(*SourceOverlay, *DestOverlay, ChangedVids, ChangedConnectivityTids);
 
-	// Update AppliedPreview
-	AppliedPreview->PreviewMesh->DeferredEditMesh([this, &SourceOverlay, ChangedVids, ChangedConnectivityTids](FDynamicMesh3& Mesh)
-	{
-		FDynamicMeshUVOverlay* DestOverlay = Mesh.Attributes()->GetUVLayer(UVLayerIndex);
-		CopyMeshOverlay(SourceOverlay, *DestOverlay, ChangedVids, ChangedConnectivityTids);
-	}, false);
+	UpdatePreviewsFromCanonical(ChangedVids, ChangedConnectivityTids, FastRenderUpdateTids);
+}
 
-	if (FastRenderUpdateTids) {
-		AppliedPreview->PreviewMesh->NotifyRegionDeferredEditCompleted(*FastRenderUpdateTids,
-			EMeshRenderAttributeFlags::Positions | EMeshRenderAttributeFlags::VertexUVs); // Assuming user changed positions and didn't notify yet
-	}
-	else
-	{
-		AppliedPreview->PreviewMesh->NotifyDeferredEditCompleted(UPreviewMesh::ERenderUpdateMode::FastUpdate,
-			EMeshRenderAttributeFlags::VertexUVs, false); // Assuming user changed positions and didn't notify yet
-	}
+void UUVEditorToolMeshInput::UpdateAllFromAppliedPreview(
+	const TArray<int32>* ChangedElementIDs, const TArray<int32>* ChangedConnectivityTids, const TArray<int32>* FastRenderUpdateTids)
+{
+	UpdateUnwrapPreviewFromAppliedPreview(ChangedElementIDs, ChangedConnectivityTids, FastRenderUpdateTids);
+	UpdateCanonicalFromPreviews(ChangedElementIDs, ChangedConnectivityTids);
 }
 
 /**
diff --git a/Engine/Plugins/Experimental/UVEditor/Source/UVEditorTools/Private/UVEditorLayoutTool.cpp b/Engine/Plugins/Experimental/UVEditor/Source/UVEditorTools/Private/UVEditorLayoutTool.cpp
new file mode 100644
--- /dev/null
+++ b/Engine/Plugins/Experimental/UVEditor/Source/UVEditorTools/Private/UVEditorLayoutTool.cpp
@@ -0,0 +1,128 @@
+// Copyright Epic Games, Inc. All Rights Reserved.
+
+#include "UVEditorLayoutTool.h"
+
+#include "InteractiveToolManager.h"
+#include "MeshOpPreviewHelpers.h" // UMeshOpPreviewWithBackgroundCompute
+#include "ParameterizationOps/UVLayoutOp.h"
+#include "Properties/UVLayoutProperties.h"
+#include "ToolTargets/UVEditorToolMeshInput.h"
+
+
+#define LOCTEXT_NAMESPACE "UUVSelectTool"
+
+// Tool builder
+// TODO: Could consider sharing some of the tool builder boilerplate for UV editor tools in a common base class.
+
+const FToolTargetTypeRequirements& UUVEditorLayoutToolBuilder::GetTargetRequirements() const
+{
+	static FToolTargetTypeRequirements TypeRequirements(UUVEditorToolMeshInput::StaticClass());
+	return TypeRequirements;
+}
+
+bool UUVEditorLayoutToolBuilder::CanBuildTool(const FToolBuilderState& SceneState) const
+{
+	return Targets && Targets->Num() > 0;
+}
+
+UInteractiveTool* UUVEditorLayoutToolBuilder::BuildTool(const FToolBuilderState& SceneState) const
+{
+	UUVEditorLayoutTool* NewTool = NewObject<UUVEditorLayoutTool>(SceneState.ToolManager);
+	NewTool->SetTargets(*Targets);
+
+	return NewTool;
+}
+
+
+void UUVEditorLayoutTool::Setup()
+{
+	check(Targets.Num() > 0);
+
+	UInteractiveTool::Setup();
+
+	Settings = NewObject<UUVLayoutProperties>(this);
+	Settings->RestoreProperties(this);
+	AddToolPropertySource(Settings);
+
+	for (TObjectPtr<UUVEditorToolMeshInput> Target : Targets)
+	{
+		UUVLayoutOperatorFactory* OpFactory = NewObject<UUVLayoutOperatorFactory>();
+		OpFactory->TargetTransform = Target->AppliedPreview->PreviewMesh->GetTransform();
+		OpFactory->Settings = Settings;
+		OpFactory->OriginalMesh = Target->AppliedCanonical;
+		OpFactory->GetSelectedUVChannel = [Target]() { return Target->UVLayerIndex; };
+
+		Target->AppliedPreview->ChangeOpFactory(OpFactory);
+		Target->AppliedPreview->OnMeshUpdated.AddWeakLambda(this, [Target](UMeshOpPreviewWithBackgroundCompute* Preview) {
+			Target->UpdateUnwrapPreviewFromAppliedPreview();
+		});
+
+		Target->AppliedPreview->InvalidateResult();
+	}
+
+	SetToolDisplayName(LOCTEXT("ToolName", "UV Layout"));
+	GetToolManager()->DisplayMessage(LOCTEXT("OnStartUVLayoutTool", "Transform/Rotate/Scale existing UV Charts using various strategies"),
+		EToolMessageLevel::UserNotification);
+}
+
+void UUVEditorLayoutTool::Shutdown(EToolShutdownType ShutdownType)
+{
+	Settings->SaveProperties(this);
+
+	for (TObjectPtr<UUVEditorToolMeshInput> Target : Targets)
+	{
+		Target->AppliedPreview->OnMeshUpdated.RemoveAll(this);
+		Target->AppliedPreview->ClearOpFactory();
+	}
+
+	if (ShutdownType == EToolShutdownType::Accept)
+	{
+		for (TObjectPtr<UUVEditorToolMeshInput> Target : Targets)
+		{
+			Target->UpdateCanonicalFromPreviews();
+		}
+	}
+	else
+	{
+		// Reset the inputs
+		for (TObjectPtr<UUVEditorToolMeshInput> Target : Targets)
+		{
+			Target->UpdatePreviewsFromCanonical();
+		}
+	}
+
+	Settings = nullptr;
+	Targets.Empty();
+}
+
+void UUVEditorLayoutTool::OnTick(float DeltaTime)
+{
+	for (TObjectPtr<UUVEditorToolMeshInput> Target : Targets)
+	{
+		Target->AppliedPreview->Tick(DeltaTime);
+	}
+}
+
+
+
+void UUVEditorLayoutTool::OnPropertyModified(UObject* PropertySet, FProperty* Property)
+{
+	for (TObjectPtr<UUVEditorToolMeshInput> Target : Targets)
+	{
+		Target->AppliedPreview->InvalidateResult();
+	}
+}
+
+bool UUVEditorLayoutTool::CanAccept() const
+{
+	for (TObjectPtr<UUVEditorToolMeshInput> Target : Targets)
+	{
+		if (!Target->AppliedPreview->HaveValidResult())
+		{
+			return false;
+		}
+	}
+	return true;
+}
+
+#undef LOCTEXT_NAMESPACE
diff --git a/Engine/Plugins/Experimental/UVEditor/Source/UVEditorTools/Private/UVSelectTool.cpp b/Engine/Plugins/Experimental/UVEditor/Source/UVEditorTools/Private/UVSelectTool.cpp
--- a/Engine/Plugins/Experimental/UVEditor/Source/UVEditorTools/Private/UVSelectTool.cpp
+++ b/Engine/Plugins/Experimental/UVEditor/Source/UVEditorTools/Private/UVSelectTool.cpp
@@ -86,6 +86,10 @@ void UUVSelectTool::Setup()
 			Tree->SetMesh(Target->UnwrapCanonical.Get());
 			TreeStore->Set(Target->UnwrapCanonical.Get(), Tree);
 		}
+		if (!Tree->IsValid())
+		{
+			Tree->Build();
+		}
 		AABBTrees.Add(Tree);
 	}
 
@@ -445,7 +449,7 @@ void UUVSelectTool::ApplyGizmoTransform()
 					MeshIn.SetVertex(MovingVids[i], TransformToApply.TransformPosition(MovingVertOriginalPositions[i]));
 				}
 			}, false);
-		Targets[SelectionTargetIndex]->UpdateUnwrapPreviewOverlay(&MovingVids, nullptr, &SelectedTids);
+		Targets[SelectionTargetIndex]->UpdateUnwrapPreviewOverlayFromPositions(&MovingVids, nullptr, &SelectedTids);
 
 		SelectionMechanic->SetDrawnElementsTransform((FTransform)TransformToApply);
 
diff --git a/Engine/Plugins/Experimental/UVEditor/Source/UVEditorTools/Public/ToolTargets/UVEditorToolMeshInput.h b/Engine/Plugins/Experimental/UVEditor/Source/UVEditorTools/Public/ToolTargets/UVEditorToolMeshInput.h
--- a/Engine/Plugins/Experimental/UVEditor/Source/UVEditorTools/Public/ToolTargets/UVEditorToolMeshInput.h
+++ b/Engine/Plugins/Experimental/UVEditor/Source/UVEditorTools/Public/ToolTargets/UVEditorToolMeshInput.h
@@ -110,13 +110,13 @@ public:
 	 * Updates UnwrapPreview UV Overlay from UnwrapPreview vert positions. Issues a NotifyDeferredEditCompleted
 	 * for both positions and UVs.
 	 */
-	void UpdateUnwrapPreviewOverlay(const TArray<int32>* ChangedVids = nullptr, const TArray<int32>* ChangedConnectivityTids = nullptr, const TArray<int32>* FastRenderUpdateTids = nullptr);
+	void UpdateUnwrapPreviewOverlayFromPositions(const TArray<int32>* ChangedVids = nullptr, const TArray<int32>* ChangedConnectivityTids = nullptr, const TArray<int32>* FastRenderUpdateTids = nullptr);
 
 	/**
 	 * Updates UnwrapCanonical UV Overlay from UnwrapCanonical vert positions. Issues a NotifyDeferredEditCompleted
 	 * for both positions and UVs.
 	 */
-	void UpdateUnwrapCanonicalOverlay(const TArray<int32>* ChangedVids = nullptr, const TArray<int32>* ChangedConnectivityTids = nullptr);
+	void UpdateUnwrapCanonicalOverlayFromPositions(const TArray<int32>* ChangedVids = nullptr, const TArray<int32>* ChangedConnectivityTids = nullptr);
 
 	/**
 	 * Updates the AppliedPreview from UnwrapPreview, without updating the non-preview meshes. Useful for updates during
@@ -126,8 +126,8 @@ public:
 	void UpdateAppliedPreviewFromUnwrapPreview(const TArray<int32>* ChangedVids = nullptr, const TArray<int32>* ChangedConnectivityTids = nullptr, const TArray<int32>* FastRenderUpdateTids = nullptr);
 
 	/**
-	 * Updates only the UnwrapPreview from AppliedPreview, without updating the non-preview meshes. Useful for updates during
-	 * a drag, etc, when we only care about updating the visible items.
+	 * Updates only the UnwrapPreview from AppliedPreview, without updating the non-preview meshes. Useful for transient
+	 * updates when we only care about updating the visible items.
 	 */
 	void UpdateUnwrapPreviewFromAppliedPreview(
 		const TArray<int32>* ChangedElementIDs = nullptr, const TArray<int32>* ChangedConnectivityTids = nullptr, const TArray<int32>* FastRenderUpdateTids = nullptr);
@@ -138,13 +138,22 @@ public:
 	 */
 	void UpdateCanonicalFromPreviews(const TArray<int32>* ChangedVids = nullptr, const TArray<int32>* ChangedConnectivityTids = nullptr);
 
+	/**
+	 * Updates the preview meshes from their canonical counterparts. Useful mainly as a way to reset the previews.
+	 */
+	void UpdatePreviewsFromCanonical(const TArray<int32>* ChangedVids = nullptr, const TArray<int32>* ChangedConnectivityTids = nullptr,
+		const TArray<int32>* FastRenderUpdateTids = nullptr);
+
 	/**
 	 * Updates the other meshes using the mesh in UnwrapPreview. Assumes that the overlay in UnwrapPreview is updated.
+	 * Used to update everything when the applied preview was not being updated in tandem with the unwrap preview (otherwise one 
+	 * would use UpdateCanonicalFromPreviews).
 	 */
 	void UpdateAllFromUnwrapPreview(const TArray<int32>* ChangedVids = nullptr, const TArray<int32>* ChangedConnectivityTids = nullptr, const TArray<int32>* FastRenderUpdateTids = nullptr);
 
 	/**
-	 * Updates the other meshes using the mesh in UnwrapCanonical. Assumes that overlay in UnwrapCanonical is updated.
+	 * Updates the other meshes using the mesh in UnwrapCanonical. Useful when immediately applying an operation that
+	 * operates directly on the unwrap mesh.
 	 */
 	void UpdateAllFromUnwrapCanonical(const TArray<int32>* ChangedVids = nullptr, const TArray<int32>* ChangedConnectivityTids = nullptr, const TArray<int32>* FastRenderUpdateTids = nullptr);
 
@@ -157,8 +166,6 @@ public:
 	virtual bool IsValid() const override;
 protected:
 
-	void UpdateAppliedOverlays(const UE::Geometry::FDynamicMeshUVOverlay& SourceOverlay,
-		const TArray<int32>* ChangedVids, const TArray<int32>* ChangedConnectivityTids, const TArray<int32>* FastRenderUpdateTids);
 	void UpdateOtherUnwrap(const UE::Geometry::FDynamicMesh3& SourceUnwrapMesh,
 		UE::Geometry::FDynamicMesh3& DestUnwrapMesh, const TArray<int32>* ChangedVids, const TArray<int32>* ChangedConnectivityTids);
 };
\ No newline at end of file
diff --git a/Engine/Plugins/Experimental/UVEditor/Source/UVEditorTools/Public/UVEditorLayoutTool.h b/Engine/Plugins/Experimental/UVEditor/Source/UVEditorTools/Public/UVEditorLayoutTool.h
new file mode 100644
--- /dev/null
+++ b/Engine/Plugins/Experimental/UVEditor/Source/UVEditorTools/Public/UVEditorLayoutTool.h
@@ -0,0 +1,67 @@
+// Copyright Epic Games, Inc. All Rights Reserved.
+
+#pragma once
+
+#include "CoreMinimal.h"
+
+#include "InteractiveTool.h"
+#include "InteractiveToolBuilder.h"
+
+#include "UVEditorLayoutTool.generated.h"
+
+class UUVEditorToolMeshInput;
+class UUVLayoutProperties;
+
+UCLASS()
+class UVEDITORTOOLS_API UUVEditorLayoutToolBuilder : public UInteractiveToolBuilder
+{
+	GENERATED_BODY()
+
+public:
+	virtual bool CanBuildTool(const FToolBuilderState& SceneState) const override;
+	virtual UInteractiveTool* BuildTool(const FToolBuilderState& SceneState) const override;
+
+	// This is a pointer so that it can be updated under the builder without
+	// having to set it in the mode after initializing targets.
+	const TArray<TObjectPtr<UUVEditorToolMeshInput>>* Targets = nullptr;
+
+protected:
+	virtual const FToolTargetTypeRequirements& GetTargetRequirements() const override;
+};
+
+/**
+ * 
+ */
+UCLASS()
+class UVEDITORTOOLS_API UUVEditorLayoutTool : public UInteractiveTool
+{
+	GENERATED_BODY()
+
+public:
+
+	/**
+	 * The tool will operate on the meshes given here.
+	 */
+	virtual void SetTargets(const TArray<TObjectPtr<UUVEditorToolMeshInput>>& TargetsIn)
+	{
+		Targets = TargetsIn;
+	}
+
+	virtual void Setup() override;
+	virtual void Shutdown(EToolShutdownType ShutdownType) override;
+
+	virtual void OnTick(float DeltaTime) override;
+
+	virtual bool HasCancel() const override { return true; }
+	virtual bool HasAccept() const override { return true; }
+	virtual bool CanAccept() const override;
+
+	virtual void OnPropertyModified(UObject* PropertySet, FProperty* Property) override;
+protected:
+
+	UPROPERTY()
+	TArray<TObjectPtr<UUVEditorToolMeshInput>> Targets;
+
+	UPROPERTY()
+	TObjectPtr<UUVLayoutProperties> Settings = nullptr;
+};
diff --git a/Engine/Plugins/Experimental/UVEditor/Source/UVEditorTools/UVEditorTools.Build.cs b/Engine/Plugins/Experimental/UVEditor/Source/UVEditorTools/UVEditorTools.Build.cs
--- a/Engine/Plugins/Experimental/UVEditor/Source/UVEditorTools/UVEditorTools.Build.cs
+++ b/Engine/Plugins/Experimental/UVEditor/Source/UVEditorTools/UVEditorTools.Build.cs
@@ -49,6 +49,7 @@ public class UVEditorTools : ModuleRules
 				"DynamicMesh",
 				"GeometryCore",
 				"ModelingComponents",
+				"ModelingOperators"
 				// ... add private dependencies that you statically link with here ...	
 			}
 			);
