fix copy-pasted convert-to-polygon enum tooltip comment

#rnx

diff --git a/Engine/Plugins/Experimental/MeshModelingToolsetExp/Source/MeshModelingToolsExp/Public/ConvertToPolygonsTool.h b/Engine/Plugins/Experimental/MeshModelingToolsetExp/Source/MeshModelingToolsExp/Public/ConvertToPolygonsTool.h
--- a/Engine/Plugins/Experimental/MeshModelingToolsetExp/Source/MeshModelingToolsExp/Public/ConvertToPolygonsTool.h
+++ b/Engine/Plugins/Experimental/MeshModelingToolsetExp/Source/MeshModelingToolsExp/Public/ConvertToPolygonsTool.h
@@ -38,7 +38,7 @@ enum class EConvertToPolygonsMode
 	FromUVIslands  UMETA(DisplayName = "From UV Islands"),
 	/** Create Polygroups based on Connected Triangles */
 	FromConnectedTris UMETA(DisplayName = "From Connected Tris"),
-	/** Create Polygroups based on Connected Triangles */
+	/** Create Polygroups centered on well-spaced sample points, approximating a surface Voronoi diagram */
 	FromFurthestPointSampling UMETA(DisplayName = "Furthest Point Sampling")
 };
 
