remove ensure on packing failure, handle the error propagating up

#rnx

diff --git a/Engine/Plugins/Experimental/MeshModelingToolsetExp/Source/ModelingOperators/Private/ParameterizationOps/RecomputeUVsOp.cpp b/Engine/Plugins/Experimental/MeshModelingToolsetExp/Source/ModelingOperators/Private/ParameterizationOps/RecomputeUVsOp.cpp
--- a/Engine/Plugins/Experimental/MeshModelingToolsetExp/Source/ModelingOperators/Private/ParameterizationOps/RecomputeUVsOp.cpp
+++ b/Engine/Plugins/Experimental/MeshModelingToolsetExp/Source/ModelingOperators/Private/ParameterizationOps/RecomputeUVsOp.cpp
@@ -222,7 +222,12 @@ bool FRecomputeUVsOp::CalculateResult_Basic(FProgressCancel* Progress)
 
 	if (bPackUVs)
 	{
-		UVEditor.QuickPack(PackingTextureResolution, PackingGutterWidth);
+		bool bPackingSuccess = UVEditor.QuickPack(PackingTextureResolution, PackingGutterWidth);
+		if (!bPackingSuccess)
+		{
+			NewResultInfo.AddError(FGeometryError(0, LOCTEXT("IslandPackingFailed", "Failed to pack UV islands")));
+			return false;
+		}
 	}
 
 	if (Progress && Progress->Cancelled())
@@ -355,7 +360,12 @@ bool FRecomputeUVsOp::CalculateResult_RegionOptimization(FProgressCancel* Progre
 
 	if (bPackUVs)
 	{
-		UVEditor.QuickPack(PackingTextureResolution, PackingGutterWidth);
+		bool bPackingSuccess = UVEditor.QuickPack(PackingTextureResolution, PackingGutterWidth);
+		if (!bPackingSuccess)
+		{
+			NewResultInfo.AddError(FGeometryError(0, LOCTEXT("IslandPackingFailed", "Failed to pack UV islands")));
+			return false;
+		}
 	}
 
 	if (Progress && Progress->Cancelled())
diff --git a/Engine/Source/Runtime/GeometryCore/Private/Parameterization/UVPacking.cpp b/Engine/Source/Runtime/GeometryCore/Private/Parameterization/UVPacking.cpp
--- a/Engine/Source/Runtime/GeometryCore/Private/Parameterization/UVPacking.cpp
+++ b/Engine/Source/Runtime/GeometryCore/Private/Parameterization/UVPacking.cpp
@@ -839,7 +839,7 @@ bool FUVPacker::StandardPack(IUVMeshView* Mesh, int NumIslands, TFunctionRef<voi
 
 
 	bool bPackingFound = Packer.FindBestPacking(AllCharts);
-	if (ensure(bPackingFound) == false)
+	if (bPackingFound == false)
 	{
 		return false;
 	}
