remove dead/unused enum in volumetomesh tool

#rnx

diff --git a/Engine/Plugins/Experimental/MeshModelingToolsetExp/Source/MeshModelingToolsExp/Public/VolumeToMeshTool.h b/Engine/Plugins/Experimental/MeshModelingToolsetExp/Source/MeshModelingToolsExp/Public/VolumeToMeshTool.h
--- a/Engine/Plugins/Experimental/MeshModelingToolsetExp/Source/MeshModelingToolsExp/Public/VolumeToMeshTool.h
+++ b/Engine/Plugins/Experimental/MeshModelingToolsetExp/Source/MeshModelingToolsExp/Public/VolumeToMeshTool.h
@@ -26,15 +26,6 @@ public:
 
 
 
-UENUM()
-enum class EVolumeToMeshMode
-{
-	/** Convert based on Angle Tolerance between Face Normals */
-	TriangulatePolygons,
-	/** Create PolyGroups based on UV Islands */
-	MinimalPolygons
-};
-
 
 
 UCLASS()
