MeshModelingToolset
* ApproximateActors - Fix forced scene capture of MRS when PackMRS is off
* MeshSceneAdapter - Use GetMaterial_Concurrent() rather than GetMaterial() for code that is executed in parallel tasks

diff --git a/Engine/Plugins/Experimental/MeshModelingToolsetExp/Source/GeometryProcessingAdapters/Private/ApproximateActorsImpl.cpp b/Engine/Plugins/Experimental/MeshModelingToolsetExp/Source/GeometryProcessingAdapters/Private/ApproximateActorsImpl.cpp
--- a/Engine/Plugins/Experimental/MeshModelingToolsetExp/Source/GeometryProcessingAdapters/Private/ApproximateActorsImpl.cpp
+++ b/Engine/Plugins/Experimental/MeshModelingToolsetExp/Source/GeometryProcessingAdapters/Private/ApproximateActorsImpl.cpp
@@ -101,9 +101,9 @@ static TUniquePtr<FSceneCapturePhotoSet> CapturePhotoSet(
 	else
 	{
 		SceneCapture->SetCaptureTypeEnabled(ERenderCaptureType::CombinedMRS, false);
-		SceneCapture->SetCaptureTypeEnabled(ERenderCaptureType::Roughness, true);
-		SceneCapture->SetCaptureTypeEnabled(ERenderCaptureType::Metallic, true);
-		SceneCapture->SetCaptureTypeEnabled(ERenderCaptureType::Specular, true);
+		SceneCapture->SetCaptureTypeEnabled(ERenderCaptureType::Roughness, bRoughness);
+		SceneCapture->SetCaptureTypeEnabled(ERenderCaptureType::Metallic, bMetallic);
+		SceneCapture->SetCaptureTypeEnabled(ERenderCaptureType::Specular, bSpecular);
 	}
 
 
diff --git a/Engine/Plugins/Experimental/MeshModelingToolsetExp/Source/ModelingComponents/Private/Scene/MeshSceneAdapter.cpp b/Engine/Plugins/Experimental/MeshModelingToolsetExp/Source/ModelingComponents/Private/Scene/MeshSceneAdapter.cpp
--- a/Engine/Plugins/Experimental/MeshModelingToolsetExp/Source/ModelingComponents/Private/Scene/MeshSceneAdapter.cpp
+++ b/Engine/Plugins/Experimental/MeshModelingToolsetExp/Source/ModelingComponents/Private/Scene/MeshSceneAdapter.cpp
@@ -492,10 +492,9 @@ public:
 			int32 MaterialIndex = GetStaticMeshMaterialIndexFromSlotName(StaticMeshIn, MaterialSlotName, true);
 
 			UMaterialInterface* MaterialInterface = GetStaticMeshMaterialFromSlotName(StaticMeshIn, MaterialSlotName, true);
-			bool bValidPolyGroup = (MaterialInterface != nullptr)
-				&& (MaterialInterface->GetMaterial() != nullptr)
-				&& (MaterialInterface->GetMaterial()->MaterialDomain == EMaterialDomain::MD_Surface);
+			const UMaterial* Material = MaterialInterface != nullptr ? MaterialInterface->GetMaterial_Concurrent() : nullptr;
 
+			bool bValidPolyGroup = Material != nullptr && Material->MaterialDomain == EMaterialDomain::MD_Surface;
 			if (bValidPolyGroup)
 			{
 				ValidPolyGroups.Add(PolygonGroupID);
