ModelingTools: improve naming of new Actors. Use AActor::GetActorNameOrLabel() instead of raw UObject name as string for new Actors if available. Use FActorLabelUtilities::SetActorLabelUnique() to set names so that they get auto-incremented unique numeric suffixes. Do all deletes before new-object creation in UConvertMeshesTool::Shutdown() so that new objects get the same name as old objects. Modify UE::Modeling::GetComponentAssetBaseName() to return Actor label if it makes sense and there is no Asset name, this improves naming in many mesh-creation Tools (possibly want this behavior even if a StaticMesh asset is available, leaving that for a future CL)
#rnx

diff --git a/Engine/Plugins/Experimental/MeshModelingToolsetExp/Source/MeshModelingToolsExp/Private/ConvertMeshesTool.cpp b/Engine/Plugins/Experimental/MeshModelingToolsetExp/Source/MeshModelingToolsExp/Private/ConvertMeshesTool.cpp
--- a/Engine/Plugins/Experimental/MeshModelingToolsetExp/Source/MeshModelingToolsExp/Private/ConvertMeshesTool.cpp
+++ b/Engine/Plugins/Experimental/MeshModelingToolsetExp/Source/MeshModelingToolsExp/Private/ConvertMeshesTool.cpp
@@ -90,7 +90,11 @@ void UConvertMeshesTool::Shutdown(EToolShutdownType ShutdownType)
 
 		TArray<AActor*> NewSelectedActors;
 		TSet<AActor*> DeleteActors;
+		TArray<FCreateMeshObjectParams> NewMeshObjects;
 
+		// Accumulate info for new mesh objects. Do not immediately create them because then
+		// the new Actors will get a unique-name incremented suffix, because the convert-from
+		// Actors still exist.
 		for (int32 k = 0; k < Targets.Num(); ++k)
 		{
 			AActor* TargetActor = UE::ToolTarget::GetTargetActor(Targets[k]);
@@ -99,7 +103,7 @@ void UConvertMeshesTool::Shutdown(EToolShutdownType ShutdownType)
 
 			FTransform SourceTransform = (FTransform)UE::ToolTarget::GetLocalToWorldTransform(Targets[k]);
 			FDynamicMesh3 SourceMesh = UE::ToolTarget::GetDynamicMeshCopy(Targets[k], true);
-			FString AssetName = TargetActor->GetName();
+			FString AssetName = TargetActor->GetActorNameOrLabel();
 			FComponentMaterialSet Materials = UE::ToolTarget::GetMaterialSet(Targets[0]);
 			const FComponentMaterialSet* TransferMaterials = (BasicProperties->bTransferMaterials) ? &Materials : nullptr;
 
@@ -113,6 +117,18 @@ void UConvertMeshesTool::Shutdown(EToolShutdownType ShutdownType)
 			}
 			NewMeshObjectParams.SetMesh(MoveTemp(SourceMesh));
 
+			NewMeshObjects.Add(MoveTemp(NewMeshObjectParams));
+		}
+
+		// delete all the existing Actors we want to get rid of
+		for (AActor* DeleteActor : DeleteActors)
+		{
+			DeleteActor->Destroy();
+		}
+
+		// spawn new mesh objects
+		for (FCreateMeshObjectParams& NewMeshObjectParams : NewMeshObjects)
+		{
 			OutputTypeProperties->ConfigureCreateMeshObjectParams(NewMeshObjectParams);
 			FCreateMeshObjectResult Result = UE::Modeling::CreateMeshObject(GetToolManager(), MoveTemp(NewMeshObjectParams));
 			if (Result.IsOK())
@@ -121,11 +137,6 @@ void UConvertMeshesTool::Shutdown(EToolShutdownType ShutdownType)
 			}
 		}
 
-		for (AActor* DeleteActor : DeleteActors)
-		{
-			DeleteActor->Destroy();
-		}
-
 		ToolSelectionUtil::SetNewActorSelection(GetToolManager(), NewSelectedActors);
 
 		GetToolManager()->EndUndoTransaction();
diff --git a/Engine/Plugins/Experimental/MeshModelingToolsetExp/Source/MeshModelingToolsExp/Private/SplitMeshesTool.cpp b/Engine/Plugins/Experimental/MeshModelingToolsetExp/Source/MeshModelingToolsExp/Private/SplitMeshesTool.cpp
--- a/Engine/Plugins/Experimental/MeshModelingToolsetExp/Source/MeshModelingToolsExp/Private/SplitMeshesTool.cpp
+++ b/Engine/Plugins/Experimental/MeshModelingToolsetExp/Source/MeshModelingToolsExp/Private/SplitMeshesTool.cpp
@@ -117,7 +117,7 @@ void USplitMeshesTool::Shutdown(EToolShutdownType ShutdownType)
 
 			FTransform3d SourceTransform = UE::ToolTarget::GetLocalToWorldTransform(Targets[ti]);
 			FDynamicMesh3 SourceMesh = UE::ToolTarget::GetDynamicMeshCopy(Targets[ti], true);
-			FString AssetName = TargetActor->GetName();
+			FString AssetName = TargetActor->GetActorNameOrLabel();
 
 			FCreateMeshObjectParams BaseMeshObjectParams;
 			BaseMeshObjectParams.TargetWorld = TargetWorld;
diff --git a/Engine/Plugins/Experimental/MeshModelingToolsetExp/Source/ModelingComponents/Private/ModelingObjectsCreationAPI.cpp b/Engine/Plugins/Experimental/MeshModelingToolsetExp/Source/ModelingComponents/Private/ModelingObjectsCreationAPI.cpp
--- a/Engine/Plugins/Experimental/MeshModelingToolsetExp/Source/ModelingComponents/Private/ModelingObjectsCreationAPI.cpp
+++ b/Engine/Plugins/Experimental/MeshModelingToolsetExp/Source/ModelingComponents/Private/ModelingObjectsCreationAPI.cpp
@@ -91,13 +91,15 @@ FString UE::Modeling::GetComponentAssetBaseName(UPrimitiveComponent* Component,
 		return TEXT("InvalidComponent");
 	}
 
-	FString ResultName = Component->GetName();
+	// default to the Actor Name (or Label in Editor) if this is the unique root component, otherwise use component name
+	FString ResultName = (Component->GetOwner()->GetRootComponent() == Component) ? 
+		Component->GetOwner()->GetActorNameOrLabel() : Component->GetName();
 	if (bRemoveAutoGeneratedSuffixes)
 	{
 		ResultName = UE::Modeling::StripGeneratedAssetSuffixFromName(ResultName);
 	}
 
-	// handle static meshcomponents
+	// If this is a static mesh component, get the asset name instead.
 	UStaticMeshComponent* StaticMeshComponent = Cast<UStaticMeshComponent>(Component);
 	if (StaticMeshComponent != nullptr)
 	{
diff --git a/Engine/Plugins/Experimental/MeshModelingToolsetExp/Source/ModelingComponentsEditorOnly/Private/EditorModelingObjectsCreationAPI.cpp b/Engine/Plugins/Experimental/MeshModelingToolsetExp/Source/ModelingComponentsEditorOnly/Private/EditorModelingObjectsCreationAPI.cpp
--- a/Engine/Plugins/Experimental/MeshModelingToolsetExp/Source/ModelingComponentsEditorOnly/Private/EditorModelingObjectsCreationAPI.cpp
+++ b/Engine/Plugins/Experimental/MeshModelingToolsetExp/Source/ModelingComponentsEditorOnly/Private/EditorModelingObjectsCreationAPI.cpp
@@ -22,6 +22,7 @@
 #include "Engine/Polys.h"
 #include "Model.h"
 #include "BSPOps.h"		// in UnrealEd
+#include "Editor/EditorEngine.h"		// for FActorLabelUtilities
 
 #include "DynamicMeshActor.h"
 #include "Components/DynamicMeshComponent.h"
@@ -174,6 +175,7 @@ FCreateMeshObjectResult UEditorModelingObjectsCreationAPI::CreateVolume(FCreateM
 	}
 
 	NewVolumeActor->SetActorTransform(CreateMeshParams.Transform);
+	FActorLabelUtilities::SetActorLabelUnique(NewVolumeActor, CreateMeshParams.BaseName);
 	NewVolumeActor->PostEditChange();
 	
 	// emit result
@@ -224,6 +226,7 @@ FCreateMeshObjectResult UEditorModelingObjectsCreationAPI::CreateDynamicMeshActo
 	}
 
 	NewActor->SetActorTransform(CreateMeshParams.Transform);
+	FActorLabelUtilities::SetActorLabelUnique(NewActor, CreateMeshParams.BaseName);
 
 	// set materials
 	TArray<UMaterialInterface*> ComponentMaterials = FilterMaterials(CreateMeshParams.Materials);
@@ -308,7 +311,7 @@ FCreateMeshObjectResult UEditorModelingObjectsCreationAPI::CreateStaticMeshAsset
 	FActorSpawnParameters SpawnInfo;
 	// @todo nothing here is specific to AStaticMeshActor...could we pass in a CDO and clone it instead of using SpawnActor?
 	AStaticMeshActor* StaticMeshActor = CreateMeshParams.TargetWorld->SpawnActor<AStaticMeshActor>(FVector::ZeroVector, Rotation, SpawnInfo);
-	StaticMeshActor->SetActorLabel(*CreateMeshParams.BaseName);
+	FActorLabelUtilities::SetActorLabelUnique(StaticMeshActor, CreateMeshParams.BaseName);
 	UStaticMeshComponent* StaticMeshComponent = StaticMeshActor->GetStaticMeshComponent();
 
 	// set the mesh
diff --git a/Engine/Plugins/Runtime/MeshModelingToolset/Source/MeshModelingTools/Private/AddPrimitiveTool.cpp b/Engine/Plugins/Runtime/MeshModelingToolset/Source/MeshModelingTools/Private/AddPrimitiveTool.cpp
--- a/Engine/Plugins/Runtime/MeshModelingToolset/Source/MeshModelingTools/Private/AddPrimitiveTool.cpp
+++ b/Engine/Plugins/Runtime/MeshModelingToolset/Source/MeshModelingTools/Private/AddPrimitiveTool.cpp
@@ -601,7 +601,7 @@ void UAddArrowPrimitiveTool::GenerateMesh(FDynamicMesh3* OutMesh) const
 UAddSpherePrimitiveTool::UAddSpherePrimitiveTool(const FObjectInitializer& ObjectInitializer)
 	: Super(ObjectInitializer.SetDefaultSubobjectClass<UProceduralSphereToolProperties>(TEXT("ShapeSettings")))
 {
-	AssetName = TEXT("ASphere");
+	AssetName = TEXT("Sphere");
 	SetToolDisplayName(LOCTEXT("SphereToolName", "Create Spheres"));
 }
 
