GeometryProcessing: add compact check in FindPoloygonsAlgorithm, normalize average normal in MeshNormals
#rnx

diff --git a/Engine/Plugins/Experimental/MeshModelingToolsetExp/Source/MeshModelingToolsExp/Private/FindPolygonsAlgorithm.cpp b/Engine/Plugins/Experimental/MeshModelingToolsetExp/Source/MeshModelingToolsExp/Private/FindPolygonsAlgorithm.cpp
--- a/Engine/Plugins/Experimental/MeshModelingToolsetExp/Source/MeshModelingToolsExp/Private/FindPolygonsAlgorithm.cpp
+++ b/Engine/Plugins/Experimental/MeshModelingToolsetExp/Source/MeshModelingToolsExp/Private/FindPolygonsAlgorithm.cpp
@@ -156,14 +156,17 @@ public:
 	/** Build a Face Dual Graph for a triangle mesh */
 	static void MakeFaceDualGraphForMesh(FDynamicMesh3* Mesh, FMeshFaceDualGraph& FaceGraph)
 	{
+		// if not true, code below needs updating
+		check(Mesh->IsCompactT());
+
 		for (int32 tid : Mesh->TriangleIndicesItr())
 		{
 			FVector3d Normal, Centroid;
 			double Area;
 			Mesh->GetTriInfo(tid, Normal, Area, Centroid);
 
-			int32 newvid = FaceGraph.AppendVertex(Centroid, Normal, Area);
-			check(newvid == tid);
+			int32 newid = FaceGraph.AppendVertex(Centroid, Normal, Area);
+			check(newid == tid);
 		}
 
 		for (int32 tid : Mesh->TriangleIndicesItr())
diff --git a/Engine/Source/Runtime/GeometryCore/Private/DynamicMesh/MeshNormals.cpp b/Engine/Source/Runtime/GeometryCore/Private/DynamicMesh/MeshNormals.cpp
--- a/Engine/Source/Runtime/GeometryCore/Private/DynamicMesh/MeshNormals.cpp
+++ b/Engine/Source/Runtime/GeometryCore/Private/DynamicMesh/MeshNormals.cpp
@@ -348,6 +348,7 @@ void FMeshNormals::SmoothVertexNormals(FDynamicMesh3& Mesh, int32 SmoothingRound
 					{
 						SmoothedNormal += (FVector3d)Mesh.GetVertexNormal(nbrvid);
 					});
+					Normalize(SmoothedNormal);
 					SmoothedNormals[vid] = Lerp((FVector3d)Mesh.GetVertexNormal(vid), SmoothedNormal, SmoothingAlpha);
 					Normalize(SmoothedNormals[vid]);
 				}
