ModelingComponents: Invalidate cached MeshDescription conversion when the mesh is updated in UDynamicMeshComponentToolTarget
ModelingTools: Attribute Editor minor fixes/improvements. Removed OptimizeForEditing button/action. Disabled lightmaps and weightmaps when target is a DynamicMeshComponent. Add support for deleting polygroup layers.
#rnx

diff --git a/Engine/Plugins/Experimental/MeshModelingToolsetExp/Source/MeshModelingToolsEditorOnly/Private/AttributeEditorTool.cpp b/Engine/Plugins/Experimental/MeshModelingToolsetExp/Source/MeshModelingToolsEditorOnly/Private/AttributeEditorTool.cpp
--- a/Engine/Plugins/Experimental/MeshModelingToolsetExp/Source/MeshModelingToolsEditorOnly/Private/AttributeEditorTool.cpp
+++ b/Engine/Plugins/Experimental/MeshModelingToolsetExp/Source/MeshModelingToolsEditorOnly/Private/AttributeEditorTool.cpp
@@ -23,6 +23,7 @@
 #include "TargetInterfaces/MeshDescriptionProvider.h"
 #include "TargetInterfaces/PrimitiveComponentBackedTarget.h"
 #include "ToolTargetManager.h"
+#include "ModelingToolTargetUtil.h"
 
 #include "ExplicitUseGeometryMathTypes.h"		// using UE::Geometry::(math types)
 using namespace UE::Geometry;
@@ -107,23 +108,25 @@ void UAttributeEditorTool::Setup()
 {
 	UInteractiveTool::Setup();
 
-	OptimizeActions = NewObject<UAttributeEditorOptimizeActions>(this);
-	OptimizeActions->Initialize(this);
-	AddToolPropertySource(OptimizeActions);
-
 	NormalsActions = NewObject<UAttributeEditorNormalsActions>(this);
 	NormalsActions->Initialize(this);
 	AddToolPropertySource(NormalsActions);
 
 	if (Targets.Num() == 1)
 	{
+		UPrimitiveComponent* TargetComponent = UE::ToolTarget::GetTargetComponent(Targets[0]);
+		bTargetIsStaticMesh = Cast<UStaticMeshComponent>(TargetComponent) != nullptr;
+
 		UVActions = NewObject<UAttributeEditorUVActions>(this);
 		UVActions->Initialize(this);
 		AddToolPropertySource(UVActions);
 
-		LightmapUVActions = NewObject<UAttributeEditorLightmapUVActions>(this);
-		LightmapUVActions->Initialize(this);
-		AddToolPropertySource(LightmapUVActions);
+		if (bTargetIsStaticMesh)
+		{
+			LightmapUVActions = NewObject<UAttributeEditorLightmapUVActions>(this);
+			LightmapUVActions->Initialize(this);
+			AddToolPropertySource(LightmapUVActions);
+		}
 
 		NewAttributeProps = NewObject<UAttributeEditorNewAttributeActions>(this);
 		NewAttributeProps->Initialize(this);
@@ -304,7 +307,7 @@ static bool RemoveAttribute(FMeshDescription* Mesh, EAttributeEditorElementType
 
 void UAttributeEditorTool::InitializeAttributeLists()
 {
-	const FMeshDescription* Mesh = TargetMeshProviderInterface(0)->GetMeshDescription();
+	const FMeshDescription* Mesh = UE::ToolTarget::GetMeshDescription(Targets[0]);
 
 	TVertexInstanceAttributesConstRef<FVector2D> InstanceUVs =
 		Mesh->VertexInstanceAttributes().GetAttributesRef<FVector2D>(MeshAttribute::VertexInstance::TextureCoordinate);
@@ -316,16 +319,19 @@ void UAttributeEditorTool::InitializeAttributeLists()
 	}
 	UVActions->UVLayer = UVActions->UVLayerNamesList[0];
 
-	UStaticMeshComponent* StaticMeshComponent = Cast<UStaticMeshComponent>(TargetComponentInterface(0)->GetOwnerComponent());
-	if (StaticMeshComponent && StaticMeshComponent->GetStaticMesh() != nullptr)
+	if (bTargetIsStaticMesh && LightmapUVActions != nullptr)
 	{
-		const UStaticMesh* StaticMesh = StaticMeshComponent->GetStaticMesh();
-		const FMeshBuildSettings& BuildSettings = StaticMesh->GetSourceModel(0).BuildSettings;
-		LightmapUVActions->bGenerateLightmapUVs = BuildSettings.bGenerateLightmapUVs;
-		LightmapUVActions->SourceUVIndex = BuildSettings.SrcLightmapIndex;
-		LightmapUVActions->DestinationUVIndex = BuildSettings.DstLightmapIndex;
+		UStaticMeshComponent* StaticMeshComponent = Cast<UStaticMeshComponent>(UE::ToolTarget::GetTargetComponent(Targets[0]));
+		if (StaticMeshComponent->GetStaticMesh() != nullptr)
+		{
+			const UStaticMesh* StaticMesh = StaticMeshComponent->GetStaticMesh();
+			const FMeshBuildSettings& BuildSettings = StaticMesh->GetSourceModel(0).BuildSettings;
+			LightmapUVActions->bGenerateLightmapUVs = BuildSettings.bGenerateLightmapUVs;
+			LightmapUVActions->SourceUVIndex = BuildSettings.SrcLightmapIndex;
+			LightmapUVActions->DestinationUVIndex = BuildSettings.DstLightmapIndex;
 
-		bHaveAutoGeneratedLightmapUVSet = (LightmapUVActions->DestinationUVIndex >= InstanceUVs.GetNumChannels());
+			bHaveAutoGeneratedLightmapUVSet = (LightmapUVActions->DestinationUVIndex >= InstanceUVs.GetNumChannels());
+		}
 	}
 
 	ExtractAttribList(Mesh, Mesh->VertexAttributes(), EAttributeEditorElementType::Vertex, this->VertexAttributes, AttributeProps->VertexAttributes);
@@ -344,16 +350,16 @@ void UAttributeEditorTool::InitializeAttributeLists()
 	//	&this->TriangleAttributes, & this->PolygonAttributes,
 	//	& this->EdgeAttributes, & this->GroupAttributes };
 	TArray<TArray<FAttributeEditorAttribInfo>*> AttribInfos = {
-		&this->VertexAttributes, &this->PolygonAttributes };
+		&this->VertexAttributes, &this->PolygonAttributes, &this->TriangleAttributes };
 
 	for (TArray<FAttributeEditorAttribInfo>* AttribInfoList : AttribInfos)
 	{
 		for (FAttributeEditorAttribInfo& AttribInfo : *AttribInfoList)
 		{
-			//if (IsReservedName(AttribInfo.Name) == false)
-			//{
+			if (FSkeletalMeshAttributes::IsReservedAttributeName(AttribInfo.Name) == false)
+			{
 				ModifyAttributeProps->AttributeNamesList.Add(AttribInfo.Name.ToString());
-			//}
+			}
 			//CopyAttributeProps->FromAttribute.Add(AttribInfo.Name);
 			//CopyAttributeProps->ToAttribute.Add(AttribInfo.Name);
 		}
@@ -382,9 +388,6 @@ void UAttributeEditorTool::OnTick(float DeltaTime)
 {
 	switch (PendingAction)
 	{
-	case EAttributeEditorToolActions::OptimizeForEditing:
-		OptimizeForEditing();
-		break;
 	case EAttributeEditorToolActions::ClearNormals:
 		ClearNormals();
 		break;
@@ -434,64 +437,11 @@ void UAttributeEditorTool::OnTick(float DeltaTime)
 
 
 
-void UAttributeEditorTool::OptimizeForEditing()
-{
-	GetToolManager()->BeginUndoTransaction(LOCTEXT("OptimizeForEditing", "Optimize For Editing"));
-
-	for (int32 ComponentIdx = 0; ComponentIdx < Targets.Num(); ComponentIdx++)
-	{
-		UStaticMeshComponent* StaticMeshComponent = Cast<UStaticMeshComponent>(TargetComponentInterface(ComponentIdx)->GetOwnerComponent());
-		if (StaticMeshComponent == nullptr || StaticMeshComponent->GetStaticMesh() == nullptr)
-		{
-			continue;
-		}
-
-		UStaticMesh* StaticMesh = StaticMeshComponent->GetStaticMesh();
-		StaticMesh->Modify();
-		StaticMesh->SetNumSourceModels(1);		// discard extra source models
-		FMeshBuildSettings& BuildSettings = StaticMesh->GetSourceModel(0).BuildSettings;
-		BuildSettings.bGenerateLightmapUVs = false;
-		BuildSettings.bBuildReversedIndexBuffer = false;
-		BuildSettings.bRemoveDegenerates = false;
-
-		// dramatically reduce distance field resolution to speed up editing
-		BuildSettings.DistanceFieldResolutionScale = 0.01;
-
-		// this will call StaticMesh->PostEditChange()...
-		TargetMeshCommitterInterface(ComponentIdx)->CommitMeshDescription([&](const IMeshDescriptionCommitter::FCommitterParams& CommitParams)
-		{
-			TEdgeAttributesRef<bool> EdgeHardnesses = CommitParams.MeshDescriptionOut->EdgeAttributes().GetAttributesRef<bool>(MeshAttribute::Edge::IsHard);
-			if (EdgeHardnesses.IsValid())
-			{
-				for (FEdgeID ElID : CommitParams.MeshDescriptionOut->Edges().GetElementIDs())
-				{
-					EdgeHardnesses[ElID] = false;
-				}
-			}
-
-			// force computation of normals/tangents if they are auto-generated
-			if (BuildSettings.bRecomputeNormals || BuildSettings.bRecomputeTangents)
-			{
-				UE::MeshDescription::InitializeAutoGeneratedAttributes(*CommitParams.MeshDescriptionOut, StaticMeshComponent, 0);
-			}
-
-			// now clear these build settings
-			BuildSettings.bUseMikkTSpace = false;
-			BuildSettings.bRecomputeNormals = false;
-			BuildSettings.bRecomputeTangents = false;
-		});
-
-	}
-	GetToolManager()->EndUndoTransaction();
-
-	// update attrib lists
-	bAttributeListsValid = false;
-}
 
 
 void UAttributeEditorTool::ClearNormals()
 {
-	GetToolManager()->BeginUndoTransaction(LOCTEXT("ClearNormalsTransactionMessage", "Clear Normals"));
+	GetToolManager()->BeginUndoTransaction(LOCTEXT("ClearNormalsTransactionMessage", "Reset Normals"));
 
 	for (int32 ComponentIdx = 0; ComponentIdx < Targets.Num(); ComponentIdx++)
 	{
@@ -682,7 +632,7 @@ void UAttributeEditorTool::AddNewAttribute(EAttributeEditorElementType ElemType,
 		return;
 	}
 
-	const FMeshDescription* CurMesh = TargetMeshProviderInterface(0)->GetMeshDescription();
+	const FMeshDescription* CurMesh = UE::ToolTarget::GetMeshDescription(Targets[0]);
 	if (HasAttribute(CurMesh, ElemType, AttributeName))
 	{
 		GetToolManager()->DisplayMessage(LOCTEXT("ErrorAddingDuplicateNameMessage", "Attribute with this name already exists"), EToolMessageLevel::UserWarning);
@@ -697,10 +647,7 @@ void UAttributeEditorTool::AddNewAttribute(EAttributeEditorElementType ElemType,
 	}
 
 	GetToolManager()->BeginUndoTransaction(LOCTEXT("NewAttributeTransactionMessage", "Add Attribute"));
-	TargetMeshCommitterInterface(0)->CommitMeshDescription([&NewMesh](const IMeshDescriptionCommitter::FCommitterParams& CommitParams)
-	{
-		*CommitParams.MeshDescriptionOut = NewMesh;
-	});
+	UE::ToolTarget::CommitMeshDescriptionUpdate(Targets[0], &NewMesh);
 	GetToolManager()->EndUndoTransaction();
 
 	bAttributeListsValid = false;
@@ -720,6 +667,12 @@ void UAttributeEditorTool::AddNewAttribute()
 
 void UAttributeEditorTool::AddNewWeightMap()
 {
+	if (bTargetIsStaticMesh == false)
+	{
+		GetToolManager()->DisplayMessage(LOCTEXT("ErrorAddingWeightMap", "Currently cannot add Weight Maps to Dynamic Mesh Components"), EToolMessageLevel::UserWarning);
+		return;
+	}
+
 	AddNewAttribute(EAttributeEditorElementType::Vertex, EAttributeEditorAttribType::Float, FName(NewAttributeProps->NewName));
 }
 
@@ -735,7 +688,7 @@ void UAttributeEditorTool::ClearAttribute()
 
 void UAttributeEditorTool::DeleteAttribute()
 {
-	const FMeshDescription* CurMesh = TargetMeshProviderInterface(0)->GetMeshDescription();
+	const FMeshDescription* CurMesh = UE::ToolTarget::GetMeshDescription(Targets[0]);
 	FName SelectedName(ModifyAttributeProps->Attribute);
 
 	// We check on the skeletal mesh attributes because it is a superset of the static mesh
@@ -746,16 +699,30 @@ void UAttributeEditorTool::DeleteAttribute()
 		return;
 	}
 
-	bool bIsVertex = HasAttribute(CurMesh, EAttributeEditorElementType::Vertex, SelectedName);
-	bool bIsPoly = HasAttribute(CurMesh, EAttributeEditorElementType::PolygonGroup, SelectedName);
-	if (bIsVertex == false && bIsPoly == false)
+	EAttributeEditorElementType ElemType = EAttributeEditorElementType::Vertex;
+	bool bIsDeletableAttribute = false;
+	if (HasAttribute(CurMesh, EAttributeEditorElementType::Vertex, SelectedName))
+	{
+		bIsDeletableAttribute = true;
+		ElemType = EAttributeEditorElementType::Vertex;
+	}
+	else if (HasAttribute(CurMesh, EAttributeEditorElementType::PolygonGroup, SelectedName))
+	{
+		bIsDeletableAttribute = true;
+		ElemType = EAttributeEditorElementType::Polygon;
+	}
+	else if (HasAttribute(CurMesh, EAttributeEditorElementType::Triangle, SelectedName))
+	{
+		bIsDeletableAttribute = true;
+		ElemType = EAttributeEditorElementType::Triangle;
+	}
+
+	if (!bIsDeletableAttribute)
 	{
 		GetToolManager()->DisplayMessage(LOCTEXT("CannotDeleteAttribError", "Cannot delete the selected attribute"), EToolMessageLevel::UserWarning);
 		return;
 	}
 
-	EAttributeEditorElementType ElemType = (bIsVertex) ? EAttributeEditorElementType::Vertex : EAttributeEditorElementType::Polygon;
-
 	FMeshDescription NewMesh = *CurMesh;
 	if (RemoveAttribute(&NewMesh, ElemType, SelectedName) == false)
 	{
@@ -764,10 +731,7 @@ void UAttributeEditorTool::DeleteAttribute()
 	}
 
 	GetToolManager()->BeginUndoTransaction(LOCTEXT("RemoveAttributeTransactionMessage", "Remove Attribute"));
-	TargetMeshCommitterInterface(0)->CommitMeshDescription([&NewMesh](const IMeshDescriptionCommitter::FCommitterParams& CommitParams)
-	{
-		*CommitParams.MeshDescriptionOut = NewMesh;
-	});
+	UE::ToolTarget::CommitMeshDescriptionUpdate(Targets[0], &NewMesh);
 	GetToolManager()->EndUndoTransaction();
 
 	bAttributeListsValid = false;
diff --git a/Engine/Plugins/Experimental/MeshModelingToolsetExp/Source/MeshModelingToolsEditorOnly/Public/AttributeEditorTool.h b/Engine/Plugins/Experimental/MeshModelingToolsetExp/Source/MeshModelingToolsEditorOnly/Public/AttributeEditorTool.h
--- a/Engine/Plugins/Experimental/MeshModelingToolsetExp/Source/MeshModelingToolsEditorOnly/Public/AttributeEditorTool.h
+++ b/Engine/Plugins/Experimental/MeshModelingToolsetExp/Source/MeshModelingToolsEditorOnly/Public/AttributeEditorTool.h
@@ -96,8 +96,6 @@ enum class EAttributeEditorToolActions
 {
 	NoAction,
 
-	OptimizeForEditing,
-
 	ClearNormals,
 	ClearAllUVs,
 	ClearSelectedUVs,
@@ -132,21 +130,6 @@ public:
 };
 
 
-UCLASS()
-class MESHMODELINGTOOLSEDITORONLY_API UAttributeEditorOptimizeActions : public UAttributeEditorActionPropertySet
-{
-	GENERATED_BODY()
-public:
-	/** Optimize this Asset for Mesh Editing by disabling expensive Static Mesh Rebuild Settings and reducing Distance Field resolution */
-	UFUNCTION(CallInEditor, Category = Optimize, meta = (DisplayPriority = 1))
-	void OptimizeForEditing()
-	{
-		PostAction(EAttributeEditorToolActions::OptimizeForEditing);
-	}
-};
-
-
-
 UCLASS()
 class MESHMODELINGTOOLSEDITORONLY_API UAttributeEditorNormalsActions : public UAttributeEditorActionPropertySet
 {
@@ -365,9 +348,6 @@ public:
 
 protected:
 
-	UPROPERTY()
-	TObjectPtr<UAttributeEditorOptimizeActions> OptimizeActions;
-
 	UPROPERTY()
 	TObjectPtr<UAttributeEditorNormalsActions> NormalsActions;
 
@@ -393,6 +373,8 @@ protected:
 protected:
 	UWorld* TargetWorld;
 
+	bool bTargetIsStaticMesh = true;
+
 	TArray<FAttributeEditorAttribInfo> VertexAttributes;
 	TArray<FAttributeEditorAttribInfo> InstanceAttributes;
 	TArray<FAttributeEditorAttribInfo> TriangleAttributes;
@@ -400,7 +382,6 @@ protected:
 	TArray<FAttributeEditorAttribInfo> EdgeAttributes;
 	TArray<FAttributeEditorAttribInfo> GroupAttributes;
 
-
 	bool bAttributeListsValid = false;
 	bool bHaveAutoGeneratedLightmapUVSet = false;
 	void InitializeAttributeLists();
@@ -408,7 +389,6 @@ protected:
 	void UpdateAutoGeneratedLightmapUVChannel(IPrimitiveComponentBackedTarget* Target, int32 NewMaxUVChannels);
 
 	EAttributeEditorToolActions PendingAction = EAttributeEditorToolActions::NoAction;
-	void OptimizeForEditing();
 	void ClearNormals();
 	void ClearUVs();
 	void DeleteSelectedUVSet();
diff --git a/Engine/Plugins/Experimental/MeshModelingToolsetExp/Source/ModelingComponentsEditorOnly/Private/ToolTargets/DynamicMeshComponentToolTarget.cpp b/Engine/Plugins/Experimental/MeshModelingToolsetExp/Source/ModelingComponentsEditorOnly/Private/ToolTargets/DynamicMeshComponentToolTarget.cpp
--- a/Engine/Plugins/Experimental/MeshModelingToolsetExp/Source/ModelingComponentsEditorOnly/Private/ToolTargets/DynamicMeshComponentToolTarget.cpp
+++ b/Engine/Plugins/Experimental/MeshModelingToolsetExp/Source/ModelingComponentsEditorOnly/Private/ToolTargets/DynamicMeshComponentToolTarget.cpp
@@ -165,6 +165,13 @@ void UDynamicMeshComponentToolTarget::CommitDynamicMeshChange(TUniquePtr<FToolCo
 	UDynamicMesh* DynamicMesh = GetDynamicMeshContainer();
 	DynamicMesh->Modify();
 	GUndo->StoreUndo(DynamicMesh, MoveTemp(Change));
+
+	// Invalidate any cached MeshDescription. 
+	if (bHaveMeshDescription)
+	{
+		ConvertedMeshDescription = nullptr;
+		bHaveMeshDescription = false;
+	}
 }
 
 
