MeshModelingTools: Add a PivotAct tool to make it easier to add an empty actor that can act as a pivot for animation.

#rnx

diff --git a/Engine/Plugins/Editor/ModelingToolsEditorMode/Source/ModelingToolsEditorMode/Private/ModelingToolsEditorMode.cpp b/Engine/Plugins/Editor/ModelingToolsEditorMode/Source/ModelingToolsEditorMode/Private/ModelingToolsEditorMode.cpp
--- a/Engine/Plugins/Editor/ModelingToolsEditorMode/Source/ModelingToolsEditorMode/Private/ModelingToolsEditorMode.cpp
+++ b/Engine/Plugins/Editor/ModelingToolsEditorMode/Source/ModelingToolsEditorMode/Private/ModelingToolsEditorMode.cpp
@@ -68,6 +68,7 @@
 #include "UVProjectionTool.h"
 #include "UVLayoutTool.h"
 #include "EditMeshMaterialsTool.h"
+#include "AddPivotActorTool.h"
 #include "EditPivotTool.h"
 #include "BakeTransformTool.h"
 #include "CombineMeshesTool.h"
@@ -430,6 +431,7 @@ void UModelingToolsEditorMode::Enter()
 	RegisterTool(ToolManagerCommands.BeginDisplaceMeshTool, TEXT("BeginDisplaceMeshTool"), NewObject<UDisplaceMeshToolBuilder>());
 	RegisterTool(ToolManagerCommands.BeginMeshSpaceDeformerTool, TEXT("BeginMeshSpaceDeformerTool"), NewObject<UMeshSpaceDeformerToolBuilder>());
 	RegisterTool(ToolManagerCommands.BeginTransformMeshesTool, TEXT("BeginTransformMeshesTool"), NewObject<UTransformMeshesToolBuilder>());
+	RegisterTool(ToolManagerCommands.BeginAddPivotActorTool, TEXT("BeginAddPivotActorTool"), NewObject<UAddPivotActorToolBuilder>());
 	RegisterTool(ToolManagerCommands.BeginEditPivotTool, TEXT("BeginEditPivotTool"), NewObject<UEditPivotToolBuilder>());
 	RegisterTool(ToolManagerCommands.BeginAlignObjectsTool, TEXT("BeginAlignObjectsTool"), NewObject<UAlignObjectsToolBuilder>());
 	RegisterTool(ToolManagerCommands.BeginTransferMeshTool, TEXT("BeginTransferMeshTool"), NewObject<UTransferMeshToolBuilder>());
diff --git a/Engine/Plugins/Editor/ModelingToolsEditorMode/Source/ModelingToolsEditorMode/Private/ModelingToolsEditorModeToolkit.cpp b/Engine/Plugins/Editor/ModelingToolsEditorMode/Source/ModelingToolsEditorMode/Private/ModelingToolsEditorModeToolkit.cpp
--- a/Engine/Plugins/Editor/ModelingToolsEditorMode/Source/ModelingToolsEditorMode/Private/ModelingToolsEditorModeToolkit.cpp
+++ b/Engine/Plugins/Editor/ModelingToolsEditorMode/Source/ModelingToolsEditorMode/Private/ModelingToolsEditorModeToolkit.cpp
@@ -570,6 +570,7 @@ void FModelingToolsEditorModeToolkit::BuildToolPalette_Experimental(FName Palett
 		ToolbarBuilder.AddToolBarButton(Commands.BeginTransformMeshesTool);
 		ToolbarBuilder.AddToolBarButton(Commands.BeginAlignObjectsTool);
 		ToolbarBuilder.AddToolBarButton(Commands.BeginEditPivotTool);
+		ToolbarBuilder.AddToolBarButton(Commands.BeginAddPivotActorTool);
 		ToolbarBuilder.AddToolBarButton(Commands.BeginBakeTransformTool);
 		ToolbarBuilder.AddToolBarButton(Commands.BeginTransferMeshTool);
 		ToolbarBuilder.AddToolBarButton(Commands.BeginConvertMeshesTool);
diff --git a/Engine/Plugins/Editor/ModelingToolsEditorMode/Source/ModelingToolsEditorMode/Private/ModelingToolsManagerActions.cpp b/Engine/Plugins/Editor/ModelingToolsEditorMode/Source/ModelingToolsEditorMode/Private/ModelingToolsManagerActions.cpp
--- a/Engine/Plugins/Editor/ModelingToolsEditorMode/Source/ModelingToolsEditorMode/Private/ModelingToolsManagerActions.cpp
+++ b/Engine/Plugins/Editor/ModelingToolsEditorMode/Source/ModelingToolsEditorMode/Private/ModelingToolsManagerActions.cpp
@@ -49,6 +49,7 @@ void FModelingToolsManagerCommands::RegisterCommands()
 	UI_COMMAND(BeginMeshSpaceDeformerTool, "Warp", "Reshape Mesh using Space Deformers", EUserInterfaceActionType::ToggleButton, FInputChord());
 	UI_COMMAND(BeginTransformMeshesTool, "XForm", "Transform selected Meshes", EUserInterfaceActionType::ToggleButton, FInputChord());
 	UI_COMMAND(BeginEditPivotTool, "Pivot", "Edit Mesh Pivots", EUserInterfaceActionType::ToggleButton, FInputChord());
+	UI_COMMAND(BeginAddPivotActorTool, "PivotAct", "Add actor to act as a pivot for child component(s)", EUserInterfaceActionType::ToggleButton, FInputChord());
 	UI_COMMAND(BeginBakeTransformTool, "BakeRS", "Bake Scale/Rotation into Mesh Asset", EUserInterfaceActionType::ToggleButton, FInputChord());
 	UI_COMMAND(BeginCombineMeshesTool, "Append", "Combine Selection into new StaticMesh Asset", EUserInterfaceActionType::ToggleButton, FInputChord());
 	UI_COMMAND(BeginDuplicateMeshesTool, "Dupe", "Duplicate Selection into new StaticMesh Asset", EUserInterfaceActionType::ToggleButton, FInputChord());
diff --git a/Engine/Plugins/Editor/ModelingToolsEditorMode/Source/ModelingToolsEditorMode/Public/ModelingToolsManagerActions.h b/Engine/Plugins/Editor/ModelingToolsEditorMode/Source/ModelingToolsEditorMode/Public/ModelingToolsManagerActions.h
--- a/Engine/Plugins/Editor/ModelingToolsEditorMode/Source/ModelingToolsEditorMode/Public/ModelingToolsManagerActions.h
+++ b/Engine/Plugins/Editor/ModelingToolsEditorMode/Source/ModelingToolsEditorMode/Public/ModelingToolsManagerActions.h
@@ -43,6 +43,7 @@ public:
 	TSharedPtr<FUICommandInfo> BeginMeshSpaceDeformerTool;
 	TSharedPtr<FUICommandInfo> BeginTransformMeshesTool;
 	TSharedPtr<FUICommandInfo> BeginEditPivotTool;
+	TSharedPtr<FUICommandInfo> BeginAddPivotActorTool;
 	TSharedPtr<FUICommandInfo> BeginBakeTransformTool;
 	TSharedPtr<FUICommandInfo> BeginCombineMeshesTool;
 	TSharedPtr<FUICommandInfo> BeginDuplicateMeshesTool;
diff --git a/Engine/Plugins/Experimental/MeshModelingToolsetExp/Source/MeshModelingToolsEditorOnly/Private/AddPivotActorTool.cpp b/Engine/Plugins/Experimental/MeshModelingToolsetExp/Source/MeshModelingToolsEditorOnly/Private/AddPivotActorTool.cpp
new file mode 100644
--- /dev/null
+++ b/Engine/Plugins/Experimental/MeshModelingToolsetExp/Source/MeshModelingToolsEditorOnly/Private/AddPivotActorTool.cpp
@@ -0,0 +1,154 @@
+// Copyright Epic Games, Inc. All Rights Reserved.
+
+#include "AddPivotActorTool.h"
+
+#include "ActorFactories/ActorFactoryEmptyActor.h"
+#include "BaseGizmos/TransformGizmo.h"
+#include "BaseGizmos/TransformGizmoUtil.h"
+#include "BaseGizmos/TransformProxy.h"
+#include "Mechanics/DragAlignmentMechanic.h"
+#include "ModelingToolTargetUtil.h"
+#include "Selection/ToolSelectionUtil.h"
+#include "TargetInterfaces/PrimitiveComponentBackedTarget.h"
+#include "ToolTargetManager.h"
+
+#define LOCTEXT_NAMESPACE "UAddPivotActorTool"
+
+using namespace UE::Geometry;
+
+const FToolTargetTypeRequirements& UAddPivotActorToolBuilder::GetTargetRequirements() const
+{
+	static FToolTargetTypeRequirements TypeRequirements({
+		UPrimitiveComponentBackedTarget::StaticClass()
+		});
+	return TypeRequirements;
+}
+
+bool UAddPivotActorToolBuilder::CanBuildTool(const FToolBuilderState& SceneState) const
+{
+	// There are some limitations for when we can use this tool. 
+	// 1. We operate on the actor, not component level.
+	//   TODO: Is there a good way to operate on a sub-actor level? Or should we be checking that
+	//   we've selected all the components of each actor?
+	// 2. If there are multiple actors selected, they need to have a common parent (or no parent),
+	//   because otherwise we will be breaking up the user's hierarchy when we nest everything under
+	//   the empty actor.
+	// 3. All of the actors need to be marked as movable because non-movable items can't be nested
+	//   under a movable one.
+
+	TSet<AActor*> ParentActors;
+	bool bAllActorsMovable = true;
+
+	SceneState.TargetManager->EnumerateSelectedAndTargetableComponents(SceneState, GetTargetRequirements(), 
+		[&bAllActorsMovable, &ParentActors](UActorComponent* Component) {
+			AActor* Actor = Component->GetOwner();
+			bAllActorsMovable = bAllActorsMovable && Actor->IsRootComponentMovable();
+			if (bAllActorsMovable)
+			{
+				ParentActors.Add(Actor->GetAttachParentActor());
+			}
+		});
+
+	return bAllActorsMovable && ParentActors.Num() == 1;
+}
+
+UInteractiveTool* UAddPivotActorToolBuilder::BuildTool(const FToolBuilderState& SceneState) const
+{
+	UAddPivotActorTool* NewTool = NewObject<UAddPivotActorTool>(SceneState.ToolManager);
+	NewTool->SetTargets(SceneState.TargetManager->BuildAllSelectedTargetable(
+		SceneState, GetTargetRequirements()));
+	NewTool->SetWorld(SceneState.World);
+
+	return NewTool;
+}
+
+void UAddPivotActorTool::Setup()
+{
+	GetToolManager()->DisplayMessage(LOCTEXT("OnStartTool", 
+		"Adds an empty actor as the parent of the selected actors. Use gizmo to choose where/how "
+		"the empty actor is placed. Hold Ctrl to snap to items in scene."), 
+		EToolMessageLevel::UserNotification);
+
+	// Figure out where to start the gizmo. The location will be the average,
+	// and the rotation will either be identity, or the target's if there is
+	// only one.
+	FVector3d StartTranslation = FVector3d::Zero();
+	for (TObjectPtr<UToolTarget> Target : Targets)
+	{
+		StartTranslation += UE::ToolTarget::GetLocalToWorldTransform(Target).GetTranslation();
+	}
+	StartTranslation /= Targets.Num();
+
+	FTransform StartTransform(StartTranslation);
+	if (Targets.Num() == 1)
+	{
+		StartTransform.SetRotation(FQuat(
+			UE::ToolTarget::GetLocalToWorldTransform(Targets[0]).GetRotation()));
+	}
+
+	// Set up the gizmo.
+	TransformProxy = NewObject<UTransformProxy>(this);
+	TransformProxy->SetTransform(StartTransform);
+	TransformGizmo = UE::TransformGizmoUtil::CreateCustomTransformGizmo(
+		GetToolManager()->GetPairedGizmoManager(),
+		ETransformGizmoSubElements::StandardTranslateRotate, this);
+	TransformGizmo->SetActiveTarget(TransformProxy, GetToolManager());
+
+	DragAlignmentMechanic = NewObject<UDragAlignmentMechanic>(this);
+	DragAlignmentMechanic->Setup(this);
+	DragAlignmentMechanic->AddToGizmo(TransformGizmo);
+}
+
+void UAddPivotActorTool::Shutdown(EToolShutdownType ShutdownType)
+{
+	if (ShutdownType == EToolShutdownType::Accept)
+	{
+		GetToolManager()->BeginUndoTransaction(LOCTEXT("AddPivotActorToolTransactionName", "Add Empty Actor"));
+
+		// Create an empty actor at the location of the gizmo. The way we do it here, using this factory, is
+		// editor-only.
+		UActorFactoryEmptyActor* EmptyActorFactory = NewObject<UActorFactoryEmptyActor>();
+		FAssetData AssetData(EmptyActorFactory->GetDefaultActorClass(FAssetData()));
+		FActorSpawnParameters SpawnParams;
+		SpawnParams.Name = TEXT("PivotActor");
+		SpawnParams.NameMode = FActorSpawnParameters::ESpawnActorNameMode::Requested;
+		AActor* NewActor = EmptyActorFactory->CreateActor(AssetData.GetAsset(), 
+			TargetWorld->GetCurrentLevel(), 
+			TransformProxy->GetTransform(), 
+			SpawnParams);
+
+		// Grab the first selected target. It will have the same parent as the other ones, so
+		// we'll use it to figure out the empty actor's parent.
+		AActor* TargetActor = UE::ToolTarget::GetTargetActor(Targets[0]);
+
+		// This is also editor-only: it's the label that shows up in the hierarchy
+		NewActor->SetActorLabel(Targets.Num() == 1 ? TargetActor->GetActorLabel() + TEXT("_Pivot")
+			: TEXT("Pivot"));
+
+		// Attach the empty actor in the correct place in the hierarchy
+		if (AActor* ParentActor = TargetActor->GetAttachParentActor())
+		{
+			NewActor->AttachToActor(ParentActor, FAttachmentTransformRules(EAttachmentRule::KeepWorld, false));
+		}
+		for (TObjectPtr<UToolTarget> Target : Targets)
+		{
+			TargetActor = UE::ToolTarget::GetTargetActor(Target); 
+			TargetActor->AttachToActor(NewActor, FAttachmentTransformRules(EAttachmentRule::KeepWorld, false));
+		}
+		ToolSelectionUtil::SetNewActorSelection(GetToolManager(), NewActor);
+
+		GetToolManager()->EndUndoTransaction();
+	}
+
+	GetToolManager()->GetPairedGizmoManager()->DestroyAllGizmosByOwner(this);
+	TransformProxy = nullptr;
+	TransformGizmo = nullptr;
+	DragAlignmentMechanic->Shutdown();
+}
+
+void UAddPivotActorTool::Render(IToolsContextRenderAPI* RenderAPI)
+{
+	DragAlignmentMechanic->Render(RenderAPI);
+}
+
+#undef LOCTEXT_NAMESPACE
diff --git a/Engine/Plugins/Experimental/MeshModelingToolsetExp/Source/MeshModelingToolsEditorOnly/Public/AddPivotActorTool.h b/Engine/Plugins/Experimental/MeshModelingToolsetExp/Source/MeshModelingToolsEditorOnly/Public/AddPivotActorTool.h
new file mode 100644
--- /dev/null
+++ b/Engine/Plugins/Experimental/MeshModelingToolsetExp/Source/MeshModelingToolsEditorOnly/Public/AddPivotActorTool.h
@@ -0,0 +1,62 @@
+// Copyright Epic Games, Inc. All Rights Reserved.
+
+#pragma once
+
+#include "CoreMinimal.h"
+
+#include "InteractiveToolBuilder.h"
+#include "MultiSelectionTool.h"
+
+#include "AddPivotActorTool.generated.h"
+
+class UDragAlignmentMechanic;
+class UTransformGizmo;
+class UTransformProxy;
+
+UCLASS()
+class MESHMODELINGTOOLSEDITORONLY_API UAddPivotActorToolBuilder : public UInteractiveToolBuilder
+{
+	GENERATED_BODY()
+public:
+	virtual bool CanBuildTool(const FToolBuilderState& SceneState) const override;
+	virtual UInteractiveTool* BuildTool(const FToolBuilderState& SceneState) const override;
+
+protected:
+	virtual const FToolTargetTypeRequirements& GetTargetRequirements() const override;
+};
+
+/** 
+ * Given selected actors, creates an empty actor as the parent of those actors, at a location
+ * specified using the gizmo. This is useful for creating a permanent alternate pivot to use in
+ * animation.
+ */
+UCLASS()
+class MESHMODELINGTOOLSEDITORONLY_API UAddPivotActorTool : public UMultiSelectionTool
+{
+	GENERATED_BODY()
+public:
+
+	virtual void SetWorld(UWorld* World) { TargetWorld = World; }
+
+	virtual void Setup() override;
+	virtual void Shutdown(EToolShutdownType ShutdownType) override;
+
+	virtual void Render(IToolsContextRenderAPI* RenderAPI) override;
+
+	virtual bool HasAccept() const override { return true; }
+	virtual bool HasCancel() const override { return true; }
+	// Uses the base class CanAccept
+
+protected:
+
+	UPROPERTY()
+	TObjectPtr<UDragAlignmentMechanic> DragAlignmentMechanic = nullptr;
+
+	UPROPERTY()
+	TObjectPtr<UTransformGizmo> TransformGizmo = nullptr;
+
+	UPROPERTY()
+	TObjectPtr<UTransformProxy> TransformProxy = nullptr;
+
+	TObjectPtr<UWorld> TargetWorld;
+};
