MeshModelingTools: Minor clang build fix.

#rnx

diff --git a/Engine/Plugins/Experimental/MeshModelingToolsetExp/Source/MeshModelingToolsEditorOnly/Private/AddPivotActorTool.cpp b/Engine/Plugins/Experimental/MeshModelingToolsetExp/Source/MeshModelingToolsEditorOnly/Private/AddPivotActorTool.cpp
--- a/Engine/Plugins/Experimental/MeshModelingToolsetExp/Source/MeshModelingToolsEditorOnly/Private/AddPivotActorTool.cpp
+++ b/Engine/Plugins/Experimental/MeshModelingToolsetExp/Source/MeshModelingToolsEditorOnly/Private/AddPivotActorTool.cpp
@@ -18,9 +18,9 @@ using namespace UE::Geometry;
 
 const FToolTargetTypeRequirements& UAddPivotActorToolBuilder::GetTargetRequirements() const
 {
-	static FToolTargetTypeRequirements TypeRequirements({
+	static FToolTargetTypeRequirements TypeRequirements(
 		UPrimitiveComponentBackedTarget::StaticClass()
-		});
+		);
 	return TypeRequirements;
 }
 
