GeometryProcessing: remove most calls to FDynamicMesh3::GetCachedBounds()
#rnx

diff --git a/Engine/Plugins/Experimental/MeshModelingToolsetExp/Source/MeshModelingToolsEditorOnly/Private/BspConversionTool.cpp b/Engine/Plugins/Experimental/MeshModelingToolsetExp/Source/MeshModelingToolsEditorOnly/Private/BspConversionTool.cpp
--- a/Engine/Plugins/Experimental/MeshModelingToolsetExp/Source/MeshModelingToolsEditorOnly/Private/BspConversionTool.cpp
+++ b/Engine/Plugins/Experimental/MeshModelingToolsetExp/Source/MeshModelingToolsEditorOnly/Private/BspConversionTool.cpp
@@ -354,7 +354,7 @@ bool UBspConversionTool::CombineThenConvert(FText*)
 	}
 	else
 	{
-		NewPivot = OutputMesh.GetCachedBounds().Center();
+		NewPivot = OutputMesh.GetBounds().Center();
 	}
 
 	MeshTransforms::Translate(OutputMesh, -NewPivot);
diff --git a/Engine/Plugins/Experimental/MeshModelingToolsetExp/Source/MeshModelingToolsEditorOnly/Private/Hair/GroomToMeshTool.cpp b/Engine/Plugins/Experimental/MeshModelingToolsetExp/Source/MeshModelingToolsEditorOnly/Private/Hair/GroomToMeshTool.cpp
--- a/Engine/Plugins/Experimental/MeshModelingToolsetExp/Source/MeshModelingToolsEditorOnly/Private/Hair/GroomToMeshTool.cpp
+++ b/Engine/Plugins/Experimental/MeshModelingToolsetExp/Source/MeshModelingToolsEditorOnly/Private/Hair/GroomToMeshTool.cpp
@@ -537,7 +537,7 @@ TSharedPtr<FDynamicMesh3, ESPMode::ThreadSafe> UGroomToMeshTool::UpdateMorpholog
 		ImplicitMorphology.Source = CachedMorphologyResult.Get();
 		FDynamicMeshAABBTree3 Spatial(ImplicitMorphology.Source, true);
 		ImplicitMorphology.SourceSpatial = &Spatial;
-		ImplicitMorphology.SetCellSizesAndDistance(CachedMorphologyResult->GetCachedBounds(), NewSettings.CloseDist, NewSettings.VoxelCount, NewSettings.VoxelCount);
+		ImplicitMorphology.SetCellSizesAndDistance(CachedMorphologyResult->GetBounds(true), NewSettings.CloseDist, NewSettings.VoxelCount, NewSettings.VoxelCount);
 
 		CachedMorphologyResult = MakeShared<FDynamicMesh3, ESPMode::ThreadSafe>(&ImplicitMorphology.Generate());
 		if (CachedMorphologyResult->TriangleCount() == 0)
@@ -554,7 +554,7 @@ TSharedPtr<FDynamicMesh3, ESPMode::ThreadSafe> UGroomToMeshTool::UpdateMorpholog
 		ImplicitMorphology.Source = CachedMorphologyResult.Get();
 		FDynamicMeshAABBTree3 Spatial(ImplicitMorphology.Source, true);
 		ImplicitMorphology.SourceSpatial = &Spatial;
-		ImplicitMorphology.SetCellSizesAndDistance(CachedMorphologyResult->GetCachedBounds(), NewSettings.OpenDist, NewSettings.VoxelCount, NewSettings.VoxelCount);
+		ImplicitMorphology.SetCellSizesAndDistance(CachedMorphologyResult->GetBounds(true), NewSettings.OpenDist, NewSettings.VoxelCount, NewSettings.VoxelCount);
 
 		CachedMorphologyResult = MakeShared<FDynamicMesh3, ESPMode::ThreadSafe>(&ImplicitMorphology.Generate());
 		if (CachedMorphologyResult->TriangleCount() == 0)
@@ -903,7 +903,7 @@ TSharedPtr<FDynamicMesh3, ESPMode::ThreadSafe> UGroomToMeshTool::UpdateSmoothing
 		// compute area of the input mesh and compute normalization scaling factor
 		FVector2d VolArea = TMeshQueries<FDynamicMesh3>::GetVolumeArea(*EditMesh);
 		double UnitScalingMeasure = FMathd::Max(0.01, FMathd::Sqrt(VolArea.Y / 6.0));  // 6.0 is a bit arbitrary here...surface area of unit box
-		FAxisAlignedBox3d Bounds = InputMesh->GetCachedBounds();
+		FAxisAlignedBox3d Bounds = InputMesh->GetBounds(true);
 		FVector3d SrcTranslate = Bounds.Center();
 		MeshTransforms::Translate(*EditMesh, -SrcTranslate);
 		double SrcScale = UnitScalingMeasure;
@@ -1153,7 +1153,7 @@ TSharedPtr<FDynamicMesh3, ESPMode::ThreadSafe> UGroomToMeshTool::UpdateUVs_ExpMa
 	
 	FDynamicMeshAABBTree3 BVTree(EditMesh, true);
 
-	FAxisAlignedBox3d Bounds = EditMesh->GetCachedBounds();
+	FAxisAlignedBox3d Bounds = EditMesh->GetBounds(true);
 	FVector3d Center = Bounds.Center();
 	FVector3d Extents = Bounds.Extents();
 
diff --git a/Engine/Plugins/Experimental/MeshModelingToolsetExp/Source/MeshModelingToolsExp/Private/AddPatchTool.cpp b/Engine/Plugins/Experimental/MeshModelingToolsetExp/Source/MeshModelingToolsExp/Private/AddPatchTool.cpp
--- a/Engine/Plugins/Experimental/MeshModelingToolsetExp/Source/MeshModelingToolsExp/Private/AddPatchTool.cpp
+++ b/Engine/Plugins/Experimental/MeshModelingToolsetExp/Source/MeshModelingToolsExp/Private/AddPatchTool.cpp
@@ -276,7 +276,7 @@ void UAddPatchTool::GeneratePreviewBaseMesh()
 	}
 
 	// recenter mesh
-	FAxisAlignedBox3d Bounds = BaseMesh->GetCachedBounds();
+	FAxisAlignedBox3d Bounds = BaseMesh->GetBounds(true);
 	FVector3d TargetOrigin = Bounds.Center();
 	for (int vid : BaseMesh->VertexIndicesItr())
 	{
diff --git a/Engine/Plugins/Experimental/MeshModelingToolsetExp/Source/MeshModelingToolsExp/Private/DynamicMeshSculptTool.cpp b/Engine/Plugins/Experimental/MeshModelingToolsetExp/Source/MeshModelingToolsExp/Private/DynamicMeshSculptTool.cpp
--- a/Engine/Plugins/Experimental/MeshModelingToolsetExp/Source/MeshModelingToolsExp/Private/DynamicMeshSculptTool.cpp
+++ b/Engine/Plugins/Experimental/MeshModelingToolsetExp/Source/MeshModelingToolsExp/Private/DynamicMeshSculptTool.cpp
@@ -143,7 +143,7 @@ void UDynamicMeshSculptTool::Setup()
 	PendingTargetUpdate.Wait();
 
 	// initialize brush radius range interval, brush properties
-	double MaxDimension = DynamicMeshComponent->GetMesh()->GetCachedBounds().MaxDim();
+	double MaxDimension = DynamicMeshComponent->GetMesh()->GetBounds(true).MaxDim();
 	BrushRelativeSizeRange = FInterval1d(MaxDimension*0.01, MaxDimension);
 	BrushProperties = NewObject<USculptBrushProperties>(this);
 	BrushProperties->bShowStrength = false;
diff --git a/Engine/Plugins/Experimental/MeshModelingToolsetExp/Source/MeshModelingToolsExp/Private/MeshGroupPaintTool.cpp b/Engine/Plugins/Experimental/MeshModelingToolsetExp/Source/MeshModelingToolsExp/Private/MeshGroupPaintTool.cpp
--- a/Engine/Plugins/Experimental/MeshModelingToolsetExp/Source/MeshModelingToolsExp/Private/MeshGroupPaintTool.cpp
+++ b/Engine/Plugins/Experimental/MeshModelingToolsetExp/Source/MeshModelingToolsExp/Private/MeshGroupPaintTool.cpp
@@ -101,7 +101,7 @@ void UMeshGroupPaintTool::Setup()
 
 	FDynamicMesh3* Mesh = GetSculptMesh();
 	Mesh->EnableVertexColors(FVector3f::One());
-	FAxisAlignedBox3d Bounds = Mesh->GetCachedBounds();
+	FAxisAlignedBox3d Bounds = Mesh->GetBounds(true);
 
 	TFuture<void> PrecomputeFuture = Async(GroupPaintToolAsyncExecTarget, [&]()
 	{
diff --git a/Engine/Plugins/Experimental/MeshModelingToolsetExp/Source/MeshModelingToolsExp/Private/Sculpting/MeshSculptToolBase.cpp b/Engine/Plugins/Experimental/MeshModelingToolsetExp/Source/MeshModelingToolsExp/Private/Sculpting/MeshSculptToolBase.cpp
--- a/Engine/Plugins/Experimental/MeshModelingToolsetExp/Source/MeshModelingToolsExp/Private/Sculpting/MeshSculptToolBase.cpp
+++ b/Engine/Plugins/Experimental/MeshModelingToolsetExp/Source/MeshModelingToolsExp/Private/Sculpting/MeshSculptToolBase.cpp
@@ -243,7 +243,7 @@ void UMeshSculptToolBase::InitializeSculptMeshComponent(UBaseDynamicMeshComponen
 
 	// initialize from LOD-0 MeshDescription
 	Component->SetMesh(UE::ToolTarget::GetDynamicMeshCopy(Target));
-	double MaxDimension = Component->GetMesh()->GetCachedBounds().MaxDim();
+	double MaxDimension = Component->GetMesh()->GetBounds(true).MaxDim();
 
 	// bake rotation and scaling into mesh because handling these inside sculpting is a mess
 	// Note: this transform does not include translation ( so only the 3x3 transform)
diff --git a/Engine/Plugins/Experimental/MeshModelingToolsetExp/Source/MeshModelingToolsExp/Private/SeamSculptTool.cpp b/Engine/Plugins/Experimental/MeshModelingToolsetExp/Source/MeshModelingToolsExp/Private/SeamSculptTool.cpp
--- a/Engine/Plugins/Experimental/MeshModelingToolsetExp/Source/MeshModelingToolsExp/Private/SeamSculptTool.cpp
+++ b/Engine/Plugins/Experimental/MeshModelingToolsetExp/Source/MeshModelingToolsExp/Private/SeamSculptTool.cpp
@@ -67,7 +67,7 @@ void USeamSculptTool::Setup()
 	MeshTransform = UE::Geometry::FTransform3d(TargetComponent->GetWorldTransform());
 	InputMesh = MakeShared<FDynamicMesh3, ESPMode::ThreadSafe>(*PreviewMesh->GetMesh());
 	FMeshNormals::QuickComputeVertexNormals(*InputMesh);
-	NormalOffset = InputMesh->GetCachedBounds().MinDim() * 0.001;
+	NormalOffset = InputMesh->GetBounds(true).MinDim() * 0.001;
 
 	// disable shadows
 	//PreviewMesh->GetRootComponent()->bCastDynamicShadow = false;
diff --git a/Engine/Plugins/Experimental/MeshModelingToolsetExp/Source/MeshModelingToolsExp/Private/SkinWeightsBindingTool.cpp b/Engine/Plugins/Experimental/MeshModelingToolsetExp/Source/MeshModelingToolsExp/Private/SkinWeightsBindingTool.cpp
--- a/Engine/Plugins/Experimental/MeshModelingToolsetExp/Source/MeshModelingToolsExp/Private/SkinWeightsBindingTool.cpp
+++ b/Engine/Plugins/Experimental/MeshModelingToolsetExp/Source/MeshModelingToolsExp/Private/SkinWeightsBindingTool.cpp
@@ -504,7 +504,7 @@ private:
 		const int32 NumVertices = InMesh.VertexCount();
 
 		// Use the diagonal size of the bbox to make the bone distance falloff scale invariant.
-		const float DiagBounds = InMesh.GetCachedBounds().DiagonalLength();
+		const float DiagBounds = InMesh.GetBounds(true).DiagonalLength();
 	
 		const FTransformHierarchyQuery Skeleton(TransformHierarchy);
 
@@ -569,7 +569,7 @@ private:
 		const int32 NumVertices = InMesh.VertexCount();
 
 		// Use the diagonal size of the bbox to make the bone distance falloff scale invariant.
-		const float DiagBounds = InMesh.GetCachedBounds().DiagonalLength();
+		const float DiagBounds = InMesh.GetBounds(true).DiagonalLength();
 		
 		const FTransformHierarchyQuery Skeleton(TransformHierarchy);
 
diff --git a/Engine/Plugins/Experimental/MeshModelingToolsetExp/Source/ModelingComponents/Private/BaseTools/BaseCreateFromSelectedTool.cpp b/Engine/Plugins/Experimental/MeshModelingToolsetExp/Source/ModelingComponents/Private/BaseTools/BaseCreateFromSelectedTool.cpp
--- a/Engine/Plugins/Experimental/MeshModelingToolsetExp/Source/ModelingComponents/Private/BaseTools/BaseCreateFromSelectedTool.cpp
+++ b/Engine/Plugins/Experimental/MeshModelingToolsetExp/Source/ModelingComponents/Private/BaseTools/BaseCreateFromSelectedTool.cpp
@@ -221,7 +221,7 @@ void UBaseCreateFromSelectedTool::GenerateAsset(const FDynamicMeshOpResult& OpRe
 	}
 	else // in the multi-selection case, center the pivot for the combined result
 	{
-		FVector3d Center = OpResult.Mesh->GetCachedBounds().Center();
+		FVector3d Center = OpResult.Mesh->GetBounds().Center();
 		double Rescale = OpResult.Transform.GetScale().X;
 		FTransform3d LocalTransform(-Center * Rescale);
 		LocalTransform.SetScale(FVector3d(Rescale, Rescale, Rescale));
diff --git a/Engine/Plugins/Experimental/MeshModelingToolsetExp/Source/ModelingComponents/Private/BaseTools/BaseMeshProcessingTool.cpp b/Engine/Plugins/Experimental/MeshModelingToolsetExp/Source/ModelingComponents/Private/BaseTools/BaseMeshProcessingTool.cpp
--- a/Engine/Plugins/Experimental/MeshModelingToolsetExp/Source/ModelingComponents/Private/BaseTools/BaseMeshProcessingTool.cpp
+++ b/Engine/Plugins/Experimental/MeshModelingToolsetExp/Source/ModelingComponents/Private/BaseTools/BaseMeshProcessingTool.cpp
@@ -91,7 +91,7 @@ void UBaseMeshProcessingTool::Setup()
 		double UnitScalingMeasure = FMathd::Max(0.01, FMathd::Sqrt(VolArea.Y / 6.0));  // 6.0 is a bit arbitrary here...surface area of unit box
 
 		// translate to origin and then apply inverse of scale
-		FAxisAlignedBox3d Bounds = InitialMesh.GetCachedBounds();
+		FAxisAlignedBox3d Bounds = InitialMesh.GetBounds();
 		SrcTranslate = Bounds.Center();
 		MeshTransforms::Translate(InitialMesh, -SrcTranslate);
 		SrcScale = UnitScalingMeasure;
diff --git a/Engine/Plugins/Experimental/MeshModelingToolsetExp/Source/ModelingComponents/Private/Components/OctreeDynamicMeshComponent.cpp b/Engine/Plugins/Experimental/MeshModelingToolsetExp/Source/ModelingComponents/Private/Components/OctreeDynamicMeshComponent.cpp
--- a/Engine/Plugins/Experimental/MeshModelingToolsetExp/Source/ModelingComponents/Private/Components/OctreeDynamicMeshComponent.cpp
+++ b/Engine/Plugins/Experimental/MeshModelingToolsetExp/Source/ModelingComponents/Private/Components/OctreeDynamicMeshComponent.cpp
@@ -58,7 +58,7 @@ void UOctreeDynamicMeshComponent::SetMesh(UE::Geometry::FDynamicMesh3&& MoveMesh
 {
 	MeshObject->SetMesh(MoveTemp(MoveMesh));
 
-	FAxisAlignedBox3d MeshBounds = GetMesh()->GetCachedBounds();
+	FAxisAlignedBox3d MeshBounds = GetMesh()->GetBounds(true);
 	Octree = MakeUnique<FDynamicMeshOctree3>();
 	Octree->RootDimension = MeshBounds.MaxDim() * 0.25;
 
@@ -96,7 +96,7 @@ void UOctreeDynamicMeshComponent::ApplyTransform(const UE::Geometry::FTransform3
 	}
 	else
 	{
-		FAxisAlignedBox3d MeshBounds = GetMesh()->GetCachedBounds();
+		FAxisAlignedBox3d MeshBounds = GetMesh()->GetBounds(true);
 		Octree = MakeUnique<FDynamicMeshOctree3>();
 		Octree->RootDimension = MeshBounds.MaxDim() * 0.25;
 		Octree->Initialize(GetMesh());
diff --git a/Engine/Plugins/Experimental/MeshModelingToolsetExp/Source/ModelingOperators/Private/CompositionOps/CurveSweepOp.cpp b/Engine/Plugins/Experimental/MeshModelingToolsetExp/Source/ModelingOperators/Private/CompositionOps/CurveSweepOp.cpp
--- a/Engine/Plugins/Experimental/MeshModelingToolsetExp/Source/ModelingOperators/Private/CompositionOps/CurveSweepOp.cpp
+++ b/Engine/Plugins/Experimental/MeshModelingToolsetExp/Source/ModelingOperators/Private/CompositionOps/CurveSweepOp.cpp
@@ -203,7 +203,7 @@ void FCurveSweepOp::CalculateResult(FProgressCancel* Progress)
 	}
 
 	// Set the pivot of the mesh to the bounding box center
-	FVector3d PivotLocation = ResultMesh->GetCachedBounds().Center();
+	FVector3d PivotLocation = ResultMesh->GetBounds(true).Center();
 	MeshTransforms::Translate(*ResultMesh, -PivotLocation);
 	ResultTransform = FTransform3d::Identity();
 	ResultTransform.SetTranslation(PivotLocation);
diff --git a/Engine/Plugins/Experimental/MeshModelingToolsetExp/Source/ModelingOperators/Private/CompositionOps/VoxelBlendMeshesOp.cpp b/Engine/Plugins/Experimental/MeshModelingToolsetExp/Source/ModelingOperators/Private/CompositionOps/VoxelBlendMeshesOp.cpp
--- a/Engine/Plugins/Experimental/MeshModelingToolsetExp/Source/ModelingOperators/Private/CompositionOps/VoxelBlendMeshesOp.cpp
+++ b/Engine/Plugins/Experimental/MeshModelingToolsetExp/Source/ModelingOperators/Private/CompositionOps/VoxelBlendMeshesOp.cpp
@@ -57,7 +57,7 @@ void FVoxelBlendMeshesOp::CalculateResult(FProgressCancel* Progress)
 			if (ThickenShells > 0 && !TransformedMeshes[MeshIdx].IsClosed())
 			{
 				// thickness should be at least a cell wide so we don't end up deleting a bunch of the input surface
-				double CellSize = TransformedMeshes[MeshIdx].GetCachedBounds().MaxDim() / InputVoxelCount;
+				double CellSize = TransformedMeshes[MeshIdx].GetBounds(true).MaxDim() / InputVoxelCount;
 				double SafeThickness = FMathd::Max(CellSize * 2, ThickenShells);
 
 				FMeshNormals::QuickComputeVertexNormals(TransformedMeshes[MeshIdx]);
@@ -85,7 +85,7 @@ void FVoxelBlendMeshesOp::CalculateResult(FProgressCancel* Progress)
 			continue;
 		}
 		ImplicitBlend.Sources.Add(&TransformedMeshes[MeshIdx]);
-		FAxisAlignedBox3d& SourceBounds = ImplicitBlend.SourceBounds.Add_GetRef(TransformedMeshes[MeshIdx].GetCachedBounds());
+		FAxisAlignedBox3d& SourceBounds = ImplicitBlend.SourceBounds.Add_GetRef(TransformedMeshes[MeshIdx].GetBounds(true));
 		CombinedBounds.Contain(SourceBounds);
 	}
 
diff --git a/Engine/Plugins/Experimental/MeshModelingToolsetExp/Source/ModelingOperators/Private/CompositionOps/VoxelMorphologyMeshesOp.cpp b/Engine/Plugins/Experimental/MeshModelingToolsetExp/Source/ModelingOperators/Private/CompositionOps/VoxelMorphologyMeshesOp.cpp
--- a/Engine/Plugins/Experimental/MeshModelingToolsetExp/Source/ModelingOperators/Private/CompositionOps/VoxelMorphologyMeshesOp.cpp
+++ b/Engine/Plugins/Experimental/MeshModelingToolsetExp/Source/ModelingOperators/Private/CompositionOps/VoxelMorphologyMeshesOp.cpp
@@ -88,7 +88,7 @@ void FVoxelMorphologyMeshesOp::CalculateResult(FProgressCancel* Progress)
 	if (bVoxWrapInput && ThickenShells > 0)
 	{
 		// positive offsets should be at least a cell wide so we don't end up deleting a bunch of the input surface
-		double CellSize = CombinedMesh.GetCachedBounds().MaxDim() / InputVoxelCount;
+		double CellSize = CombinedMesh.GetBounds(true).MaxDim() / InputVoxelCount;
 		double SafeThickness = FMathd::Max(CellSize * 2, ThickenShells);
 
 		FMeshNormals::QuickComputeVertexNormals(CombinedMesh);
@@ -123,7 +123,7 @@ void FVoxelMorphologyMeshesOp::CalculateResult(FProgressCancel* Progress)
 	}
 
 	ImplicitMorphology.SourceSpatial = &Spatial;
-	ImplicitMorphology.SetCellSizesAndDistance(CombinedMesh.GetCachedBounds(), Distance, InputVoxelCount, OutputVoxelCount);
+	ImplicitMorphology.SetCellSizesAndDistance(CombinedMesh.GetBounds(true), Distance, InputVoxelCount, OutputVoxelCount);
 	
 	ImplicitMorphology.CancelF = [&Progress]()
 	{
diff --git a/Engine/Plugins/Experimental/MeshModelingToolsetExp/Source/ModelingOperators/Private/CompositionOps/VoxelSolidifyMeshesOp.cpp b/Engine/Plugins/Experimental/MeshModelingToolsetExp/Source/ModelingOperators/Private/CompositionOps/VoxelSolidifyMeshesOp.cpp
--- a/Engine/Plugins/Experimental/MeshModelingToolsetExp/Source/ModelingOperators/Private/CompositionOps/VoxelSolidifyMeshesOp.cpp
+++ b/Engine/Plugins/Experimental/MeshModelingToolsetExp/Source/ModelingOperators/Private/CompositionOps/VoxelSolidifyMeshesOp.cpp
@@ -64,7 +64,7 @@ void FVoxelSolidifyMeshesOp::CalculateResult(FProgressCancel* Progress)
 	if (bApplyThickenShells)
 	{
 		// thickness should be at least a cell wide so we don't end up deleting a bunch of the input surface
-		double CellSize = CombinedMesh.GetCachedBounds().MaxDim() / InputVoxelCount;
+		double CellSize = CombinedMesh.GetBounds(true).MaxDim() / InputVoxelCount;
 		double SafeThickness = FMathd::Max(CellSize * 2, ThickenShells);
 
 		FMeshNormals::QuickComputeVertexNormals(CombinedMesh);
diff --git a/Engine/Plugins/Experimental/PlanarCutPlugin/Source/PlanarCut/Private/PlanarCut.cpp b/Engine/Plugins/Experimental/PlanarCutPlugin/Source/PlanarCut/Private/PlanarCut.cpp
--- a/Engine/Plugins/Experimental/PlanarCutPlugin/Source/PlanarCut/Private/PlanarCut.cpp
+++ b/Engine/Plugins/Experimental/PlanarCutPlugin/Source/PlanarCut/Private/PlanarCut.cpp
@@ -1293,7 +1293,7 @@ void ConvertToMeshDescription(
 
 	if (bCenterPivot)
 	{
-		FAxisAlignedBox3d Bounds = CombinedMesh.GetCachedBounds();
+		FAxisAlignedBox3d Bounds = CombinedMesh.GetBounds(true);
 		FVector3d Translate = -Bounds.Center();
 		MeshTransforms::Translate(CombinedMesh, Translate);
 		TransformOut = FTransform((FVector)-Translate);
diff --git a/Engine/Plugins/Runtime/GeometryProcessing/Source/DynamicMesh/Private/Operations/MeshBoolean.cpp b/Engine/Plugins/Runtime/GeometryProcessing/Source/DynamicMesh/Private/Operations/MeshBoolean.cpp
--- a/Engine/Plugins/Runtime/GeometryProcessing/Source/DynamicMesh/Private/Operations/MeshBoolean.cpp
+++ b/Engine/Plugins/Runtime/GeometryProcessing/Source/DynamicMesh/Private/Operations/MeshBoolean.cpp
@@ -215,8 +215,8 @@ bool FMeshBoolean::Compute()
 	FDynamicMesh3* CutMesh[2]{ Result, &CutMeshB }; // just an alias to keep things organized
 
 	// transform the copies to a shared space (centered at the origin and scaled to a unit cube)
-	FAxisAlignedBox3d CombinedAABB(CutMesh[0]->GetCachedBounds(), Transforms[0]);
-	FAxisAlignedBox3d MeshB_AABB(CutMesh[1]->GetCachedBounds(), Transforms[1]);
+	FAxisAlignedBox3d CombinedAABB(CutMesh[0]->GetBounds(true), Transforms[0]);
+	FAxisAlignedBox3d MeshB_AABB(CutMesh[1]->GetBounds(true), Transforms[1]);
 	CombinedAABB.Contain(MeshB_AABB);
 	double ScaleFactor = 1.0 / FMath::Clamp(CombinedAABB.MaxDim(), 0.01, 1000000.0);
 	for (int MeshIdx = 0; MeshIdx < 2; MeshIdx++)
@@ -504,7 +504,7 @@ bool FMeshBoolean::Compute()
 			int OtherMeshIdx = 1 - MeshIdx;
 			FDynamicMesh3& OtherMesh = *CutMesh[OtherMeshIdx];
 
-			TPointHashGrid3d<int> OtherMeshPointHash(OtherMesh.GetCachedBounds().MaxDim() / 64, -1);
+			TPointHashGrid3d<int> OtherMeshPointHash(OtherMesh.GetBounds(true).MaxDim() / 64, -1);
 			for (int BoundaryVID : PossUnmatchedBdryVerts[OtherMeshIdx])
 			{
 				OtherMeshPointHash.InsertPointUnsafe(BoundaryVID, OtherMesh.GetVertex(BoundaryVID));
diff --git a/Engine/Plugins/Runtime/GeometryProcessing/Source/DynamicMesh/Private/Operations/MeshSelfUnion.cpp b/Engine/Plugins/Runtime/GeometryProcessing/Source/DynamicMesh/Private/Operations/MeshSelfUnion.cpp
--- a/Engine/Plugins/Runtime/GeometryProcessing/Source/DynamicMesh/Private/Operations/MeshSelfUnion.cpp
+++ b/Engine/Plugins/Runtime/GeometryProcessing/Source/DynamicMesh/Private/Operations/MeshSelfUnion.cpp
@@ -23,7 +23,7 @@ using namespace UE::Geometry;
 bool FMeshSelfUnion::Compute()
 {
 	// transform the mesh to a shared space (centered at the origin, and scaled to a unit cube)
-	FAxisAlignedBox3d AABB = Mesh->GetCachedBounds();
+	FAxisAlignedBox3d AABB = Mesh->GetBounds(true);
 	double ScaleFactor = 1.0 / FMath::Clamp(AABB.MaxDim(), 0.01, 1000000.0);
 	FTransform3d TransformToCenteredBox = FTransform3d::Identity();
 	TransformToCenteredBox.SetTranslation(ScaleFactor * (TransformToCenteredBox.GetTranslation() - AABB.Center()));
@@ -277,7 +277,7 @@ bool FMeshSelfUnion::Compute()
 	}
 
 	// Hash boundary verts for faster search
-	TPointHashGrid3d<int> PointHash(Mesh->GetCachedBounds().MaxDim() / 64, -1);
+	TPointHashGrid3d<int> PointHash(Mesh->GetBounds(true).MaxDim() / 64, -1);
 	for (int BoundaryVID : PossUnmatchedBdryVerts)
 	{
 		PointHash.InsertPointUnsafe(BoundaryVID, Mesh->GetVertex(BoundaryVID));
diff --git a/Engine/Plugins/Runtime/GeometryProcessing/Source/DynamicMesh/Private/Operations/RepairOrientation.cpp b/Engine/Plugins/Runtime/GeometryProcessing/Source/DynamicMesh/Private/Operations/RepairOrientation.cpp
--- a/Engine/Plugins/Runtime/GeometryProcessing/Source/DynamicMesh/Private/Operations/RepairOrientation.cpp
+++ b/Engine/Plugins/Runtime/GeometryProcessing/Source/DynamicMesh/Private/Operations/RepairOrientation.cpp
@@ -89,7 +89,7 @@ void FMeshRepairOrientation::ComputeStatistics(FDynamicMeshAABBTree3* Tree)
 void FMeshRepairOrientation::ComputeComponentStatistics(FDynamicMeshAABBTree3* Tree, Component& C)
 {
 	C.InFacing = C.OutFacing = 0;
-	double Dist = 2 * Mesh->GetCachedBounds().DiagonalLength();
+	double Dist = 2 * Mesh->GetBounds(true).DiagonalLength();
 	
 	// only want to raycast triangles in this component
 	TSet<int> TrisInComponent; TrisInComponent.Append(C.Triangles);
diff --git a/Engine/Plugins/Runtime/MeshModelingToolset/Source/MeshModelingTools/Private/AddPrimitiveTool.cpp b/Engine/Plugins/Runtime/MeshModelingToolset/Source/MeshModelingTools/Private/AddPrimitiveTool.cpp
--- a/Engine/Plugins/Runtime/MeshModelingToolset/Source/MeshModelingTools/Private/AddPrimitiveTool.cpp
+++ b/Engine/Plugins/Runtime/MeshModelingToolset/Source/MeshModelingTools/Private/AddPrimitiveTool.cpp
@@ -294,7 +294,7 @@ void UAddPrimitiveTool::UpdatePreviewMesh()
 	}
 
 	// set mesh position
-	FAxisAlignedBox3d Bounds = NewMesh.GetCachedBounds();
+	FAxisAlignedBox3d Bounds = NewMesh.GetBounds(true);
 	FVector3d TargetOrigin = Bounds.Center();
 	if (ShapeSettings->PivotLocation == EMakeMeshPivotLocation::Base)
 	{
diff --git a/Engine/Source/Runtime/GeometryCore/Private/DynamicMesh/Operations/MergeCoincidentMeshEdges.cpp b/Engine/Source/Runtime/GeometryCore/Private/DynamicMesh/Operations/MergeCoincidentMeshEdges.cpp
--- a/Engine/Source/Runtime/GeometryCore/Private/DynamicMesh/Operations/MergeCoincidentMeshEdges.cpp
+++ b/Engine/Source/Runtime/GeometryCore/Private/DynamicMesh/Operations/MergeCoincidentMeshEdges.cpp
@@ -30,7 +30,7 @@ bool FMergeCoincidentMeshEdges::Apply()
 	int hashN = 64;
 	if (Mesh->TriangleCount() > 100000)   hashN = 128;
 	if (Mesh->TriangleCount() > 1000000)  hashN = 256;
-	FAxisAlignedBox3d bounds = Mesh->GetCachedBounds();
+	FAxisAlignedBox3d bounds = Mesh->GetBounds(true);
 	double CellSize = bounds.MaxDim() / (double)hashN;
 	MidpointsHash.Build(CellSize, bounds.Min);
 
