ModelingTools: Replace DynamicMeshComponent with UPreviewMesh in BakeTexture

#rnx

diff --git a/Engine/Plugins/Experimental/MeshModelingToolsetExp/Source/MeshModelingToolsExp/Private/BakeMeshAttributeMapsTool.cpp b/Engine/Plugins/Experimental/MeshModelingToolsetExp/Source/MeshModelingToolsExp/Private/BakeMeshAttributeMapsTool.cpp
--- a/Engine/Plugins/Experimental/MeshModelingToolsetExp/Source/MeshModelingToolsExp/Private/BakeMeshAttributeMapsTool.cpp
+++ b/Engine/Plugins/Experimental/MeshModelingToolsetExp/Source/MeshModelingToolsExp/Private/BakeMeshAttributeMapsTool.cpp
@@ -15,8 +15,6 @@
 #include "Sampling/MeshResampleImageEvaluator.h"
 #include "Util/IndexUtil.h"
 
-#include "Components/DynamicMeshComponent.h"
-
 #include "Materials/Material.h"
 #include "Materials/MaterialInstanceDynamic.h"
 #include "ImageUtils.h"
@@ -93,7 +91,7 @@ UInteractiveTool* UBakeMeshAttributeMapsToolBuilder::BuildTool(const FToolBuilde
 
 	TArray<TObjectPtr<UToolTarget>> Targets = SceneState.TargetManager->BuildAllSelectedTargetable(SceneState, GetTargetRequirements());
 	NewTool->SetTargets(MoveTemp(Targets));
-
+	NewTool->SetWorld(SceneState.World);
 	return NewTool;
 }
 
@@ -283,35 +281,9 @@ void UBakeMeshAttributeMapsTool::Setup()
 {
 	UInteractiveTool::Setup();
 
+	// Setup preview materials
 	InitializeEmptyMaps();
 
-	// create dynamic mesh component to use for live preview
-	// TODO: convert to UPreviewMesh
-	AActor* ParentActor = UE::ToolTarget::GetTargetActor(Targets[0]);
-	DynamicMeshComponent = NewObject<UDynamicMeshComponent>(ParentActor);
-	DynamicMeshComponent->SetupAttachment(ParentActor->GetRootComponent());
-	DynamicMeshComponent->RegisterComponent();
-	DynamicMeshComponent->SetWorldTransform((FTransform)UE::ToolTarget::GetLocalToWorldTransform(Targets[0]));
-
-	// transfer materials
-	FComponentMaterialSet MaterialSet = UE::ToolTarget::GetMaterialSet(Targets[0]);
-	for (int k = 0; k < MaterialSet.Materials.Num(); ++k)
-	{
-		DynamicMeshComponent->SetMaterial(k, MaterialSet.Materials[k]);
-	}
-
-	FDynamicMesh3 InputMeshWithTangents = UE::ToolTarget::GetDynamicMeshCopy(Targets[0], true);
-	DynamicMeshComponent->SetTangentsType(EDynamicMeshComponentTangentsMode::ExternallyProvided);
-	DynamicMeshComponent->SetMesh(MoveTemp(InputMeshWithTangents));
-	
-	DynamicMeshComponent->ProcessMesh([&](const FDynamicMesh3& ReadMesh)
-	{
-		BaseMesh.Copy(ReadMesh);
-		BaseMeshTangents = MakeShared<FMeshTangentsd, ESPMode::ThreadSafe>(&BaseMesh);
-		BaseMeshTangents->CopyTriVertexTangents(ReadMesh);
-	});
-	BaseSpatial.SetMesh(&BaseMesh, true);
-
 	UMaterial* Material = LoadObject<UMaterial>(nullptr, TEXT("/MeshModelingToolsetExp/Materials/BakePreviewMaterial"));
 	check(Material);
 	if (Material != nullptr)
@@ -320,7 +292,6 @@ void UBakeMeshAttributeMapsTool::Setup()
 		PreviewMaterial->SetTextureParameterValue(TEXT("NormalMap"), EmptyNormalMap);
 		PreviewMaterial->SetTextureParameterValue(TEXT("OcclusionMap"), EmptyColorMapWhite);
 		PreviewMaterial->SetTextureParameterValue(TEXT("ColorMap"), EmptyColorMapWhite);
-		DynamicMeshComponent->SetOverrideRenderMaterial(PreviewMaterial);
 	}
 	UMaterial* BentNormalMaterial = LoadObject<UMaterial>(nullptr, TEXT("/MeshModelingToolsetExp/Materials/BakeBentNormalPreviewMaterial"));
 	check(BentNormalMaterial);
@@ -335,10 +306,30 @@ void UBakeMeshAttributeMapsTool::Setup()
 		WorkingPreviewMaterial = UMaterialInstanceDynamic::Create(WorkingMaterial, GetToolManager());
 	}
 
+	// Initialize preview mesh
 	bIsBakeToSelf = (Targets.Num() == 1);
 
 	UE::ToolTarget::HideSourceObject(Targets[0]);
 
+	const FDynamicMesh3 InputMeshWithTangents = UE::ToolTarget::GetDynamicMeshCopy(Targets[0], true);
+	PreviewMesh = NewObject<UPreviewMesh>(this);
+	PreviewMesh->CreateInWorld(TargetWorld, FTransform::Identity);
+	PreviewMesh->SetTransform(static_cast<FTransform>(UE::ToolTarget::GetLocalToWorldTransform(Targets[0])));
+	PreviewMesh->SetTangentsMode(EDynamicMeshComponentTangentsMode::ExternallyProvided);
+	PreviewMesh->ReplaceMesh(InputMeshWithTangents);
+	PreviewMesh->SetMaterials(UE::ToolTarget::GetMaterialSet(Targets[0]).Materials);
+	PreviewMesh->SetOverrideRenderMaterial(PreviewMaterial);
+	PreviewMesh->SetVisible(true);
+
+	PreviewMesh->ProcessMesh([this](const FDynamicMesh3& Mesh)
+	{
+		BaseMesh.Copy(Mesh);
+		BaseSpatial.SetMesh(&BaseMesh, true);
+		BaseMeshTangents = MakeShared<FMeshTangentsd, ESPMode::ThreadSafe>(&BaseMesh);
+		BaseMeshTangents->CopyTriVertexTangents(Mesh);
+	});
+
+	// Setup tool property sets
 	Settings = NewObject<UBakeMeshAttributeMapsToolProperties>(this);
 	Settings->RestoreProperties(this);
 	Settings->UVLayerNamesList.Reset();
@@ -364,6 +355,7 @@ void UBakeMeshAttributeMapsTool::Setup()
 	Settings->WatchProperty(Settings->bUseWorldSpace, [this](bool) { bDetailMeshValid = false; bInputsDirty = true; });
 	Settings->WatchProperty(Settings->Thickness, [this](float) { bInputsDirty = true; });
 	Settings->WatchProperty(Settings->Multisampling, [this](EBakeMultisampling) { bInputsDirty = true; });
+	
 
 	NormalMapProps = NewObject<UBakedNormalMapToolProperties>(this);
 	NormalMapProps->RestoreProperties(this);
@@ -506,6 +498,12 @@ TUniquePtr<UE::Geometry::TGenericDataOperator<FMeshMapBaker>> UBakeMeshAttribute
 }
 
 
+void UBakeMeshAttributeMapsTool::SetWorld(UWorld* World)
+{
+	TargetWorld = World;
+}
+
+
 void UBakeMeshAttributeMapsTool::Shutdown(EToolShutdownType ShutdownType)
 {
 	Settings->SaveProperties(this);
@@ -520,7 +518,7 @@ void UBakeMeshAttributeMapsTool::Shutdown(EToolShutdownType ShutdownType)
 	{
 		Compute->Shutdown();
 	}
-	if (DynamicMeshComponent != nullptr)
+	if (PreviewMesh != nullptr)
 	{
 		UE::ToolTarget::ShowSourceObject(Targets[0]);
 
@@ -588,9 +586,9 @@ void UBakeMeshAttributeMapsTool::Shutdown(EToolShutdownType ShutdownType)
 			ensure(bCreatedAssetOK);
 		}
 
-		DynamicMeshComponent->UnregisterComponent();
-		DynamicMeshComponent->DestroyComponent();
-		DynamicMeshComponent = nullptr;
+		PreviewMesh->SetVisible(false);
+		PreviewMesh->Disconnect();
+		PreviewMesh = nullptr;
 	}
 }
 
@@ -604,7 +602,7 @@ void UBakeMeshAttributeMapsTool::OnTick(float DeltaTime)
 		float ElapsedComputeTime = Compute->GetElapsedComputeTime();
 		if (!CanAccept() && ElapsedComputeTime > SecondsBeforeWorkingMaterial)
 		{
-			DynamicMeshComponent->SetOverrideRenderMaterial(WorkingPreviewMaterial);
+			PreviewMesh->SetOverrideRenderMaterial(WorkingPreviewMaterial);
 		}
 	}
 }
@@ -1100,7 +1098,7 @@ EBakeOpState UBakeMeshAttributeMapsTool::UpdateResult_MultiTexture()
 
 void UBakeMeshAttributeMapsTool::UpdateVisualization()
 {
-	DynamicMeshComponent->SetOverrideRenderMaterial(PreviewMaterial);
+	PreviewMesh->SetOverrideRenderMaterial(PreviewMaterial);
 
 	// Map CachedMaps to Settings->Result
 	int NumResults = Settings->Result.Num();
@@ -1145,7 +1143,7 @@ void UBakeMeshAttributeMapsTool::UpdateVisualization()
 				}
 				BentNormalPreviewMaterial->SetTextureParameterValue(TEXT("ColorMap"), EmptyColorMapWhite);
 				BentNormalPreviewMaterial->SetTextureParameterValue(TEXT("BentNormalMap"), PreviewMap);
-				DynamicMeshComponent->SetOverrideRenderMaterial(BentNormalPreviewMaterial);
+				PreviewMesh->SetOverrideRenderMaterial(BentNormalPreviewMaterial);
 				break;
 			case EBakeMapType::Curvature:
 			case EBakeMapType::NormalImage:
diff --git a/Engine/Plugins/Experimental/MeshModelingToolsetExp/Source/MeshModelingToolsExp/Public/BakeMeshAttributeMapsTool.h b/Engine/Plugins/Experimental/MeshModelingToolsetExp/Source/MeshModelingToolsExp/Public/BakeMeshAttributeMapsTool.h
--- a/Engine/Plugins/Experimental/MeshModelingToolsetExp/Source/MeshModelingToolsExp/Public/BakeMeshAttributeMapsTool.h
+++ b/Engine/Plugins/Experimental/MeshModelingToolsetExp/Source/MeshModelingToolsExp/Public/BakeMeshAttributeMapsTool.h
@@ -12,6 +12,7 @@
 #include "Sampling/MeshMapBaker.h"
 #include "ModelingOperators.h"
 #include "MeshOpPreviewHelpers.h"
+#include "PreviewMesh.h"
 #include "BakeMeshAttributeToolCommon.h"
 #include "BakeMeshAttributeMapsTool.generated.h"
 
@@ -167,6 +168,8 @@ public:
 	// IGenericDataOperatorFactory API
 	virtual TUniquePtr<UE::Geometry::TGenericDataOperator<UE::Geometry::FMeshMapBaker>> MakeNewOperator() override;
 
+	void SetWorld(UWorld* World);
+
 protected:
 	// need to update bResultValid if these are modified, so we don't publicly expose them. 
 	// @todo setters/getters for these
@@ -196,7 +199,8 @@ protected:
 protected:
 	friend class FMeshMapBakerOp;
 	
-	UDynamicMeshComponent* DynamicMeshComponent;
+	UPROPERTY()
+	TObjectPtr<UPreviewMesh> PreviewMesh;
 
 	UPROPERTY()
 	TObjectPtr<UMaterialInstanceDynamic> PreviewMaterial;
@@ -208,6 +212,8 @@ protected:
 	TObjectPtr<UMaterialInstanceDynamic> WorkingPreviewMaterial;
 	float SecondsBeforeWorkingMaterial = 0.75;
 
+	UWorld* TargetWorld = nullptr;
+
 	TSharedPtr<UE::Geometry::TMeshTangents<double>, ESPMode::ThreadSafe> BaseMeshTangents;
 	UE::Geometry::FDynamicMesh3 BaseMesh;
 	UE::Geometry::FDynamicMeshAABBTree3 BaseSpatial;
