ModelingTools: expose PatchBasedMeshUVGenerator as an option in ParameterizeMeshOp, and in ParameterizeMeshTool
#rnx

diff --git a/Engine/Plugins/Experimental/MeshModelingToolsetExp/Source/MeshModelingToolsEditorOnly/Private/ParameterizeMeshTool.cpp b/Engine/Plugins/Experimental/MeshModelingToolsetExp/Source/MeshModelingToolsEditorOnly/Private/ParameterizeMeshTool.cpp
--- a/Engine/Plugins/Experimental/MeshModelingToolsetExp/Source/MeshModelingToolsEditorOnly/Private/ParameterizeMeshTool.cpp
+++ b/Engine/Plugins/Experimental/MeshModelingToolsetExp/Source/MeshModelingToolsEditorOnly/Private/ParameterizeMeshTool.cpp
@@ -35,18 +35,19 @@ void UParameterizeMeshTool::Setup()
 
 	InputMesh = MakeShared<FDynamicMesh3, ESPMode::ThreadSafe>();
 	*InputMesh = UE::ToolTarget::GetDynamicMeshCopy(Target);
+	FTransform InputTransform = (FTransform)UE::ToolTarget::GetLocalToWorldTransform(Target);
+	FComponentMaterialSet MaterialSet = UE::ToolTarget::GetMaterialSet(Target);
 
 	Preview = NewObject<UMeshOpPreviewWithBackgroundCompute>(this);
 	Preview->Setup(this->TargetWorld, this);
 	Preview->PreviewMesh->SetTangentsMode(EDynamicMeshComponentTangentsMode::AutoCalculated);
 	Preview->PreviewMesh->ReplaceMesh(*InputMesh);
-	Preview->ConfigureMaterials(UE::ToolTarget::GetMaterialSet(Target).Materials,
-								ToolSetupUtil::GetDefaultWorkingMaterial(GetToolManager()));
-	Preview->PreviewMesh->SetTransform((FTransform)UE::ToolTarget::GetLocalToWorldTransform(Target));
+	Preview->ConfigureMaterials(MaterialSet.Materials, ToolSetupUtil::GetDefaultWorkingMaterial(GetToolManager()));
+	Preview->PreviewMesh->SetTransform(InputTransform);
 
 	Preview->OnMeshUpdated.AddLambda([this](UMeshOpPreviewWithBackgroundCompute* Op)
 	{
-		MaterialSettings->UpdateMaterials();
+		OnPreviewMeshUpdated();
 	});
 
 	UE::ToolTarget::HideSourceObject(Target);
@@ -79,6 +80,12 @@ void UParameterizeMeshTool::Setup()
 	AddToolPropertySource(XAtlasProperties);
 	SetToolPropertySourceEnabled(XAtlasProperties, false);
 
+	PatchBuilderProperties = NewObject<UParameterizeMeshToolPatchBuilderProperties>(this);
+	PatchBuilderProperties->RestoreProperties(this);
+	AddToolPropertySource(PatchBuilderProperties);
+	SetToolPropertySourceEnabled(PatchBuilderProperties, false);
+
+
 	MaterialSettings = NewObject<UExistingMeshMaterialProperties>(this);
 	MaterialSettings->MaterialMode = ESetMeshMaterialMode::Checkerboard;
 	MaterialSettings->RestoreProperties(this, TEXT("ModelingUVTools"));
@@ -87,6 +94,18 @@ void UParameterizeMeshTool::Setup()
 	MaterialSettings->UpdateMaterials();
 	Preview->OverrideMaterial = MaterialSettings->GetActiveOverrideMaterial();
 
+	if (bCreateUVLayoutViewOnSetup)
+	{
+		UVLayoutView = NewObject<UUVLayoutPreview>(this);
+		UVLayoutView->CreateInWorld(TargetWorld);
+		UVLayoutView->SetSourceMaterials(MaterialSet);
+		UVLayoutView->SetSourceWorldPosition(InputTransform, UE::ToolTarget::GetTargetActor(Target)->GetComponentsBoundingBox());
+		UVLayoutView->Settings->bVisible = false;
+		UVLayoutView->Settings->bShowWireframe = false;
+		UVLayoutView->Settings->RestoreProperties(this, TEXT("ParameterizeMeshTool"));
+		AddToolPropertySource(UVLayoutView->Settings);
+	}
+
 	Preview->InvalidateResult();    // start compute
 
 	SetToolDisplayName(LOCTEXT("ToolNameGlobal", "AutoUV"));
@@ -97,7 +116,11 @@ void UParameterizeMeshTool::Setup()
 
 void UParameterizeMeshTool::OnPropertyModified(UObject* PropertySet, FProperty* Property)
 {
-	if (PropertySet != MaterialSettings)
+	if (PropertySet == UVChannelProperties 
+		|| PropertySet == Settings 
+		|| PropertySet == UVAtlasProperties  
+		|| PropertySet == XAtlasProperties  
+		|| PropertySet == PatchBuilderProperties )
 	{
 		Preview->InvalidateResult();
 	}
@@ -111,6 +134,7 @@ void UParameterizeMeshTool::OnMethodTypeChanged()
 {
 	SetToolPropertySourceEnabled(UVAtlasProperties, Settings->Method == EParameterizeMeshUVMethod::UVAtlas);
 	SetToolPropertySourceEnabled(XAtlasProperties, Settings->Method == EParameterizeMeshUVMethod::XAtlas);
+	SetToolPropertySourceEnabled(PatchBuilderProperties, Settings->Method == EParameterizeMeshUVMethod::PatchBuilder);
 
 	Preview->InvalidateResult();
 }
@@ -118,9 +142,18 @@ void UParameterizeMeshTool::OnMethodTypeChanged()
 
 void UParameterizeMeshTool::Shutdown(EToolShutdownType ShutdownType)
 {
+	if (UVLayoutView)
+	{
+		UVLayoutView->Settings->SaveProperties(this, TEXT("RecomputeUVsTool"));
+		UVLayoutView->Disconnect();
+	}
+
 	UVChannelProperties->SaveProperties(this);
 	Settings->SaveProperties(this);
 	MaterialSettings->SaveProperties(this, TEXT("ModelingUVTools"));
+	UVAtlasProperties->SaveProperties(this);
+	XAtlasProperties->SaveProperties(this);
+	PatchBuilderProperties->SaveProperties(this);
 
 	FDynamicMeshOpResult Result = Preview->Shutdown();
 	
@@ -140,9 +173,22 @@ void UParameterizeMeshTool::Shutdown(EToolShutdownType ShutdownType)
 
 }
 
+void UParameterizeMeshTool::Render(IToolsContextRenderAPI* RenderAPI)
+{
+	if (UVLayoutView)
+	{
+		UVLayoutView->Render(RenderAPI);
+	}
+}
+
 void UParameterizeMeshTool::OnTick(float DeltaTime)
 {
 	Preview->Tick(DeltaTime);
+
+	if (UVLayoutView)
+	{
+		UVLayoutView->OnTick(DeltaTime);
+	}
 }
 
 bool UParameterizeMeshTool::CanAccept() const
@@ -166,6 +212,17 @@ TUniquePtr<FDynamicMeshOperator> UParameterizeMeshTool::MakeNewOperator()
 	// xatlas options
 	ParameterizeMeshOp->XAtlasMaxIterations = XAtlasProperties->MaxIterations;
 
+	// patchbuilder options
+	ParameterizeMeshOp->InitialPatchCount = PatchBuilderProperties->InitialPatches;
+	ParameterizeMeshOp->PatchCurvatureAlignmentWeight = PatchBuilderProperties->CurvatureAlignment;
+	ParameterizeMeshOp->PatchMergingMetricThresh = PatchBuilderProperties->MergingThreshold;
+	ParameterizeMeshOp->PatchMergingAngleThresh = PatchBuilderProperties->MaxAngleDeviation;
+	ParameterizeMeshOp->ExpMapNormalSmoothingSteps = PatchBuilderProperties->SmoothingSteps;
+	ParameterizeMeshOp->ExpMapNormalSmoothingAlpha = PatchBuilderProperties->SmoothingAlpha;
+	ParameterizeMeshOp->bEnablePacking = PatchBuilderProperties->bAutoPack;
+	ParameterizeMeshOp->Width = ParameterizeMeshOp->Height = PatchBuilderProperties->TextureResolution;
+
+
 	ParameterizeMeshOp->SetTransform(UE::ToolTarget::GetLocalToWorldTransform(Target));
 
 	return ParameterizeMeshOp;
@@ -173,4 +230,22 @@ TUniquePtr<FDynamicMeshOperator> UParameterizeMeshTool::MakeNewOperator()
 
 
 
+void UParameterizeMeshTool::OnPreviewMeshUpdated()
+{
+	if (UVLayoutView)
+	{
+		int32 UVChannel = UVChannelProperties ? UVChannelProperties->GetSelectedChannelIndex(true) : 0;
+		Preview->PreviewMesh->ProcessMesh([&](const FDynamicMesh3& NewMesh)
+		{
+			UVLayoutView->UpdateUVMesh(&NewMesh, UVChannel);
+		});
+	}
+
+	if (MaterialSettings)
+	{
+		MaterialSettings->UpdateMaterials();
+	}
+}
+
+
 #undef LOCTEXT_NAMESPACE
diff --git a/Engine/Plugins/Experimental/MeshModelingToolsetExp/Source/MeshModelingToolsEditorOnly/Public/ParameterizeMeshTool.h b/Engine/Plugins/Experimental/MeshModelingToolsetExp/Source/MeshModelingToolsEditorOnly/Public/ParameterizeMeshTool.h
--- a/Engine/Plugins/Experimental/MeshModelingToolsetExp/Source/MeshModelingToolsEditorOnly/Public/ParameterizeMeshTool.h
+++ b/Engine/Plugins/Experimental/MeshModelingToolsetExp/Source/MeshModelingToolsEditorOnly/Public/ParameterizeMeshTool.h
@@ -8,6 +8,7 @@
 #include "MeshOpPreviewHelpers.h"
 #include "Properties/MeshMaterialProperties.h"
 #include "Properties/MeshUVChannelProperties.h"
+#include "Drawing/UVLayoutPreview.h"
 
 #include "ParameterizeMeshTool.generated.h"
 
@@ -30,10 +31,12 @@ UENUM()
 enum class EParameterizeMeshUVMethod
 {
 	// keep values the same as UE::Geometry::EParamOpBackend!
-	/** Compute Automatic UVs using UVAtlas */
-	UVAtlas = 0,
+	/** Compute Automatic UVs using PatchBuilder */
+	PatchBuilder = 0,
 	/** Compute Automatic UVs using XAtlas */
-	XAtlas = 1
+	XAtlas = 1,
+	/** Compute Automatic UVs using UVAtlas */
+	UVAtlas = 2,
 };
 
 
@@ -45,7 +48,7 @@ class MESHMODELINGTOOLSEDITORONLY_API UParameterizeMeshToolProperties : public U
 public:
 	/** Automatic UV Generation technique to use */
 	UPROPERTY(EditAnywhere, Category = Options)
-	EParameterizeMeshUVMethod Method = EParameterizeMeshUVMethod::UVAtlas;
+	EParameterizeMeshUVMethod Method = EParameterizeMeshUVMethod::PatchBuilder;
 };
 
 
@@ -80,6 +83,48 @@ public:
 
 
 
+UCLASS()
+class MESHMODELINGTOOLSEDITORONLY_API UParameterizeMeshToolPatchBuilderProperties : public UInteractiveToolPropertySet
+{
+	GENERATED_BODY()
+public:
+	/** Number of initial patches mesh will be split into before computing island merging */
+	UPROPERTY(EditAnywhere, Category = "PatchBuilder", meta = (UIMin = "1", UIMax = "1000", ClampMin = "1", ClampMax = "99999999"))
+	int InitialPatches = 100;
+
+	/** This parameter controls alignment of the initial patches to creases in the mesh */
+	UPROPERTY(EditAnywhere, Category = "PatchBuilder", meta = (UIMin = "0.1", UIMax = "2.0", ClampMin = "0.01", ClampMax = "100.0"))
+	float CurvatureAlignment = 1.0f;
+
+	/** Distortion/Stretching Threshold for island merging - larger values increase the allowable UV stretching */
+	UPROPERTY(EditAnywhere, Category = "PatchBuilder", meta = (UIMin = "1.0", UIMax = "5.0", ClampMin = "1.0", EditCondition = "bIslandMerging == true"))
+	float MergingThreshold = 1.5f;
+
+	/** UV islands will not be merged if their average face normals deviate by larger than this amount */
+	UPROPERTY(EditAnywhere, Category = "PatchBuilder", meta = (UIMin = "0.0", UIMax = "90.0", ClampMin = "0.0", ClampMax = "180.0"))
+	float MaxAngleDeviation = 45.0f;
+
+	/** Number of smoothing steps to apply in the ExpMap UV Generation method (Smoothing slightly increases distortion but produces more stable results) */
+	UPROPERTY(EditAnywhere, Category = "PatchBuilder", meta = (UIMin = "0", UIMax = "25", ClampMin = "0", ClampMax = "1000"))
+	int SmoothingSteps = 5;
+
+	/** Smoothing parameter, larger values result in faster smoothing in each step */
+	UPROPERTY(EditAnywhere, Category = "PatchBuilder", meta = (UIMin = "0", UIMax = "1.0", ClampMin = "0", ClampMax = "1.0"))
+	float SmoothingAlpha = 0.25f;
+
+
+	/** If enabled, result UVs are automatically packed into the standard UV 0-1 square */
+	UPROPERTY(EditAnywhere, Category = "PatchBuilder")
+	bool bAutoPack = true;
+
+	/** Target texture resolution used for UV packing, which determines gutter size */
+	UPROPERTY(EditAnywhere, Category = "PatchBuilder", meta = (UIMin = "64", UIMax = "2048", ClampMin = "2", ClampMax = "4096", EditCondition = bAutoPack))
+	int TextureResolution = 1024;
+};
+
+
+
+
 /**
  * UParameterizeMeshTool automatically decomposes the input mesh into charts, solves for UVs,
  * and then packs the resulting charts
@@ -93,6 +138,7 @@ public:
 	virtual void Setup() override;
 	virtual void Shutdown(EToolShutdownType ShutdownType) override;
 
+	virtual void Render(IToolsContextRenderAPI* RenderAPI) override;
 	virtual void OnTick(float DeltaTime) override;
 
 	virtual bool HasCancel() const override { return true; }
@@ -117,10 +163,19 @@ protected:
 	UPROPERTY()
 	TObjectPtr<UParameterizeMeshToolXAtlasProperties> XAtlasProperties = nullptr;
 
+	UPROPERTY()
+	TObjectPtr<UParameterizeMeshToolPatchBuilderProperties> PatchBuilderProperties = nullptr;
+
 
 	UPROPERTY()
 	TObjectPtr<UExistingMeshMaterialProperties> MaterialSettings = nullptr;
 
+	UPROPERTY()
+	bool bCreateUVLayoutViewOnSetup = true;
+
+	UPROPERTY()
+	TObjectPtr<UUVLayoutPreview> UVLayoutView = nullptr;
+
 protected:
 	UPROPERTY()
 	TObjectPtr<UMeshOpPreviewWithBackgroundCompute> Preview = nullptr;
@@ -130,4 +185,5 @@ protected:
 
 	void OnMethodTypeChanged();
 
+	void OnPreviewMeshUpdated();
 };
diff --git a/Engine/Plugins/Experimental/MeshModelingToolsetExp/Source/ModelingOperatorsEditorOnly/Private/ParameterizationOps/ParameterizeMeshOp.cpp b/Engine/Plugins/Experimental/MeshModelingToolsetExp/Source/ModelingOperatorsEditorOnly/Private/ParameterizationOps/ParameterizeMeshOp.cpp
--- a/Engine/Plugins/Experimental/MeshModelingToolsetExp/Source/ModelingOperatorsEditorOnly/Private/ParameterizationOps/ParameterizeMeshOp.cpp
+++ b/Engine/Plugins/Experimental/MeshModelingToolsetExp/Source/ModelingOperatorsEditorOnly/Private/ParameterizationOps/ParameterizeMeshOp.cpp
@@ -7,6 +7,7 @@
 
 #include "Parameterization/MeshLocalParam.h"
 #include "Parameterization/DynamicMeshUVEditor.h"
+#include "Parameterization/PatchBasedMeshUVGenerator.h"
 #include "DynamicMesh/DynamicMeshAABBTree3.h"
 #include "DynamicMesh/MeshNormals.h"
 #include "DynamicSubmesh3.h"
@@ -362,6 +363,34 @@ bool FParameterizeMeshOp::ComputeUVs_XAtlas(FDynamicMesh3& Mesh,  TFunction<bool
 
 
 
+bool FParameterizeMeshOp::ComputeUVs_PatchBuilder(FDynamicMesh3& InOutMesh, FProgressCancel* ProgressCancel)
+{
+	FPatchBasedMeshUVGenerator UVGenerator;
+	UVGenerator.TargetPatchCount = FMath::Max(1,InitialPatchCount);
+	UVGenerator.bNormalWeightedPatches = true;
+	UVGenerator.PatchNormalWeight = FMath::Clamp(PatchCurvatureAlignmentWeight, 0.0, 999999.0);
+	UVGenerator.MinPatchSize = 2;
+
+	UVGenerator.MergingThreshold = FMath::Clamp(PatchMergingMetricThresh, 0.001, 9999.0);
+	UVGenerator.MaxNormalDeviationDeg = FMath::Clamp(PatchMergingAngleThresh, 0.0, 180.0);
+
+	UVGenerator.NormalSmoothingRounds = FMath::Clamp(ExpMapNormalSmoothingSteps, 0, 9999);
+	UVGenerator.NormalSmoothingAlpha = FMath::Clamp(ExpMapNormalSmoothingAlpha, 0.0, 1.0);
+
+	UVGenerator.bAutoAlignPatches = true;
+	UVGenerator.bAutoPack = bEnablePacking;
+	UVGenerator.PackingTextureResolution = FMath::Clamp(this->Width, 16, 4096);
+	UVGenerator.PackingGutterWidth = this->Gutter;
+
+
+	FDynamicMeshUVEditor UVEditor(&InOutMesh, UVLayer, true);
+	FGeometryResult Result = UVGenerator.AutoComputeUVs(*UVEditor.GetMesh(), *UVEditor.GetOverlay(), ProgressCancel);
+
+	SetResultInfo(NewResultInfo);
+	return (Result.HasFailed() == false);
+}
+
+
 
 void FParameterizeMeshOp::CalculateResult(FProgressCancel* Progress)
 {
@@ -372,13 +401,19 @@ void FParameterizeMeshOp::CalculateResult(FProgressCancel* Progress)
 	}
 
 	ResultMesh = MakeUnique<FDynamicMesh3>(*InputMesh);
-
 	if (Progress && Progress->Cancelled())
 	{
 		SetResultInfo(FGeometryResult::Cancelled());
 		return;
 	}
 
+	// if using PatchBuilder, it handles all the output stuff
+	if (Method == EParamOpBackend::PatchBuilder)
+	{
+		ComputeUVs_PatchBuilder(*ResultMesh, Progress);
+		return;
+	}
+
 	// The UV atlas callback uses a float progress-based interrupter. 
 	TFunction<bool(float)> Interrupter = [Progress](float)->bool {return !(Progress && Progress->Cancelled()); };
 
diff --git a/Engine/Plugins/Experimental/MeshModelingToolsetExp/Source/ModelingOperatorsEditorOnly/Public/ParameterizationOps/ParameterizeMeshOp.h b/Engine/Plugins/Experimental/MeshModelingToolsetExp/Source/ModelingOperatorsEditorOnly/Public/ParameterizationOps/ParameterizeMeshOp.h
--- a/Engine/Plugins/Experimental/MeshModelingToolsetExp/Source/ModelingOperatorsEditorOnly/Public/ParameterizationOps/ParameterizeMeshOp.h
+++ b/Engine/Plugins/Experimental/MeshModelingToolsetExp/Source/ModelingOperatorsEditorOnly/Public/ParameterizationOps/ParameterizeMeshOp.h
@@ -15,8 +15,9 @@ namespace Geometry
 
 enum class EParamOpBackend
 {
-	UVAtlas = 0,
-	XAtlas = 1
+	PatchBuilder = 0,
+	XAtlas = 1,
+	UVAtlas = 2
 };
 
 class MODELINGOPERATORSEDITORONLY_API FParameterizeMeshOp : public FDynamicMeshOperator
@@ -38,10 +39,19 @@ public:
 	// XAtlas generation parameters
 	int32 XAtlasMaxIterations = 1;
 
+	// PatchBuilder generation parameters
+	int32 InitialPatchCount = 100;
+	double PatchCurvatureAlignmentWeight = 1.0;
+	double PatchMergingMetricThresh = 1.5;
+	double PatchMergingAngleThresh = 45.0;
+	int ExpMapNormalSmoothingSteps = 5;
+	double ExpMapNormalSmoothingAlpha = 0.25;
+
 	// UV layer
 	int32 UVLayer = 0;
 
 	// Atlas Packing parameters
+	bool bEnablePacking = true;
 	int32 Height = 512;
 	int32 Width = 512;
 	float Gutter = 2.5;
@@ -85,6 +95,7 @@ protected:
 
 	bool ComputeUVs_UVAtlas(FDynamicMesh3& InOutMesh, TFunction<bool(float)>& Interrupter);
 	bool ComputeUVs_XAtlas(FDynamicMesh3& InOutMesh, TFunction<bool(float)>& Interrupter);
+	bool ComputeUVs_PatchBuilder(FDynamicMesh3& InOutMesh, FProgressCancel* Progress);
 
 	void CopyNewUVsToMesh(
 		FDynamicMesh3& Mesh,
