GeometryProcessing: LWC Fixes. Remove GeometryCore TRay3, TSphere3, and TOrientedBox3 from ExplicitUseGeometryMathTypes.h, update call sites to fully-qualified name where necessary (in some cases handled by using declaration at class scope, or by wrapping code in UE::Geometry:: namespace)
#rnx

diff --git a/Engine/Plugins/Experimental/MeshLODToolset/Source/MeshLODToolset/Private/Tools/GenerateStaticMeshLODAssetTool.cpp b/Engine/Plugins/Experimental/MeshLODToolset/Source/MeshLODToolset/Private/Tools/GenerateStaticMeshLODAssetTool.cpp
--- a/Engine/Plugins/Experimental/MeshLODToolset/Source/MeshLODToolset/Private/Tools/GenerateStaticMeshLODAssetTool.cpp
+++ b/Engine/Plugins/Experimental/MeshLODToolset/Source/MeshLODToolset/Private/Tools/GenerateStaticMeshLODAssetTool.cpp
@@ -284,7 +284,7 @@ void UGenerateStaticMeshLODAssetTool::Setup()
 	// For the first computation, display a bounding box with the working material. Otherwise it looks like nothing
 	// is happening. And we don't want to copy over the potentially huge input mesh to be the preview mesh.
 	FGridBoxMeshGenerator MeshGen;
-	MeshGen.Box = FOrientedBox3d((FVector3d)Bounds.Origin, (FVector3d)Bounds.BoxExtent);
+	MeshGen.Box = UE::Geometry::FOrientedBox3d((FVector3d)Bounds.Origin, (FVector3d)Bounds.BoxExtent);
 	MeshGen.Generate();
 	FDynamicMesh3 BoxMesh(&MeshGen);
 	PreviewWithBackgroundCompute->PreviewMesh->UpdatePreview(MoveTemp(BoxMesh));
diff --git a/Engine/Plugins/Experimental/MeshModelingToolsetExp/Source/GeometryProcessingAdapters/Private/ApproximateActorsImpl.cpp b/Engine/Plugins/Experimental/MeshModelingToolsetExp/Source/GeometryProcessingAdapters/Private/ApproximateActorsImpl.cpp
--- a/Engine/Plugins/Experimental/MeshModelingToolsetExp/Source/GeometryProcessingAdapters/Private/ApproximateActorsImpl.cpp
+++ b/Engine/Plugins/Experimental/MeshModelingToolsetExp/Source/GeometryProcessingAdapters/Private/ApproximateActorsImpl.cpp
@@ -166,7 +166,7 @@ static void BakeTexturesFromPhotoCapture(
 		FVector3d RayDir = ImagePosWorld - SurfPos;
 		double Dist = Normalize(RayDir);
 		FVector3d RayOrigin = SurfPos + RayOffsetHackDist * RayDir;
-		int32 HitTID = Spatial.FindNearestHitTriangle(FRay3d(RayOrigin, RayDir), IMeshSpatial::FQueryOptions(Dist));
+		int32 HitTID = Spatial.FindNearestHitTriangle(UE::Geometry::FRay3d(RayOrigin, RayDir), IMeshSpatial::FQueryOptions(Dist));
 		return (HitTID == IndexConstants::InvalidID);
 	};
 
diff --git a/Engine/Plugins/Experimental/MeshModelingToolsetExp/Source/MeshModelingToolsEditorOnly/Private/PolygonOnMeshTool.cpp b/Engine/Plugins/Experimental/MeshModelingToolsetExp/Source/MeshModelingToolsEditorOnly/Private/PolygonOnMeshTool.cpp
--- a/Engine/Plugins/Experimental/MeshModelingToolsetExp/Source/MeshModelingToolsEditorOnly/Private/PolygonOnMeshTool.cpp
+++ b/Engine/Plugins/Experimental/MeshModelingToolsetExp/Source/MeshModelingToolsEditorOnly/Private/PolygonOnMeshTool.cpp
@@ -424,9 +424,9 @@ bool UPolygonOnMeshTool::HitTest(const FRay& Ray, FHitResult& OutHit)
 	if (DrawPolygonMechanic != nullptr)
 	{
 		FFrame3d HitPoint;
-		if (DrawPolygonMechanic->IsHitByRay(FRay3d(Ray), HitPoint))
+		if (DrawPolygonMechanic->IsHitByRay(UE::Geometry::FRay3d(Ray), HitPoint))
 		{
-			OutHit.Distance = FRay3d(Ray).Project(HitPoint.Origin);
+			OutHit.Distance = UE::Geometry::FRay3d(Ray).Project(HitPoint.Origin);
 			OutHit.ImpactPoint = (FVector)HitPoint.Origin;
 			OutHit.ImpactNormal = (FVector)HitPoint.Z();
 			return true;
diff --git a/Engine/Plugins/Experimental/MeshModelingToolsetExp/Source/MeshModelingToolsExp/Private/DrawPolyPathTool.cpp b/Engine/Plugins/Experimental/MeshModelingToolsetExp/Source/MeshModelingToolsExp/Private/DrawPolyPathTool.cpp
--- a/Engine/Plugins/Experimental/MeshModelingToolsetExp/Source/MeshModelingToolsExp/Private/DrawPolyPathTool.cpp
+++ b/Engine/Plugins/Experimental/MeshModelingToolsetExp/Source/MeshModelingToolsExp/Private/DrawPolyPathTool.cpp
@@ -263,9 +263,9 @@ bool UDrawPolyPathTool::HitTest(const FRay& Ray, FHitResult& OutHit)
 	if (SurfacePathMechanic != nullptr)
 	{
 		FFrame3d HitPoint;
-		if (SurfacePathMechanic->IsHitByRay(FRay3d(Ray), HitPoint))
+		if (SurfacePathMechanic->IsHitByRay(UE::Geometry::FRay3d(Ray), HitPoint))
 		{
-			OutHit.Distance = FRay3d(Ray).Project(HitPoint.Origin);
+			OutHit.Distance = UE::Geometry::FRay3d(Ray).Project(HitPoint.Origin);
 			OutHit.ImpactPoint = (FVector)HitPoint.Origin;
 			OutHit.ImpactNormal = (FVector)HitPoint.Z();
 			return true;
diff --git a/Engine/Plugins/Experimental/MeshModelingToolsetExp/Source/MeshModelingToolsExp/Private/DynamicMeshBrushTool.cpp b/Engine/Plugins/Experimental/MeshModelingToolsetExp/Source/MeshModelingToolsExp/Private/DynamicMeshBrushTool.cpp
--- a/Engine/Plugins/Experimental/MeshModelingToolsetExp/Source/MeshModelingToolsExp/Private/DynamicMeshBrushTool.cpp
+++ b/Engine/Plugins/Experimental/MeshModelingToolsetExp/Source/MeshModelingToolsExp/Private/DynamicMeshBrushTool.cpp
@@ -77,7 +77,7 @@ void UDynamicMeshBrushTool::Shutdown(EToolShutdownType ShutdownType)
 
 bool UDynamicMeshBrushTool::HitTest(const FRay& Ray, FHitResult& OutHit)
 {
-	return PreviewMesh->FindRayIntersection(FRay3d(Ray), OutHit);
+	return PreviewMesh->FindRayIntersection(UE::Geometry::FRay3d(Ray), OutHit);
 }
 
 
diff --git a/Engine/Plugins/Experimental/MeshModelingToolsetExp/Source/MeshModelingToolsExp/Private/EditMeshPolygonsTool.cpp b/Engine/Plugins/Experimental/MeshModelingToolsetExp/Source/MeshModelingToolsExp/Private/EditMeshPolygonsTool.cpp
--- a/Engine/Plugins/Experimental/MeshModelingToolsetExp/Source/MeshModelingToolsExp/Private/EditMeshPolygonsTool.cpp
+++ b/Engine/Plugins/Experimental/MeshModelingToolsetExp/Source/MeshModelingToolsExp/Private/EditMeshPolygonsTool.cpp
@@ -664,7 +664,7 @@ FBox UEditMeshPolygonsTool::GetWorldSpaceFocusBox()
 
 bool UEditMeshPolygonsTool::GetWorldSpaceFocusPoint(const FRay& WorldRay, FVector& PointOut)
 {
-	FRay3d LocalRay(WorldTransform.InverseTransformPosition((FVector3d)WorldRay.Origin),
+	UE::Geometry::FRay3d LocalRay(WorldTransform.InverseTransformPosition((FVector3d)WorldRay.Origin),
 		WorldTransform.InverseTransformNormal((FVector3d)WorldRay.Direction));
 	UE::Geometry::Normalize(LocalRay.Direction);
 
diff --git a/Engine/Plugins/Experimental/MeshModelingToolsetExp/Source/MeshModelingToolsExp/Private/MeshSelectionTool.cpp b/Engine/Plugins/Experimental/MeshModelingToolsetExp/Source/MeshModelingToolsExp/Private/MeshSelectionTool.cpp
--- a/Engine/Plugins/Experimental/MeshModelingToolsetExp/Source/MeshModelingToolsExp/Private/MeshSelectionTool.cpp
+++ b/Engine/Plugins/Experimental/MeshModelingToolsetExp/Source/MeshModelingToolsExp/Private/MeshSelectionTool.cpp
@@ -597,7 +597,7 @@ void UMeshSelectionTool::UpdateFaceSelection(const FBrushStampData& Stamp, const
 		for (int tid : TriangleROI)
 		{
 			FVector3d Centroid = Mesh->GetTriCentroid(tid);
-			int HitTID = GetOctree()->FindNearestHitObject(FRay3d(LocalEyePosition, UE::Geometry::Normalized(Centroid - LocalEyePosition)));
+			int HitTID = GetOctree()->FindNearestHitObject(UE::Geometry::FRay3d(LocalEyePosition, UE::Geometry::Normalized(Centroid - LocalEyePosition)));
 			if (HitTID == tid)
 			{
 				LocalROI.Add(HitTID);
diff --git a/Engine/Plugins/Experimental/MeshModelingToolsetExp/Source/MeshModelingToolsExp/Private/Physics/ExtractCollisionGeometryTool.cpp b/Engine/Plugins/Experimental/MeshModelingToolsetExp/Source/MeshModelingToolsExp/Private/Physics/ExtractCollisionGeometryTool.cpp
--- a/Engine/Plugins/Experimental/MeshModelingToolsetExp/Source/MeshModelingToolsExp/Private/Physics/ExtractCollisionGeometryTool.cpp
+++ b/Engine/Plugins/Experimental/MeshModelingToolsetExp/Source/MeshModelingToolsExp/Private/Physics/ExtractCollisionGeometryTool.cpp
@@ -276,7 +276,7 @@ void UExtractCollisionGeometryTool::RecalculateMesh_Simple()
 	for (const FKBoxElem& Box : AggGeom.BoxElems)
 	{
 		FMinimalBoxMeshGenerator BoxGen;
-		BoxGen.Box = FOrientedBox3d(
+		BoxGen.Box = UE::Geometry::FOrientedBox3d(
 			FFrame3d(FVector3d(Box.Center), FQuaterniond(Box.Rotation.Quaternion())),
 			0.5*FVector3d(Box.X, Box.Y, Box.Z));
 		BoxGen.Generate();
diff --git a/Engine/Plugins/Experimental/MeshModelingToolsetExp/Source/MeshModelingToolsExp/Private/Sculpting/MeshSculptToolBase.cpp b/Engine/Plugins/Experimental/MeshModelingToolsetExp/Source/MeshModelingToolsExp/Private/Sculpting/MeshSculptToolBase.cpp
--- a/Engine/Plugins/Experimental/MeshModelingToolsetExp/Source/MeshModelingToolsExp/Private/Sculpting/MeshSculptToolBase.cpp
+++ b/Engine/Plugins/Experimental/MeshModelingToolsetExp/Source/MeshModelingToolsExp/Private/Sculpting/MeshSculptToolBase.cpp
@@ -501,7 +501,7 @@ void UMeshSculptToolBase::OnEndDrag(const FRay& Ray)
 
 
 
-FRay3d UMeshSculptToolBase::GetLocalRay(const FRay& WorldRay) const
+UE::Geometry::FRay3d UMeshSculptToolBase::GetLocalRay(const FRay& WorldRay) const
 {
 	FRay3d LocalRay(CurTargetTransform.InverseTransformPosition((FVector3d)WorldRay.Origin),
 		CurTargetTransform.InverseTransformVector((FVector3d)WorldRay.Direction));
diff --git a/Engine/Plugins/Experimental/MeshModelingToolsetExp/Source/MeshModelingToolsExp/Private/Tests/DisplaceMeshTool_Tests.inl b/Engine/Plugins/Experimental/MeshModelingToolsetExp/Source/MeshModelingToolsExp/Private/Tests/DisplaceMeshTool_Tests.inl
--- a/Engine/Plugins/Experimental/MeshModelingToolsetExp/Source/MeshModelingToolsExp/Private/Tests/DisplaceMeshTool_Tests.inl
+++ b/Engine/Plugins/Experimental/MeshModelingToolsetExp/Source/MeshModelingToolsExp/Private/Tests/DisplaceMeshTool_Tests.inl
@@ -27,7 +27,7 @@ namespace DisplaceMeshToolTestsLocals
 	TSharedPtr<FDynamicMesh3, ESPMode::ThreadSafe> CreateTestBoxMesh()
 	{
 		FGridBoxMeshGenerator BoxGenerator;
-		BoxGenerator.Box = FOrientedBox3d(FVector3d::Zero(), FVector3d(1.0, 1.0, 1.0));
+		BoxGenerator.Box = UE::Geometry::FOrientedBox3d(FVector3d::Zero(), FVector3d(1.0, 1.0, 1.0));
 		int EdgeNum = 5;
 		BoxGenerator.EdgeVertices = FIndex3i(EdgeNum, EdgeNum, EdgeNum);
 		BoxGenerator.Generate();
diff --git a/Engine/Plugins/Experimental/MeshModelingToolsetExp/Source/MeshModelingToolsExp/Private/ToolActivities/PolyEditInsertEdgeLoopActivity.cpp b/Engine/Plugins/Experimental/MeshModelingToolsetExp/Source/MeshModelingToolsExp/Private/ToolActivities/PolyEditInsertEdgeLoopActivity.cpp
--- a/Engine/Plugins/Experimental/MeshModelingToolsetExp/Source/MeshModelingToolsExp/Private/ToolActivities/PolyEditInsertEdgeLoopActivity.cpp
+++ b/Engine/Plugins/Experimental/MeshModelingToolsetExp/Source/MeshModelingToolsExp/Private/ToolActivities/PolyEditInsertEdgeLoopActivity.cpp
@@ -324,7 +324,7 @@ FInputRayHit UPolyEditInsertEdgeLoopActivity::HitTest(const FRay& WorldRay)
 	FInputRayHit Hit;
 
 	// See if we hit an edge
-	FRay3d LocalRay((FVector3d)TargetTransform.InverseTransformPosition(WorldRay.Origin),
+	UE::Geometry::FRay3d LocalRay((FVector3d)TargetTransform.InverseTransformPosition(WorldRay.Origin),
 		(FVector3d)TargetTransform.InverseTransformVector(WorldRay.Direction), false);
 	FGroupTopologySelection Selection;
 	FVector3d Position, Normal;
@@ -348,7 +348,7 @@ FInputRayHit UPolyEditInsertEdgeLoopActivity::HitTest(const FRay& WorldRay)
 bool UPolyEditInsertEdgeLoopActivity::UpdateHoveredItem(const FRay& WorldRay)
 {
 	// Check that we hit an edge
-	FRay3d LocalRay((FVector3d)TargetTransform.InverseTransformPosition(WorldRay.Origin),
+	UE::Geometry::FRay3d LocalRay((FVector3d)TargetTransform.InverseTransformPosition(WorldRay.Origin),
 		(FVector3d)TargetTransform.InverseTransformVector(WorldRay.Direction), false);
 
 	FGroupTopologySelection Selection;
diff --git a/Engine/Plugins/Experimental/MeshModelingToolsetExp/Source/MeshModelingToolsExp/Private/UVProjectionTool.cpp b/Engine/Plugins/Experimental/MeshModelingToolsetExp/Source/MeshModelingToolsExp/Private/UVProjectionTool.cpp
--- a/Engine/Plugins/Experimental/MeshModelingToolsetExp/Source/MeshModelingToolsExp/Private/UVProjectionTool.cpp
+++ b/Engine/Plugins/Experimental/MeshModelingToolsetExp/Source/MeshModelingToolsExp/Private/UVProjectionTool.cpp
@@ -318,11 +318,11 @@ void UUVProjectionTool::InitializeMesh()
 
 
 
-FOrientedBox3d UUVProjectionTool::GetProjectionBox()
+UE::Geometry::FOrientedBox3d UUVProjectionTool::GetProjectionBox()
 {
 	FFrame3d BoxFrame(TransformProxy->GetTransform());
 	FVector3d BoxDimensions = 0.5 * (FVector3d)BasicProperties->Dimensions;
-	return FOrientedBox3d(BoxFrame, BoxDimensions);
+	return UE::Geometry::FOrientedBox3d(BoxFrame, BoxDimensions);
 }
 
 
@@ -332,7 +332,7 @@ void UUVProjectionTool::Render(IToolsContextRenderAPI* RenderAPI)
 	ProjectionShapeVisualizer.bDepthTested = false;
 	ProjectionShapeVisualizer.BeginFrame(RenderAPI, CameraState);
 
-	FOrientedBox3d PrimBox = GetProjectionBox();
+	UE::Geometry::FOrientedBox3d PrimBox = GetProjectionBox();
 	FVector MinCorner = (FVector)-PrimBox.Extents;
 	FVector MaxCorner = (FVector)PrimBox.Extents;
 	float Width = MaxCorner.X - MinCorner.X;
@@ -548,7 +548,7 @@ void UUVProjectionTool::ApplyAction_AutoFit(bool bAlign)
 		if (BasicProperties->ProjectionType == EUVProjectionMethod::Plane || BasicProperties->ProjectionType == EUVProjectionMethod::Box || BasicProperties->ProjectionType == EUVProjectionMethod::Cylinder)
 		{
 			// compute min-volume box
-			FOrientedBox3d MinBoundingBox;
+			UE::Geometry::FOrientedBox3d MinBoundingBox;
 			FMinVolumeBox3d MinBoxCalc;
 			bool bMinBoxOK = MinBoxCalc.SolveSubsample(Points.Num(), 1000,
 				[&](int32 Index) { return WorldTransform.TransformPosition(Points[Index]); }, false, nullptr);
@@ -558,7 +558,7 @@ void UUVProjectionTool::ApplyAction_AutoFit(bool bAlign)
 			}
 			else
 			{
-				MinBoundingBox = FOrientedBox3d(WorldBounds);
+				MinBoundingBox = UE::Geometry::FOrientedBox3d(WorldBounds);
 			}
 
 			AlignFrame = MinBoundingBox.Frame;
diff --git a/Engine/Plugins/Experimental/MeshModelingToolsetExp/Source/MeshModelingToolsExp/Public/DeformMeshPolygonsTool.h b/Engine/Plugins/Experimental/MeshModelingToolsetExp/Source/MeshModelingToolsExp/Public/DeformMeshPolygonsTool.h
--- a/Engine/Plugins/Experimental/MeshModelingToolsetExp/Source/MeshModelingToolsExp/Public/DeformMeshPolygonsTool.h
+++ b/Engine/Plugins/Experimental/MeshModelingToolsetExp/Source/MeshModelingToolsExp/Public/DeformMeshPolygonsTool.h
@@ -146,6 +146,8 @@ class MESHMODELINGTOOLSEXP_API UDeformMeshPolygonsTool : public UMeshSurfacePoin
 {
 	GENERATED_BODY()
 
+	using FRay3d = UE::Geometry::FRay3d;	// allows usage of FRay3d in cpp without fully-qualified name, to be removed after LWC conversion
+
 public:
 	UDeformMeshPolygonsTool();
 
diff --git a/Engine/Plugins/Experimental/MeshModelingToolsetExp/Source/ModelingComponents/Private/Mechanics/DragAlignmentMechanic.cpp b/Engine/Plugins/Experimental/MeshModelingToolsetExp/Source/ModelingComponents/Private/Mechanics/DragAlignmentMechanic.cpp
--- a/Engine/Plugins/Experimental/MeshModelingToolsetExp/Source/ModelingComponents/Private/Mechanics/DragAlignmentMechanic.cpp
+++ b/Engine/Plugins/Experimental/MeshModelingToolsetExp/Source/ModelingComponents/Private/Mechanics/DragAlignmentMechanic.cpp
@@ -198,7 +198,7 @@ void UDragAlignmentMechanic::InitializeDeformedMeshRayCast(TFunction<FDynamicMes
 	MeshRayCast = [this, GetSpatialIn, TargetTransform, LinearDeformer](const FRay& WorldRay, FHitResult& HitResult, bool bUseFilter)
 	{
 		// Transform ray into local space
-		FRay3d LocalRay(TargetTransform.InverseTransformPosition((FVector3d)WorldRay.Origin),
+		UE::Geometry::FRay3d LocalRay(TargetTransform.InverseTransformPosition((FVector3d)WorldRay.Origin),
 			TargetTransform.InverseTransformVector((FVector3d)WorldRay.Direction));
 		UE::Geometry::Normalize(LocalRay.Direction);
 
diff --git a/Engine/Plugins/Experimental/MeshModelingToolsetExp/Source/ModelingComponents/Private/Mechanics/PlaneDistanceFromHitMechanic.cpp b/Engine/Plugins/Experimental/MeshModelingToolsetExp/Source/ModelingComponents/Private/Mechanics/PlaneDistanceFromHitMechanic.cpp
--- a/Engine/Plugins/Experimental/MeshModelingToolsetExp/Source/ModelingComponents/Private/Mechanics/PlaneDistanceFromHitMechanic.cpp
+++ b/Engine/Plugins/Experimental/MeshModelingToolsetExp/Source/ModelingComponents/Private/Mechanics/PlaneDistanceFromHitMechanic.cpp
@@ -53,7 +53,7 @@ void UPlaneDistanceFromHitMechanic::UpdateCurrentDistance(const FRay& WorldRay)
 	bCurrentHitIsWorldHit = false;
 
 	// cast ray at target object
-	FRay3d LocalRay = PreviewHeightFrame.ToFrame(WorldRay);
+	UE::Geometry::FRay3d LocalRay = PreviewHeightFrame.ToFrame(WorldRay);
 	int HitTID = PreviewHeightTargetAABB.FindNearestHitTriangle(LocalRay);
 	if (HitTID >= 0)
 	{
@@ -95,7 +95,7 @@ void UPlaneDistanceFromHitMechanic::UpdateCurrentDistance(const FRay& WorldRay)
 
 	if (bFoundHit == false && bFallbackToLineAxisPoint)
 	{
-		FDistLine3Ray3d Distance(FLine3d(PreviewHeightFrame.Origin, PreviewHeightFrame.Z()), FRay3d(WorldRay));
+		FDistLine3Ray3d Distance(FLine3d(PreviewHeightFrame.Origin, PreviewHeightFrame.Z()), UE::Geometry::FRay3d(WorldRay));
 		float DistSqr = Distance.GetSquared();
 		float WorldHitDist = WorldRay.GetParameter( (FVector)Distance.RayClosestPoint );
 		NearestHitFrameWorld = FFrame3d(Distance.RayClosestPoint, UE::Geometry::Normalized(Distance.RayClosestPoint - Distance.LineClosestPoint) );
diff --git a/Engine/Plugins/Experimental/MeshModelingToolsetExp/Source/ModelingComponents/Private/Mechanics/SpatialCurveDistanceMechanic.cpp b/Engine/Plugins/Experimental/MeshModelingToolsetExp/Source/ModelingComponents/Private/Mechanics/SpatialCurveDistanceMechanic.cpp
--- a/Engine/Plugins/Experimental/MeshModelingToolsetExp/Source/ModelingComponents/Private/Mechanics/SpatialCurveDistanceMechanic.cpp
+++ b/Engine/Plugins/Experimental/MeshModelingToolsetExp/Source/ModelingComponents/Private/Mechanics/SpatialCurveDistanceMechanic.cpp
@@ -39,7 +39,7 @@ void USpatialCurveDistanceMechanic::UpdateCurrentDistance(const FRay& WorldRay)
 	FVector3d RayNearest;
 	FVector3d CurveNearest;
 
-	FRay3d Ray(WorldRay);
+	UE::Geometry::FRay3d Ray(WorldRay);
 	int NumV = Curve.Num() - 1;
 	double MinDistSqr = TNumericLimits<double>::Max();
 	for (int k = 0; k < NumV; ++k)
diff --git a/Engine/Plugins/Experimental/MeshModelingToolsetExp/Source/ModelingComponents/Private/Scene/MeshSceneAdapter.cpp b/Engine/Plugins/Experimental/MeshModelingToolsetExp/Source/ModelingComponents/Private/Scene/MeshSceneAdapter.cpp
--- a/Engine/Plugins/Experimental/MeshModelingToolsetExp/Source/ModelingComponents/Private/Scene/MeshSceneAdapter.cpp
+++ b/Engine/Plugins/Experimental/MeshModelingToolsetExp/Source/ModelingComponents/Private/Scene/MeshSceneAdapter.cpp
@@ -37,6 +37,12 @@ static TAutoConsoleVariable<int32> CVarMeshSceneAdapterDisableMultiThreading(
 	TEXT("Determines whether or not to use multi-threading in MeshSceneAdapter.\n"));
 
 
+namespace UE
+{
+namespace Geometry
+{
+
+
 /** 
  * Compute the bounds of the vertices of Mesh, under 3D transformation TransformFunc
  * @return computed bounding box
@@ -764,6 +770,9 @@ static void CollectComponentChildMeshes(UActorComponent* Component, FActorAdapte
 }
 
 
+}	// end namespace UE
+}	// end namespace Geometry
+
 
 void FMeshSceneAdapter::AddActors(const TArray<AActor*>& ActorsSetIn)
 {
diff --git a/Engine/Plugins/Experimental/MeshModelingToolsetExp/Source/ModelingComponents/Private/Selection/GroupTopologySelector.cpp b/Engine/Plugins/Experimental/MeshModelingToolsetExp/Source/ModelingComponents/Private/Selection/GroupTopologySelector.cpp
--- a/Engine/Plugins/Experimental/MeshModelingToolsetExp/Source/ModelingComponents/Private/Selection/GroupTopologySelector.cpp
+++ b/Engine/Plugins/Experimental/MeshModelingToolsetExp/Source/ModelingComponents/Private/Selection/GroupTopologySelector.cpp
@@ -578,7 +578,7 @@ bool IsOccluded(const FGeometrySet3::FNearest& ClosestElement, const FVector3d&
 bool IsOccluded(const FVector3d& Point, const FVector3d& ViewOrigin, const FDynamicMeshAABBTree3* Spatial)
 {
 	// Shoot ray backwards to see if we hit something. 
-	FRay3d ToEyeRay(Point, UE::Geometry::Normalized(ViewOrigin - Point), true);
+	UE::Geometry::FRay3d ToEyeRay(Point, UE::Geometry::Normalized(ViewOrigin - Point), true);
 	ToEyeRay.Origin += (double)(100 * FMathf::ZeroTolerance) * ToEyeRay.Direction;
 	if (Spatial->FindNearestHitTriangle(ToEyeRay) >= 0)
 	{
diff --git a/Engine/Plugins/Experimental/MeshModelingToolsetExp/Source/ModelingComponents/Public/Selection/PolygonSelectionMechanic.h b/Engine/Plugins/Experimental/MeshModelingToolsetExp/Source/ModelingComponents/Public/Selection/PolygonSelectionMechanic.h
--- a/Engine/Plugins/Experimental/MeshModelingToolsetExp/Source/ModelingComponents/Public/Selection/PolygonSelectionMechanic.h
+++ b/Engine/Plugins/Experimental/MeshModelingToolsetExp/Source/ModelingComponents/Public/Selection/PolygonSelectionMechanic.h
@@ -84,6 +84,7 @@ class MODELINGCOMPONENTS_API UPolygonSelectionMechanic : public UInteractionMech
 	using FFrame3d = UE::Geometry::FFrame3d;
 	using FTransform3d = UE::Geometry::FTransform3d;	
 	using FAxisAlignedBox3d = UE::Geometry::FAxisAlignedBox3d;
+	using FRay3d = UE::Geometry::FRay3d;
 public:
 
 	virtual ~UPolygonSelectionMechanic();
diff --git a/Engine/Plugins/Runtime/GeometryProcessing/Source/DynamicMesh/Private/Sampling/MeshImageBakingCache.cpp b/Engine/Plugins/Runtime/GeometryProcessing/Source/DynamicMesh/Private/Sampling/MeshImageBakingCache.cpp
--- a/Engine/Plugins/Runtime/GeometryProcessing/Source/DynamicMesh/Private/Sampling/MeshImageBakingCache.cpp
+++ b/Engine/Plugins/Runtime/GeometryProcessing/Source/DynamicMesh/Private/Sampling/MeshImageBakingCache.cpp
@@ -117,7 +117,10 @@ void FMeshImageBakingCache::InvalidateOccupancy()
 
 
 
-
+namespace UE
+{
+namespace Geometry
+{
 
 
 /**
@@ -251,6 +254,8 @@ static bool GetDetailTrianglePoint_Nearest(
 }
 
 
+}
+}
 
 
 bool FMeshImageBakingCache::ValidateCache()
diff --git a/Engine/Plugins/Runtime/GeometryProcessing/Source/DynamicMesh/Private/ShapeApproximation/MeshSimpleShapeApproximation.cpp b/Engine/Plugins/Runtime/GeometryProcessing/Source/DynamicMesh/Private/ShapeApproximation/MeshSimpleShapeApproximation.cpp
--- a/Engine/Plugins/Runtime/GeometryProcessing/Source/DynamicMesh/Private/ShapeApproximation/MeshSimpleShapeApproximation.cpp
+++ b/Engine/Plugins/Runtime/GeometryProcessing/Source/DynamicMesh/Private/ShapeApproximation/MeshSimpleShapeApproximation.cpp
@@ -88,6 +88,8 @@ bool FMeshSimpleShapeApproximation::GetDetectedSimpleShape(
 
 
 
+namespace UE {
+namespace Geometry {
 
 struct FSimpleShapeFitsResult
 {
@@ -98,14 +100,13 @@ struct FSimpleShapeFitsResult
 	FOrientedBox3d Box;
 
 	bool bHaveCapsule = false;
-	UE::Geometry::FCapsule3d Capsule;
+	FCapsule3d Capsule;
 
 	bool bHaveConvex = false;
 	FDynamicMesh3 Convex;
 };
 
 
-
 static void ComputeSimpleShapeFits(const FDynamicMesh3& Mesh,
 									bool bSphere, bool bBox, bool bCapsule, bool bConvex, bool bUseExactComputationForBox, 
 								   FSimpleShapeFitsResult& FitResult,
@@ -169,6 +170,9 @@ static void ComputeSimpleShapeFits(const FDynamicMesh3& Mesh,
 }
 
 
+}
+}
+
 
 
 
diff --git a/Engine/Plugins/Runtime/MeshModelingToolset/Source/MeshModelingTools/Private/AddPrimitiveTool.cpp b/Engine/Plugins/Runtime/MeshModelingToolset/Source/MeshModelingTools/Private/AddPrimitiveTool.cpp
--- a/Engine/Plugins/Runtime/MeshModelingToolset/Source/MeshModelingTools/Private/AddPrimitiveTool.cpp
+++ b/Engine/Plugins/Runtime/MeshModelingToolset/Source/MeshModelingTools/Private/AddPrimitiveTool.cpp
@@ -391,7 +391,7 @@ void UAddBoxPrimitiveTool::GenerateMesh(FDynamicMesh3* OutMesh) const
 {
 	FGridBoxMeshGenerator BoxGen;
 	auto* BoxSettings = Cast<UProceduralBoxToolProperties>(ShapeSettings);
-	BoxGen.Box = FOrientedBox3d(FVector3d::Zero(), 0.5*FVector3d(BoxSettings->Depth, BoxSettings->Width, BoxSettings->Height));
+	BoxGen.Box = UE::Geometry::FOrientedBox3d(FVector3d::Zero(), 0.5*FVector3d(BoxSettings->Depth, BoxSettings->Width, BoxSettings->Height));
 	BoxGen.EdgeVertices = FIndex3i(BoxSettings->DepthSubdivisions + 1,
 								   BoxSettings->WidthSubdivisions + 1,
 								   BoxSettings->HeightSubdivisions + 1);
@@ -625,7 +625,7 @@ void UAddSpherePrimitiveTool::GenerateMesh(FDynamicMesh3* OutMesh) const
 	{
 		FBoxSphereGenerator SphereGen;
 		SphereGen.Radius = SphereSettings->Radius;
-		SphereGen.Box = FOrientedBox3d(FVector3d::Zero(),
+		SphereGen.Box = UE::Geometry::FOrientedBox3d(FVector3d::Zero(),
 			0.5 * FVector3d(SphereSettings->Subdivisions + 1,
 				SphereSettings->Subdivisions + 1,
 				SphereSettings->Subdivisions + 1));
diff --git a/Engine/Source/Runtime/GeometryCore/Public/ExplicitUseGeometryMathTypes.h b/Engine/Source/Runtime/GeometryCore/Public/ExplicitUseGeometryMathTypes.h
--- a/Engine/Source/Runtime/GeometryCore/Public/ExplicitUseGeometryMathTypes.h
+++ b/Engine/Source/Runtime/GeometryCore/Public/ExplicitUseGeometryMathTypes.h
@@ -94,17 +94,17 @@ using UE::Geometry::FSegment2d;
 using UE::Geometry::FSegment3f;
 using UE::Geometry::FSegment3d;
 
-using UE::Geometry::FRay3f;
-using UE::Geometry::FRay3d;
+//using UE::Geometry::FRay3f;
+//using UE::Geometry::FRay3d;
 
-using UE::Geometry::FSphere3d;
-using UE::Geometry::FSphere3f;
+//using UE::Geometry::FSphere3d;
+//using UE::Geometry::FSphere3f;
 
 using UE::Geometry::FCapsule3d;
 using UE::Geometry::FCapsule3f;
 
-using UE::Geometry::FOrientedBox3d;
-using UE::Geometry::FOrientedBox3f;
+//using UE::Geometry::FOrientedBox3d;
+//using UE::Geometry::FOrientedBox3f;
 
 using UE::Geometry::FTriangle2d;
 using UE::Geometry::FTriangle2f;
diff --git a/Engine/Source/Runtime/GeometryFramework/Private/UDynamicMesh.cpp b/Engine/Source/Runtime/GeometryFramework/Private/UDynamicMesh.cpp
--- a/Engine/Source/Runtime/GeometryFramework/Private/UDynamicMesh.cpp
+++ b/Engine/Source/Runtime/GeometryFramework/Private/UDynamicMesh.cpp
@@ -61,7 +61,7 @@ void UDynamicMesh::ResetToCube()
 	{
 		EditMesh.Clear();
 		FMinimalBoxMeshGenerator BoxGen;
-		BoxGen.Box = FOrientedBox3d(FVector3d::Zero(), 50.0 * FVector3d::One());
+		BoxGen.Box = UE::Geometry::FOrientedBox3d(FVector3d::Zero(), 50.0 * FVector3d::One());
 		EditMesh = FDynamicMesh3(&BoxGen.Generate());
 		EditMesh.EnableTriangleGroups();
 		EditMesh.Attributes()->EnableMaterialID();
@@ -525,7 +525,7 @@ void UDynamicMesh::ImportCustomProperties(const TCHAR* SourceText, FFeedbackCont
 		// if we got here we failed. Rather than produce an empty mesh, we generate a small cube
 		UE_LOG(LogTemp, Warning, TEXT("UDynamicMesh text-based property serialization incomplete, generating box as placeholder. Try increasing geometry.DynamicMesh.TextBasedDupeTriThreshold, or geometry.DynamicMesh.DupeStashTimeout."))
 		FMinimalBoxMeshGenerator BoxGen;
-		BoxGen.Box = FOrientedBox3d(FVector3d::Zero(), 50.0 * FVector3d::One());
+		BoxGen.Box = UE::Geometry::FOrientedBox3d(FVector3d::Zero(), 50.0 * FVector3d::One());
 		FDynamicMesh3 GenMesh(&BoxGen.Generate());
 		GenMesh.Attributes()->EnableMaterialID();
 		SetMesh(MoveTemp(GenMesh));
