MeshModelingTools: Fix crash when inset is used on a selection with no borders.

#rnx

diff --git a/Engine/Plugins/Experimental/MeshModelingToolsetExp/Source/MeshModelingToolsExp/Private/ToolActivities/PolyEditInsetOutsetActivity.cpp b/Engine/Plugins/Experimental/MeshModelingToolsetExp/Source/MeshModelingToolsExp/Private/ToolActivities/PolyEditInsetOutsetActivity.cpp
--- a/Engine/Plugins/Experimental/MeshModelingToolsetExp/Source/MeshModelingToolsExp/Private/ToolActivities/PolyEditInsetOutsetActivity.cpp
+++ b/Engine/Plugins/Experimental/MeshModelingToolsetExp/Source/MeshModelingToolsExp/Private/ToolActivities/PolyEditInsetOutsetActivity.cpp
@@ -18,6 +18,7 @@
 #include "ToolActivities/PolyEditActivityUtil.h"
 #include "ToolSceneQueriesUtil.h"
 
+#include "ExplicitUseGeometryMathTypes.h" // using UE::Geometry::(math types)
 using namespace UE::Geometry;
 
 #define LOCTEXT_NAMESPACE "UPolyEditInsetOutsetActivity"
@@ -67,13 +68,16 @@ EToolActivityStartResult UPolyEditInsetOutsetActivity::Start()
 	if (!CanStart())
 	{
 		ParentTool->GetToolManager()->DisplayMessage(
-			LOCTEXT("OnInsetOutsetFailedMesssage", "Cannot inset or outset without face selection."),
+			LOCTEXT("InsetOutsetNoSelectionMesssage", "Cannot inset or outset without face selection."),
 			EToolMessageLevel::UserWarning);
 		return EToolActivityStartResult::FailedStart;
 	}
 
 	Clear();
-	BeginInset();
+	if (!BeginInset())
+	{
+		return EToolActivityStartResult::FailedStart;
+	}
 	bIsRunning = true;
 
 	ActivityContext->EmitActivityStart(LOCTEXT("BeginInsetOutsetActivity", "Begin Inset/Outset"));
@@ -83,7 +87,7 @@ EToolActivityStartResult UPolyEditInsetOutsetActivity::Start()
 
 bool UPolyEditInsetOutsetActivity::CanAccept() const
 {
-	return false;
+	return true;
 }
 
 EToolActivityEndResult UPolyEditInsetOutsetActivity::End(EToolShutdownType ShutdownType)
@@ -109,7 +113,7 @@ EToolActivityEndResult UPolyEditInsetOutsetActivity::End(EToolShutdownType Shutd
 	}
 }
 
-void UPolyEditInsetOutsetActivity::BeginInset()
+bool UPolyEditInsetOutsetActivity::BeginInset()
 {
 	const FGroupTopologySelection& ActiveSelection = ActivityContext->SelectionMechanic->GetActiveSelection();
 	TArray<int32> ActiveTriangleSelection;
@@ -136,6 +140,13 @@ void UPolyEditInsetOutsetActivity::BeginInset()
 	CurveDistMechanic->CurrentDistance = 1.0f;  // initialize to something non-zero...prob should be based on polygon bounds maybe?
 
 	FMeshBoundaryLoops Loops(&InsetHitTargetMesh);
+	if (Loops.Loops.Num() == 0)
+	{
+		ParentTool->GetToolManager()->DisplayMessage(
+			LOCTEXT("InsetOutsetNoBorderMesssage", "Cannot inset or outset when selection has no border."),
+			EToolMessageLevel::UserWarning);
+		return false;
+	}
 	TArray<FVector3d> LoopVertices;
 	Loops.Loops[0].GetVertices(LoopVertices);
 	CurveDistMechanic->InitializePolyLoop(LoopVertices, UE::Geometry::FTransform3d::Identity());
@@ -149,6 +160,8 @@ void UPolyEditInsetOutsetActivity::BeginInset()
 	}
 
 	bPreviewUpdatePending = true;
+
+	return true;
 }
 
 void UPolyEditInsetOutsetActivity::ApplyInset()
diff --git a/Engine/Plugins/Experimental/MeshModelingToolsetExp/Source/MeshModelingToolsExp/Public/ToolActivities/PolyEditInsetOutsetActivity.h b/Engine/Plugins/Experimental/MeshModelingToolsetExp/Source/MeshModelingToolsExp/Public/ToolActivities/PolyEditInsetOutsetActivity.h
--- a/Engine/Plugins/Experimental/MeshModelingToolsetExp/Source/MeshModelingToolsExp/Public/ToolActivities/PolyEditInsetOutsetActivity.h
+++ b/Engine/Plugins/Experimental/MeshModelingToolsetExp/Source/MeshModelingToolsExp/Public/ToolActivities/PolyEditInsetOutsetActivity.h
@@ -83,7 +83,7 @@ public:
 
 protected:
 	void Clear();
-	void BeginInset();
+	bool BeginInset();
 	void ApplyInset();
 
 	bool bIsRunning = false;
