MeshModelingTools: Add boolean extrude option to PolyEdit. Fix extrusion of multiple faces via triangle normals. Add ability to extrude while keeping multiple faces parallel. Add options for side group setting on borders during extrude.

#rnx

diff --git a/Engine/Plugins/Experimental/MeshModelingToolsetExp/Source/MeshModelingToolsExp/Private/ToolActivities/PolyEditExtrudeActivity.cpp b/Engine/Plugins/Experimental/MeshModelingToolsetExp/Source/MeshModelingToolsExp/Private/ToolActivities/PolyEditExtrudeActivity.cpp
--- a/Engine/Plugins/Experimental/MeshModelingToolsetExp/Source/MeshModelingToolsExp/Private/ToolActivities/PolyEditExtrudeActivity.cpp
+++ b/Engine/Plugins/Experimental/MeshModelingToolsetExp/Source/MeshModelingToolsExp/Private/ToolActivities/PolyEditExtrudeActivity.cpp
@@ -5,8 +5,10 @@
 #include "BaseBehaviors/SingleClickBehavior.h"
 #include "BaseBehaviors/MouseHoverBehavior.h"
 #include "ContextObjectStore.h"
+#include "DeformationOps/ExtrudeOp.h"
 #include "Drawing/PolyEditPreviewMesh.h"
 #include "DynamicMesh/MeshNormals.h"
+#include "DynamicMesh/MeshTransforms.h"
 #include "EditMeshPolygonsTool.h"
 #include "InteractiveToolManager.h"
 #include "Mechanics/PlaneDistanceFromHitMechanic.h"
@@ -19,8 +21,44 @@
 
 #define LOCTEXT_NAMESPACE "UPolyEditExtrudeActivity"
 
+#include "ExplicitUseGeometryMathTypes.h" // using UE::Geometry::(math types)
 using namespace UE::Geometry;
 
+namespace PolyEditExtrudeActivityLocals
+{
+	/** Creates a mesh shared ptr out of the given triangles */
+	TSharedPtr<FDynamicMesh3> CreatePatchMesh(const FDynamicMesh3* SourceMesh, TArray<int32> Triangles)
+	{
+		FDynamicSubmesh3 Submesh(SourceMesh, Triangles, 
+			(int32)(EMeshComponents::FaceGroups) | (int32)(EMeshComponents::VertexNormals), true);
+		return MakeShared<FDynamicMesh3>(MoveTemp(Submesh.GetSubmesh()));
+	}
+}
+
+TUniquePtr<FDynamicMeshOperator> UPolyEditExtrudeActivity::MakeNewOperator()
+{
+	using UE::Geometry::FTransform3d;
+
+	TUniquePtr<FExtrudeOp> Op = MakeUnique<FExtrudeOp>();
+	Op->OriginalMesh = ActivityContext->CurrentMesh;
+	Op->ExtrudeMode = ExtrudeProperties->ExtrudeMode;
+	Op->DirectionMode = ExtrudeProperties->DirectionMode;
+
+	ActivityContext->CurrentTopology->GetSelectedTriangles(ActiveSelection, Op->TriangleSelection);
+
+	Op->ExtrudeDistance = ExtrudeHeightMechanic->CurrentHeight;
+	Op->UVScaleFactor = UVScaleFactor;
+
+	FTransform3d WorldTransform(ActivityContext->Preview->PreviewMesh->GetTransform());
+	Op->SetResultTransform(WorldTransform);
+	
+	Op->MeshSpaceExtrudeDirection = WorldTransform.InverseTransformVector(GetExtrudeDirection());
+	Op->bShellsToSolids = ExtrudeProperties->bShellsToSolids;
+	Op->bUseColinearityForSettingBorderGroups = ExtrudeProperties->bUseColinearityForSettingBorderGroups;
+
+	return Op;
+}
+
 void UPolyEditExtrudeActivity::Setup(UInteractiveTool* ParentToolIn)
 {
 	Super::Setup(ParentToolIn);
@@ -29,15 +67,28 @@ void UPolyEditExtrudeActivity::Setup(UInteractiveTool* ParentToolIn)
 	ExtrudeProperties->RestoreProperties(ParentTool.Get());
 	AddToolPropertySource(ExtrudeProperties);
 	SetToolPropertySourceEnabled(ExtrudeProperties, false);
-	ExtrudeProperties->WatchProperty(ExtrudeProperties->Direction,
+	ExtrudeProperties->WatchProperty(ExtrudeProperties->MeasureDirection,
 		[this](EPolyEditExtrudeDirection) { 
-			Clear();
-			BeginExtrude();
+			if (bIsRunning)
+			{
+				ReinitializeExtrudeHeightMechanic();
+				ActivityContext->Preview->InvalidateResult();
+			}
+		});
+	ExtrudeProperties->WatchProperty(ExtrudeProperties->DirectionMode,
+		[this](EPolyEditExtrudeDirectionMode) {
+			if (bIsRunning)
+			{
+				ReinitializeExtrudeHeightMechanic();
+				ActivityContext->Preview->InvalidateResult();
+			}
 		});
 	ExtrudeProperties->WatchProperty(ExtrudeProperties->ExtrudeMode,
 		[this](EPolyEditExtrudeMode) {
-			Clear();
-			BeginExtrude();
+			if (bIsRunning)
+			{
+				ActivityContext->Preview->InvalidateResult();
+			}
 		});
 
 	// Register ourselves to receive clicks and hover
@@ -53,7 +104,11 @@ void UPolyEditExtrudeActivity::Setup(UInteractiveTool* ParentToolIn)
 
 void UPolyEditExtrudeActivity::Shutdown(EToolShutdownType ShutdownType)
 {
-	Clear();
+	if (bIsRunning)
+	{
+		End(ShutdownType);
+	}
+
 	ExtrudeProperties->SaveProperties(ParentTool.Get());
 
 	ExtrudeProperties = nullptr;
@@ -81,166 +136,230 @@ EToolActivityStartResult UPolyEditExtrudeActivity::Start()
 		return EToolActivityStartResult::FailedStart;
 	}
 
-	Clear();
+	// Change the op we use in the preview to an extrude op.
+	ActivityContext->Preview->ChangeOpFactory(this);
+	ActivityContext->Preview->OnOpCompleted.AddWeakLambda(this,
+		[this](const UE::Geometry::FDynamicMeshOperator* UncastOp) {
+			const FExtrudeOp* Op = static_cast<const FExtrudeOp*>(UncastOp);
+			NewSelectionGids = Op->ExtrudedFaceNewGids;
+		});
+
+	// Temporarily clear the selection to avoid the highlighting on the preview as we extrude (which 
+	// sometimes highlights incorrect triangles in boolean mode). There's currently not a good way to
+	// disable those secondary triangle buffers without losing the filter function.
+	// The selection gets reset on End(), either to a new selection, or to the old one on cancel.
+	ActiveSelection = ActivityContext->SelectionMechanic->GetActiveSelection();
+	ActivityContext->SelectionMechanic->SetSelection(FGroupTopologySelection(), false);
+
+	SetToolPropertySourceEnabled(ExtrudeProperties, true);
+
+	// Set up the basics of the extrude height mechanic.
+	// TODO: Allow option to use a gizmo instead
+	ExtrudeHeightMechanic = NewObject<UPlaneDistanceFromHitMechanic>(this);
+	ExtrudeHeightMechanic->Setup(ParentTool.Get());
+	ExtrudeHeightMechanic->WorldHitQueryFunc = [this](const FRay& WorldRay, FHitResult& HitResult)
+	{
+		return ToolSceneQueriesUtil::FindNearestVisibleObjectHit(ActivityContext->Preview->GetWorld(),
+			HitResult, WorldRay);
+	};
+	ExtrudeHeightMechanic->WorldPointSnapFunc = [this](const FVector3d& WorldPos, FVector3d& SnapPos)
+	{
+		return ToolSceneQueriesUtil::FindWorldGridSnapPoint(ParentTool.Get(), WorldPos, SnapPos);
+	};
+	ExtrudeHeightMechanic->CurrentHeight = 1.0f; // Arbitrary intialization
+	
 	BeginExtrude();
+
+	bRequestedApply = false;
 	bIsRunning = true;
 
 	ActivityContext->EmitActivityStart(LOCTEXT("BeginExtrudeActivity", "Begin Extrude"));
+	return EToolActivityStartResult::Running;
+}
 
-	return EToolActivityStartResult();
+// Someday we may want to initialize one extrude after another in the same invocation. Anything that needs
+// to be reset in those cases goes here.
+void UPolyEditExtrudeActivity::BeginExtrude()
+{
+	using namespace PolyEditExtrudeActivityLocals;
+	using UE::Geometry::FTransform3d;
+
+	TArray<int32> ActiveTriangleSelection;
+	ActivityContext->CurrentTopology->GetSelectedTriangles(ActiveSelection, ActiveTriangleSelection);
+
+	PatchMesh = CreatePatchMesh(ActivityContext->CurrentMesh.Get(), ActiveTriangleSelection);
+	UE::Geometry::FDynamicMeshAABBTree3 PatchSpatial(PatchMesh.Get(), true);
+
+	// Get the world selection frame
+	FFrame3d ActiveSelectionFrameLocal = ActivityContext->CurrentTopology->GetSelectionFrame(ActiveSelection);
+	ActiveSelectionFrameWorld = ActiveSelectionFrameLocal;
+	ActiveSelectionFrameWorld.Origin = PatchSpatial.FindNearestPoint(ActiveSelectionFrameLocal.Origin);
+	FTransform3d WorldTransform(ActivityContext->Preview->PreviewMesh->GetTransform());
+	ActiveSelectionFrameWorld.Transform(WorldTransform);
+
+	ReinitializeExtrudeHeightMechanic();
+
+	float BoundsMaxDim = ActivityContext->CurrentMesh->GetBounds().MaxDim();
+	if (BoundsMaxDim > 0)
+	{
+		UVScaleFactor = 1.0 / BoundsMaxDim;
+	}
+}
+
+// The height mechanics has to get reinitialized whenever extrude direction changes
+void UPolyEditExtrudeActivity::ReinitializeExtrudeHeightMechanic()
+{
+	using UE::Geometry::FTransform3d;
+
+	ExtrusionFrameWorld = ActiveSelectionFrameWorld;
+	ExtrusionFrameWorld.AlignAxis(2, GetExtrudeDirection());
+
+	FDynamicMesh3 ExtrudeHitTargetMesh;	
+
+	// Make infinite-extent hit-test mesh to use in the mechanic. This might seem like something
+	// that would only be applicable in SingleDirection extrude mode, but it is in fact useful
+	// for the mechanic in general, as it tends to make the mouse interactions feel better around
+	// the extrude measurement line.
+	ExtrudeHitTargetMesh = *PatchMesh;
+	MeshTransforms::ApplyTransform(ExtrudeHitTargetMesh, 
+		(FTransform3d)ActivityContext->Preview->PreviewMesh->GetTransform());
+
+	double Length = 99999.0;
+	FVector3d ExtrudeDirection = GetExtrudeDirection();
+	MeshTransforms::Translate(ExtrudeHitTargetMesh, -Length * ExtrudeDirection);
+
+	FOffsetMeshRegion Extruder(&ExtrudeHitTargetMesh);
+	Extruder.OffsetPositionFunc = [Length, ExtrudeDirection](const FVector3d& Position,
+		const FVector3f& Normal, int VertexID)
+	{
+		return Position + 2.0 * Length * ExtrudeDirection;
+	};
+	Extruder.Triangles.Reserve(PatchMesh->TriangleCount());
+	for (int32 Tid : PatchMesh->TriangleIndicesItr())
+	{
+		Extruder.Triangles.Add(Tid);
+	}
+	Extruder.Apply();
+
+	ExtrudeHeightMechanic->Initialize(MoveTemp(ExtrudeHitTargetMesh), ExtrusionFrameWorld, true);
 }
 
 bool UPolyEditExtrudeActivity::CanAccept() const
 {
-	return false;
+	return true;
 }
 
 EToolActivityEndResult UPolyEditExtrudeActivity::End(EToolShutdownType ShutdownType)
 {
-	if (!bIsRunning)
+	if (!ensure(bIsRunning))
 	{
-		Clear();
+		EndInternal();
 		return EToolActivityEndResult::ErrorDuringEnd;
 	}
 
+	// Remember to reset the selection. Even if we change it to something else when applying the
+	// extrude, the previous selection needs to be valid.
+	ActivityContext->SelectionMechanic->SetSelection(ActiveSelection, false);
+
 	if (ShutdownType == EToolShutdownType::Cancel)
 	{
-		Clear();
-		bIsRunning = false;
+		EndInternal();
+
+		// Reset the preview
+		ActivityContext->Preview->PreviewMesh->UpdatePreview(ActivityContext->CurrentMesh.Get());
+
 		return EToolActivityEndResult::Cancelled;
 	}
 	else
 	{
+		// Stop the current compute if there is one.
+		ActivityContext->Preview->CancelCompute();
+
+		// Apply whatever we happen to have at the moment.
 		ApplyExtrude();
-		Clear();
-		bIsRunning = false;
+
+		EndInternal();
 		return EToolActivityEndResult::Completed;
 	}
 }
 
-
-void UPolyEditExtrudeActivity::BeginExtrude()
+// Does whatever kind of cleanup we need to do to end the activity.
+void UPolyEditExtrudeActivity::EndInternal()
 {
-	const FGroupTopologySelection& ActiveSelection = ActivityContext->SelectionMechanic->GetActiveSelection();
-	TArray<int32> ActiveTriangleSelection;
-	ActivityContext->CurrentTopology->GetSelectedTriangles(ActiveSelection, ActiveTriangleSelection);
-
-	UE::Geometry::FTransform3d WorldTransform(ActivityContext->Preview->PreviewMesh->GetTransform());
-
-	// Get the world frame
-	FFrame3d ActiveSelectionFrameLocal = ActivityContext->CurrentTopology->GetSelectionFrame(ActiveSelection);
-	ActiveSelectionFrameWorld = ActiveSelectionFrameLocal;
-	ActiveSelectionFrameWorld.Transform(WorldTransform);
-	ActiveSelectionFrameWorld.AlignAxis(2, GetExtrudeDirection());
-
-	// Set up a preview of the extruded portion of the mesh
-	EditPreview = PolyEditActivityUtil::CreatePolyEditPreviewMesh(*ParentTool, *ActivityContext);
-	EditPreview->InitializeExtrudeType(ActivityContext->CurrentMesh.Get(), ActiveTriangleSelection, ActiveSelectionFrameWorld.Z(), &WorldTransform, true);
-	// move world extrude frame to point on surface
-	ActiveSelectionFrameWorld.Origin = EditPreview->GetInitialPatchMeshSpatial().FindNearestPoint(ActiveSelectionFrameWorld.Origin);
-
-	// Hide the selected triangles (that are being replaced by the extruded portion)
-	ActivityContext->Preview->PreviewMesh->SetSecondaryBuffersVisibility(false);
-
-	// Set up the mechanic we use to determine how far to extrude
-	ExtrudeHeightMechanic = NewObject<UPlaneDistanceFromHitMechanic>(this);
-	ExtrudeHeightMechanic->Setup(ParentTool.Get());
-	ExtrudeHeightMechanic->WorldHitQueryFunc = [this](const FRay& WorldRay, FHitResult& HitResult)
-	{
-		return ToolSceneQueriesUtil::FindNearestVisibleObjectHit(ActivityContext->Preview->GetWorld(), HitResult, WorldRay);
-	};
-	ExtrudeHeightMechanic->WorldPointSnapFunc = [this](const FVector3d& WorldPos, FVector3d& SnapPos)
-	{
-		return ToolSceneQueriesUtil::FindWorldGridSnapPoint(ParentTool.Get(), WorldPos, SnapPos);
-	};
-	ExtrudeHeightMechanic->CurrentHeight = 1.0f;  // initialize to something non-zero...prob should be based on polygon bounds maybe?
-
-	// make inifinite-extent hit-test mesh to use in the mechanic
-	FDynamicMesh3 ExtrudeHitTargetMesh;
-	EditPreview->MakeExtrudeTypeHitTargetMesh(ExtrudeHitTargetMesh);
-	ExtrudeHeightMechanic->Initialize(MoveTemp(ExtrudeHitTargetMesh), ActiveSelectionFrameWorld, true);
-
-	SetToolPropertySourceEnabled(ExtrudeProperties, true);
-
-	float BoundsMaxDim = ActivityContext->CurrentMesh->GetBounds().MaxDim();
-	if (BoundsMaxDim > 0)
-	{
-		UVScaleFactor = 1.0 / BoundsMaxDim;
-	}
-
-	bPreviewUpdatePending = true;
+	ActivityContext->Preview->ClearOpFactory();
+	ActivityContext->Preview->OnOpCompleted.RemoveAll(this);
+	PatchMesh.Reset();
+	ExtrudeHeightMechanic->Shutdown();
+	ExtrudeHeightMechanic = nullptr;
+	SetToolPropertySourceEnabled(ExtrudeProperties, false);
+	bIsRunning = false;
 }
 
 void UPolyEditExtrudeActivity::ApplyExtrude()
 {
-	check(ExtrudeHeightMechanic != nullptr && EditPreview != nullptr);
-
-	const FGroupTopologySelection& ActiveSelection = ActivityContext->SelectionMechanic->GetActiveSelection();
-	TArray<int32> ActiveTriangleSelection;
-	ActivityContext->CurrentTopology->GetSelectedTriangles(ActiveSelection, ActiveTriangleSelection);
+	check(ExtrudeHeightMechanic != nullptr);
+
+	const FDynamicMesh3* ResultMesh = ActivityContext->Preview->PreviewMesh->GetMesh();
+
+	// Prep for undo.
+	// TODO: It would be nice if we could attach the change tracker to the op and have it actually
+	// track the changed triangles instead of us having to know which ones to save here. This is
+	// particularly true for the boolean case, where we conservatively save all the triangles,
+	// even though we only need the ones that got cut/deleted. Unfortunatley our boolean code
+	// doesn't currently track changed triangles.
+	FDynamicMeshChangeTracker ChangeTracker(ActivityContext->CurrentMesh.Get());
+	ChangeTracker.BeginChange();
+	if (ExtrudeProperties->ExtrudeMode == EPolyEditExtrudeMode::MoveAndStitch)
+	{
+		TArray<int32> ActiveTriangleSelection;
+		ActivityContext->CurrentTopology->GetSelectedTriangles(
+			ActiveSelection, ActiveTriangleSelection);
 
-	UE::Geometry::FTransform3d WorldTransform(ActivityContext->Preview->PreviewMesh->GetTransform());
-	FVector3d MeshSpaceExtrudeDirection = WorldTransform.InverseTransformVector(ActiveSelectionFrameWorld.Z());
-	double ExtrudeDistance = ExtrudeHeightMechanic->CurrentHeight;
-
-	FOffsetMeshRegion Extruder(ActivityContext->CurrentMesh.Get());
-	Extruder.UVScaleFactor = UVScaleFactor;
-	Extruder.Triangles = ActiveTriangleSelection;
-	TSet<int32> TriangleSet(ActiveTriangleSelection);
-	bool bUseNormals = (ExtrudeProperties->ExtrudeMode != EPolyEditExtrudeMode::SingleDirection);
-	Extruder.OffsetPositionFunc = [ExtrudeDistance, bUseNormals, &MeshSpaceExtrudeDirection](const FVector3d& Pos, const FVector3f& Normal, int32 VertexID) {
-		return Pos + ExtrudeDistance * (bUseNormals ? (FVector3d)Normal : MeshSpaceExtrudeDirection);
-	};
-	Extruder.bIsPositiveOffset = (ExtrudeDistance > 0);
-	Extruder.bUseFaceNormals = (ExtrudeProperties->ExtrudeMode == EPolyEditExtrudeMode::SelectedTriangleNormals);
-	Extruder.bOffsetFullComponentsAsSolids = ExtrudeProperties->bShellsToSolids;
-	Extruder.ChangeTracker = MakeUnique<FDynamicMeshChangeTracker>(ActivityContext->CurrentMesh.Get());
-	Extruder.ChangeTracker->BeginChange();
-	Extruder.Apply();
+		ChangeTracker.SaveTriangles(ActiveTriangleSelection, true /*bSaveVertices*/);
+	}
+	else // if boolean
+	{
+		TArray<int32> AllTids;
+		for (int32 Tid : ActivityContext->CurrentMesh->TriangleIndicesItr())
+		{
+			AllTids.Add(Tid);
+		}
+		ChangeTracker.SaveTriangles(AllTids, true /*bSaveVertices*/);
+	}
 
-	FMeshNormals::QuickComputeVertexNormalsForTriangles(*ActivityContext->CurrentMesh, Extruder.AllModifiedTriangles);
+	// Update current mesh
+	ActivityContext->CurrentMesh->Copy(*ResultMesh,
+		true, true, true, true);
 
-	// construct new selection
+	// Construct new selection
 	FGroupTopologySelection NewSelection;
-	if (!ActivityContext->bTriangleMode)
+	if (ActivityContext->bTriangleMode)
 	{
-		for (const FOffsetMeshRegion::FOffsetInfo& Info : Extruder.OffsetRegions)
+		TSet<int32> GidSet(NewSelectionGids);
+		for (int32 Tid : ResultMesh->TriangleIndicesItr())
 		{
-			NewSelection.SelectedGroupIDs.Append(Info.OffsetGroups);
+			if (GidSet.Contains(ResultMesh->GetTriangleGroup(Tid)))
+			{
+				NewSelection.SelectedGroupIDs.Add(Tid);
+			}
 		}
 	}
 	else
 	{
-		for (const FOffsetMeshRegion::FOffsetInfo& Info : Extruder.OffsetRegions)
-		{
-			NewSelection.SelectedGroupIDs.Append(Info.InitialTriangles);
-		}
+		NewSelection.SelectedGroupIDs.Append(NewSelectionGids);
 	}
 
 	// Emit undo  (also updates relevant structures)
 	ActivityContext->EmitCurrentMeshChangeAndUpdate(LOCTEXT("PolyMeshExtrudeChange", "Extrude"),
-		Extruder.ChangeTracker->EndChange(), NewSelection, true);
-}
-
-void UPolyEditExtrudeActivity::Clear()
-{
-	if (EditPreview != nullptr)
-	{
-		EditPreview->Disconnect();
-		EditPreview = nullptr;
-	}
-
-	ActivityContext->Preview->PreviewMesh->SetSecondaryBuffersVisibility(true);
-
-	ExtrudeHeightMechanic = nullptr;
-	SetToolPropertySourceEnabled(ExtrudeProperties, false);
+		ChangeTracker.EndChange(), NewSelection, true);
 }
 
-
+// Gets used to set the direction along which we measure the distance to extrude.
 FVector3d UPolyEditExtrudeActivity::GetExtrudeDirection() const
 {
-	using FTransform3d = UE::Geometry::FTransform3d;
-	switch (ExtrudeProperties->Direction)
+	using UE::Geometry::FTransform3d;
+
+	switch (ExtrudeProperties->MeasureDirection)
 	{
 	default:
 	case EPolyEditExtrudeDirection::SelectionNormal:
@@ -263,7 +382,7 @@ FVector3d UPolyEditExtrudeActivity::GetExtrudeDirection() const
 
 void UPolyEditExtrudeActivity::Render(IToolsContextRenderAPI* RenderAPI)
 {
-	if (ExtrudeHeightMechanic != nullptr)
+	if (ensure(bIsRunning) && ExtrudeHeightMechanic != nullptr)
 	{
 		ExtrudeHeightMechanic->Render(RenderAPI);
 	}
@@ -271,22 +390,20 @@ void UPolyEditExtrudeActivity::Render(IToolsContextRenderAPI* RenderAPI)
 
 void UPolyEditExtrudeActivity::Tick(float DeltaTime)
 {
-	if (EditPreview && bPreviewUpdatePending)
+	if (!ensure(bIsRunning))
 	{
-		switch (ExtrudeProperties->ExtrudeMode)
-		{
-		case EPolyEditExtrudeMode::SingleDirection:
-			EditPreview->UpdateExtrudeType(ExtrudeHeightMechanic->CurrentHeight);
-			break;
-		case EPolyEditExtrudeMode::SelectedTriangleNormals:
-			EditPreview->UpdateExtrudeType_FaceNormalAvg(ExtrudeHeightMechanic->CurrentHeight);
-			break;
-		case EPolyEditExtrudeMode::VertexNormals:
-			EditPreview->UpdateExtrudeType(ExtrudeHeightMechanic->CurrentHeight, true);
-			break;
-		}
+		return;
+	}
 
-		bPreviewUpdatePending = false;
+	if (bRequestedApply && ActivityContext->Preview->HaveValidResult())
+	{
+		ApplyExtrude();
+		bRequestedApply = false;
+
+		// End activity
+		EndInternal();
+		Cast<IToolActivityHost>(ParentTool)->NotifyActivitySelfEnded(this);
+		return;
 	}
 }
 
@@ -299,14 +416,9 @@ FInputRayHit UPolyEditExtrudeActivity::IsHitByClick(const FInputDeviceRay& Click
 
 void UPolyEditExtrudeActivity::OnClicked(const FInputDeviceRay& ClickPos)
 {
-	if (bIsRunning)
+	if (bIsRunning && !bRequestedApply)
 	{
-		ApplyExtrude();
-
-		// End activity
-		Clear();
-		bIsRunning = false;
-		Cast<IToolActivityHost>(ParentTool)->NotifyActivitySelfEnded(this);
+		bRequestedApply = true;
 	}
 }
 
@@ -319,8 +431,11 @@ FInputRayHit UPolyEditExtrudeActivity::BeginHoverSequenceHitTest(const FInputDev
 
 bool UPolyEditExtrudeActivity::OnUpdateHover(const FInputDeviceRay& DevicePos)
 {
-	ExtrudeHeightMechanic->UpdateCurrentDistance(DevicePos.WorldRay);
-	bPreviewUpdatePending = true;
+	if (!bRequestedApply)
+	{
+		ExtrudeHeightMechanic->UpdateCurrentDistance(DevicePos.WorldRay);
+		ActivityContext->Preview->InvalidateResult();
+	}
 	return bIsRunning;
 }
 
diff --git a/Engine/Plugins/Experimental/MeshModelingToolsetExp/Source/MeshModelingToolsExp/Public/ToolActivities/PolyEditExtrudeActivity.h b/Engine/Plugins/Experimental/MeshModelingToolsetExp/Source/MeshModelingToolsExp/Public/ToolActivities/PolyEditExtrudeActivity.h
--- a/Engine/Plugins/Experimental/MeshModelingToolsetExp/Source/MeshModelingToolsExp/Public/ToolActivities/PolyEditExtrudeActivity.h
+++ b/Engine/Plugins/Experimental/MeshModelingToolsetExp/Source/MeshModelingToolsExp/Public/ToolActivities/PolyEditExtrudeActivity.h
@@ -5,6 +5,9 @@
 #include "CoreMinimal.h"
 
 #include "BaseBehaviors/BehaviorTargetInterfaces.h"
+#include "DeformationOps/ExtrudeOp.h" // EPolyEditExtrudeMode
+#include "GeometryBase.h"
+#include "GroupTopology.h" // FGroupTopologySelection
 #include "InteractiveToolActivity.h"
 #include "FrameTypes.h"
 
@@ -13,20 +16,7 @@
 class UPolyEditActivityContext;
 class UPolyEditPreviewMesh;
 class UPlaneDistanceFromHitMechanic;
-
-UENUM()
-enum class EPolyEditExtrudeMode
-{
-	//~ TODO: this is what we actually want, but not yet implemented.
-	//~ SelectedFaceNormals,
-
-	SingleDirection,
-
-	//~ These are likely not very useful, but they can sometimes take
-	//~ the place of SelectedFaceNormals, until we implement that.
-	SelectedTriangleNormals,
-	VertexNormals,
-};
+PREDECLARE_GEOMETRY(class FDynamicMesh3);
 
 UENUM()
 enum class EPolyEditExtrudeDirection
@@ -40,7 +30,6 @@ enum class EPolyEditExtrudeDirection
 	LocalZ
 };
 
-
 UCLASS()
 class MESHMODELINGTOOLSEXP_API UPolyEditExtrudeProperties : public UInteractiveToolPropertySet
 {
@@ -48,22 +37,38 @@ class MESHMODELINGTOOLSEXP_API UPolyEditExtrudeProperties : public UInteractiveT
 
 public:
 	UPROPERTY(EditAnywhere, Category = Extrude)
-	EPolyEditExtrudeMode ExtrudeMode = EPolyEditExtrudeMode::SingleDirection;
+	EPolyEditExtrudeMode ExtrudeMode = EPolyEditExtrudeMode::MoveAndStitch;
 
+	/** Which way to move vertices during the extrude */
 	UPROPERTY(EditAnywhere, Category = Extrude)
-	EPolyEditExtrudeDirection Direction = EPolyEditExtrudeDirection::SelectionNormal;
+	EPolyEditExtrudeDirectionMode DirectionMode = EPolyEditExtrudeDirectionMode::SelectedTriangleNormalsEven;
+
+	/** What axis to measure the extrusion distance along. When the direction mode is Single Direction, also controls the direction. */
+	UPROPERTY(EditAnywhere, Category = Extrude, AdvancedDisplay)
+	EPolyEditExtrudeDirection MeasureDirection = EPolyEditExtrudeDirection::SelectionNormal;
 
 	/** Controls whether extruding an entire open-border patch should create a solid or an open shell */
 	UPROPERTY(EditAnywhere, Category = Extrude)
 	bool bShellsToSolids = true;
+
+	/** 
+	 * When extruding regions that touch the mesh border, assign the side groups (groups on the 
+	 * stitched side of the extrude) in a way that considers edge colinearity. For instance, when
+	 * true, extruding a flat rectangle will give four different groups on its sides rather than
+	 * one connected group.
+	 */
+	UPROPERTY(EditAnywhere, Category = Extrude, AdvancedDisplay)
+	bool bUseColinearityForSettingBorderGroups = true;
 };
 
 /**
  * 
  */
 UCLASS()
-class MESHMODELINGTOOLSEXP_API UPolyEditExtrudeActivity : public UInteractiveToolActivity, 
+class MESHMODELINGTOOLSEXP_API UPolyEditExtrudeActivity : public UInteractiveToolActivity,
+	public UE::Geometry::IDynamicMeshOperatorFactory,
 	public IClickBehaviorTarget, public IHoverBehaviorTarget
+
 {
 	GENERATED_BODY()
 
@@ -79,6 +84,9 @@ public:
 	virtual void Render(IToolsContextRenderAPI* RenderAPI) override;
 	virtual void Tick(float DeltaTime) override;
 
+	// IDynamicMeshOperatorFactory
+	virtual TUniquePtr<UE::Geometry::FDynamicMeshOperator> MakeNewOperator() override;
+
 	// IClickBehaviorTarget API
 	virtual FInputRayHit IsHitByClick(const FInputDeviceRay& ClickPos) override;
 	virtual void OnClicked(const FInputDeviceRay& ClickPos) override;
@@ -94,12 +102,10 @@ public:
 
 protected:
 	FVector3d GetExtrudeDirection() const;
-	void BeginExtrude();
-	void ApplyExtrude();
-	void Clear();
-
-	UPROPERTY()
-	TObjectPtr<UPolyEditPreviewMesh> EditPreview;
+	virtual void BeginExtrude();
+	virtual void ApplyExtrude();
+	virtual void ReinitializeExtrudeHeightMechanic();
+	virtual void EndInternal();
 
 	UPROPERTY()
 	TObjectPtr<UPlaneDistanceFromHitMechanic> ExtrudeHeightMechanic = nullptr;
@@ -107,9 +113,15 @@ protected:
 	UPROPERTY()
 	TObjectPtr<UPolyEditActivityContext> ActivityContext;
 
+	TSharedPtr<UE::Geometry::FDynamicMesh3> PatchMesh;
+	TArray<int32> NewSelectionGids;
+
 	bool bIsRunning = false;
 
+	UE::Geometry::FGroupTopologySelection ActiveSelection;
 	UE::Geometry::FFrame3d ActiveSelectionFrameWorld;
+	UE::Geometry::FFrame3d ExtrusionFrameWorld;
 	float UVScaleFactor = 1.0f;
-	bool bPreviewUpdatePending = false;
+
+	bool bRequestedApply = false;
 };
diff --git a/Engine/Plugins/Experimental/MeshModelingToolsetExp/Source/ModelingComponents/Private/Scene/MeshSceneAdapter.cpp b/Engine/Plugins/Experimental/MeshModelingToolsetExp/Source/ModelingComponents/Private/Scene/MeshSceneAdapter.cpp
--- a/Engine/Plugins/Experimental/MeshModelingToolsetExp/Source/ModelingComponents/Private/Scene/MeshSceneAdapter.cpp
+++ b/Engine/Plugins/Experimental/MeshModelingToolsetExp/Source/ModelingComponents/Private/Scene/MeshSceneAdapter.cpp
@@ -1870,7 +1870,7 @@ void FMeshSceneAdapter::GenerateBaseClosingMesh(double BaseHeight, double Extrud
 		{
 			Offset.Triangles.Add(tid);
 		}
-		Offset.bUseFaceNormals = true;
+		Offset.ExtrusionVectorType = FOffsetMeshRegion::EVertexExtrusionVectorType::SelectionTriNormalsAngleWeightedAverage;
 		Offset.DefaultOffsetDistance = ExtrudeHeight;
 		Offset.bIsPositiveOffset = (ExtrudeHeight > 0);
 		Offset.Apply();
diff --git a/Engine/Plugins/Experimental/MeshModelingToolsetExp/Source/ModelingOperators/Private/DeformationOps/ExtrudeOp.cpp b/Engine/Plugins/Experimental/MeshModelingToolsetExp/Source/ModelingOperators/Private/DeformationOps/ExtrudeOp.cpp
new file mode 100644
--- /dev/null
+++ b/Engine/Plugins/Experimental/MeshModelingToolsetExp/Source/ModelingOperators/Private/DeformationOps/ExtrudeOp.cpp
@@ -0,0 +1,218 @@
+// Copyright Epic Games, Inc. All Rights Reserved.
+
+#include "DeformationOps/ExtrudeOp.h"
+
+#include "DynamicMesh/DynamicMeshChangeTracker.h"
+#include "DynamicMesh/MeshNormals.h"
+#include "DynamicSubmesh3.h"
+#include "Operations/OffsetMeshRegion.h"
+#include "Operations/MeshBoolean.h"
+
+#include "Util/ProgressCancel.h"
+
+#include "ExplicitUseGeometryMathTypes.h"		// using UE::Geometry::(math types)
+using namespace UE::Geometry;
+
+namespace ExtrudeOpLocals
+{
+	// Take the selected triangles and change their groups to new ones. This makes it
+	// easier to track the groups into the result in a boolean extrude, and it also
+	// makes sure that the extrude doesn't create disconnected groups (if you do it
+	// on just some of the triangles in a group)
+	void UpdateSelectionGroups(FDynamicMesh3& Mesh, const TArray<int32> TriangleSelectionIn, TArray<int32>& SelectionGidsOut)
+	{
+		if (Mesh.HasTriangleGroups())
+		{
+			TMap<int32, int32> GroupMap;
+			for (int32 Tid : TriangleSelectionIn)
+			{
+				int32 Gid = Mesh.GetTriangleGroup(Tid);
+				int32* NewGidPointer = GroupMap.Find(Gid);
+				if (NewGidPointer)
+				{
+					Mesh.SetTriangleGroup(Tid, *NewGidPointer);
+				}
+				else
+				{
+					int32 NewGid = Mesh.AllocateTriangleGroup();
+					Mesh.SetTriangleGroup(Tid, NewGid);
+					GroupMap.Add(Gid, NewGid);
+
+					SelectionGidsOut.Add(NewGid);
+				}
+			}
+		}
+	}
+}
+
+void FExtrudeOp::CalculateResult(FProgressCancel* Progress)
+{
+	if (FMath::Abs(ExtrudeDistance) < KINDA_SMALL_NUMBER
+		|| (Progress && Progress->Cancelled()))
+	{
+		return;
+	}
+
+	// Reset result
+	ResultMesh->Copy(*OriginalMesh, true, true, true, true);
+
+	if (ExtrudeMode == EPolyEditExtrudeMode::Boolean)
+	{
+		BooleanExtrude(Progress);
+	}
+	else
+	{
+		MoveAndStitchExtrude(Progress);
+	}
+}
+
+bool FExtrudeOp::BooleanExtrude(FProgressCancel* Progress)
+{
+	using namespace ExtrudeOpLocals;
+
+	if (Progress && Progress->Cancelled())
+	{
+		return false;
+	}
+
+	TArray<int32> SelectionGids;
+	UpdateSelectionGroups(*ResultMesh, TriangleSelection, SelectionGids);
+
+	// Grab just the patch to be extruded.
+	TUniquePtr<FDynamicSubmesh3> Submesh = MakeUnique<FDynamicSubmesh3>(ResultMesh.Get(), TriangleSelection);
+
+	Submesh->MapGroupsToSubmesh(SelectionGids);
+
+	if (Progress && Progress->Cancelled())
+	{
+		return false;
+	}
+
+	// Extrude the patch:
+	FOffsetMeshRegion Extruder(&Submesh->GetSubmesh());
+
+	Extruder.Triangles.Reserve(TriangleSelection.Num());
+	for (int32 Tid : Submesh->GetSubmesh().TriangleIndicesItr())
+	{
+		Extruder.Triangles.Add(Tid);
+	}
+
+	// We need a custom function here so that we can set the side groups based on the original mesh, not the submesh
+	Extruder.LoopEdgesShouldHaveSameGroup = [&Submesh, this](int32 Eid1, int32 Eid2) {
+		int32 EidToUse1 = Submesh->MapEdgeToBaseMesh(Eid1);
+		int32 EidToUse2 = Submesh->MapEdgeToBaseMesh(Eid2);
+		return FOffsetMeshRegion::EdgesSeparateSameGroupsAndAreColinearAtBorder(
+			ResultMesh.Get(), EidToUse1, EidToUse2, bUseColinearityForSettingBorderGroups);
+	};
+
+	Extruder.UVScaleFactor = UVScaleFactor;
+	if (DirectionMode == EPolyEditExtrudeDirectionMode::SingleDirection)
+	{
+		Extruder.OffsetPositionFunc = [this](const FVector3d& Pos, const FVector3d& VertexVector, int32 VertexID)
+		{
+			return Pos + ExtrudeDistance * MeshSpaceExtrudeDirection;
+		};
+	}
+	else
+	{
+		Extruder.OffsetPositionFunc = [this](const FVector3d& Pos, const FVector3d& VertexVector, int32 VertexID) {
+			return Pos + ExtrudeDistance * VertexVector;
+		};
+	}
+
+	Extruder.ExtrusionVectorType =
+		DirectionMode == EPolyEditExtrudeDirectionMode::VertexNormals ? FOffsetMeshRegion::EVertexExtrusionVectorType::VertexNormal
+		: DirectionMode == EPolyEditExtrudeDirectionMode::SelectedTriangleNormals ? FOffsetMeshRegion::EVertexExtrusionVectorType::SelectionTriNormalsAngleWeightedAverage
+		: DirectionMode == EPolyEditExtrudeDirectionMode::SelectedTriangleNormalsEven ? FOffsetMeshRegion::EVertexExtrusionVectorType::SelectionTriNormalsAngleWeightedAdjusted
+		: FOffsetMeshRegion::EVertexExtrusionVectorType::Zero;
+
+	Extruder.bIsPositiveOffset = (ExtrudeDistance > 0);
+	Extruder.bOffsetFullComponentsAsSolids = true;
+	Extruder.Apply();
+
+	FMeshNormals::QuickComputeVertexNormalsForTriangles(*ResultMesh, Extruder.AllModifiedAndNewTriangles);
+
+	if (Progress && Progress->Cancelled())
+	{
+		return false;
+	}
+
+	// Now perform a boolean operation with our result.
+	FMeshBoolean::EBooleanOp Op = ExtrudeDistance > 0 ? FMeshBoolean::EBooleanOp::Union : FMeshBoolean::EBooleanOp::Difference;
+	FMeshBoolean MeshBoolean(OriginalMesh.Get(), FTransform3d::Identity(), &Submesh->GetSubmesh(), FTransform3d::Identity(), ResultMesh.Get(), Op);
+	MeshBoolean.bPutResultInInputSpace = true;
+	MeshBoolean.bSimplifyAlongNewEdges = true;
+	MeshBoolean.bPopulateSecondMeshGroupMap = true;
+	MeshBoolean.Progress = Progress;
+	bool bSuccess = MeshBoolean.Compute();
+
+	// Get the group ID's of the extruded face, if they still exist (they may not if we cut entirely through)
+	ExtrudedFaceNewGids.Empty();
+	for (int32 Gid : SelectionGids)
+	{
+		const int32* NewGid = MeshBoolean.SecondMeshGroupMap.FindTo(Gid);
+		if (NewGid)
+		{
+			ExtrudedFaceNewGids.Add(*NewGid);
+		}
+	}
+
+	return bSuccess;
+}
+
+bool FExtrudeOp::MoveAndStitchExtrude(FProgressCancel* Progress)
+{
+	using namespace ExtrudeOpLocals;
+
+	if (Progress && Progress->Cancelled())
+	{
+		return false;
+	}
+
+	TArray<int32> SelectionGids;
+	UpdateSelectionGroups(*ResultMesh, TriangleSelection, SelectionGids);
+
+	FOffsetMeshRegion Extruder(ResultMesh.Get());
+	Extruder.Triangles = TriangleSelection;
+
+	Extruder.LoopEdgesShouldHaveSameGroup = [this](int32 Eid1, int32 Eid2) {
+		return FOffsetMeshRegion::EdgesSeparateSameGroupsAndAreColinearAtBorder(
+			ResultMesh.Get(), Eid1, Eid2, bUseColinearityForSettingBorderGroups);
+	};
+
+	Extruder.UVScaleFactor = UVScaleFactor;
+	if (DirectionMode == EPolyEditExtrudeDirectionMode::SingleDirection)
+	{
+		Extruder.OffsetPositionFunc = [this](const FVector3d& Pos, const FVector3d& VertexVector, int32 VertexID)
+		{
+			return Pos + ExtrudeDistance * MeshSpaceExtrudeDirection;
+		};
+	}
+	else
+	{
+		Extruder.OffsetPositionFunc = [this](const FVector3d& Pos, const FVector3d& VertexVector, int32 VertexID) {
+			return Pos + ExtrudeDistance * VertexVector;
+		};
+	}
+
+	Extruder.ExtrusionVectorType =
+		DirectionMode == EPolyEditExtrudeDirectionMode::VertexNormals ? FOffsetMeshRegion::EVertexExtrusionVectorType::VertexNormal
+		: DirectionMode == EPolyEditExtrudeDirectionMode::SelectedTriangleNormals ? FOffsetMeshRegion::EVertexExtrusionVectorType::SelectionTriNormalsAngleWeightedAverage
+		: DirectionMode == EPolyEditExtrudeDirectionMode::SelectedTriangleNormalsEven ? FOffsetMeshRegion::EVertexExtrusionVectorType::SelectionTriNormalsAngleWeightedAdjusted
+		: FOffsetMeshRegion::EVertexExtrusionVectorType::Zero;
+
+	Extruder.bIsPositiveOffset = (ExtrudeDistance > 0);
+	Extruder.bOffsetFullComponentsAsSolids = bShellsToSolids;
+	Extruder.Apply();
+
+	if (Progress && Progress->Cancelled())
+	{
+		return false;
+	}
+
+	FMeshNormals::QuickComputeVertexNormalsForTriangles(*ResultMesh, Extruder.AllModifiedAndNewTriangles);
+
+	ExtrudedFaceNewGids = SelectionGids;
+
+	return true;
+}
diff --git a/Engine/Plugins/Experimental/MeshModelingToolsetExp/Source/ModelingOperators/Public/DeformationOps/ExtrudeOp.h b/Engine/Plugins/Experimental/MeshModelingToolsetExp/Source/ModelingOperators/Public/DeformationOps/ExtrudeOp.h
new file mode 100644
--- /dev/null
+++ b/Engine/Plugins/Experimental/MeshModelingToolsetExp/Source/ModelingOperators/Public/DeformationOps/ExtrudeOp.h
@@ -0,0 +1,83 @@
+// Copyright Epic Games, Inc. All Rights Reserved.
+
+#pragma once
+
+#include "CoreMinimal.h"
+
+#include "ModelingOperators.h"
+
+UENUM()
+enum class EPolyEditExtrudeDirectionMode
+{
+	// Take the angle-weighed average of the selected triangles around each
+	// extruded vertex to determine vertex movement direction.
+	SelectedTriangleNormals,
+	
+	// Like Selected Triangle Normals, but also adjusts the distances moved in
+	// an attempt to keep triangles parallel to their original facing.
+	SelectedTriangleNormalsEven,
+
+	// Vertex normals, regardless of selection.
+	VertexNormals,
+
+	// Extrude all triangles in the same direction regardless of their facing.
+	SingleDirection,
+};
+
+UENUM()
+enum class EPolyEditExtrudeMode
+{
+	// Performs extrusion by shifting the selected faces as stitching them
+	// to the border left behind.
+	MoveAndStitch,
+	
+	// Performs the extrusion by extruding selected faces into a closed mesh
+	// first, and then performing a boolean with the original mesh. This allows
+	// the extrusion to cut holes through the mesh or to bridge sections.
+	Boolean
+};
+
+namespace UE {
+namespace Geometry {
+
+class FDynamicMesh3;
+class FDynamicMeshChangeTracker;
+
+class MODELINGOPERATORS_API FExtrudeOp : public FDynamicMeshOperator
+{
+public:
+	virtual ~FExtrudeOp() {}
+
+	// Inputs:
+	TSharedPtr<const FDynamicMesh3, ESPMode::ThreadSafe> OriginalMesh;
+	TArray<int32> TriangleSelection;
+	double ExtrudeDistance = 0;
+	EPolyEditExtrudeMode ExtrudeMode = EPolyEditExtrudeMode::MoveAndStitch;
+	EPolyEditExtrudeDirectionMode DirectionMode = EPolyEditExtrudeDirectionMode::SelectedTriangleNormalsEven;
+	// Only used if DirectionMode is SingleDirection
+	FVector3d MeshSpaceExtrudeDirection;
+	// Used when setting groups for the sides when the extrusion includes a mesh border. When true,
+	// different groups are added based on colinearity of the border edges.
+	bool bUseColinearityForSettingBorderGroups = true;
+	float UVScaleFactor = 1.0f;
+	
+	/** Controls whether extruding an entire open-border patch should create a solid or an open shell */
+	bool bShellsToSolids = true;
+
+	// Outputs
+	TArray<int32> ExtrudedFaceNewGids;
+
+	void SetTransform(const FTransform& Transform)
+	{
+		ResultTransform = (FTransform3d)Transform;
+	}
+
+	// FDynamicMeshOperator implementation 
+	virtual void CalculateResult(FProgressCancel* Progress) override;
+
+protected:
+	virtual bool BooleanExtrude(FProgressCancel* Progress);
+	virtual bool MoveAndStitchExtrude(FProgressCancel* Progress);
+};
+
+}} // end UE::Geometry
diff --git a/Engine/Plugins/Runtime/GeometryProcessing/Source/DynamicMesh/Private/Operations/MeshBoolean.cpp b/Engine/Plugins/Runtime/GeometryProcessing/Source/DynamicMesh/Private/Operations/MeshBoolean.cpp
--- a/Engine/Plugins/Runtime/GeometryProcessing/Source/DynamicMesh/Private/Operations/MeshBoolean.cpp
+++ b/Engine/Plugins/Runtime/GeometryProcessing/Source/DynamicMesh/Private/Operations/MeshBoolean.cpp
@@ -668,6 +668,11 @@ bool FMeshBoolean::Compute()
 		FMeshIndexMappings IndexMaps;
 		Editor.AppendMesh(CutMesh[1], IndexMaps);
 
+		if (bPopulateSecondMeshGroupMap)
+		{
+			SecondMeshGroupMap = IndexMaps.GetGroupMap();
+		}
+
 		if (bWeldSharedEdges)
 		{
 			bool bWeldSuccess = MergeEdges(IndexMaps, CutMesh, CutBoundaryEdges, AllVIDMatches);
diff --git a/Engine/Plugins/Runtime/GeometryProcessing/Source/DynamicMesh/Private/Operations/OffsetMeshRegion.cpp b/Engine/Plugins/Runtime/GeometryProcessing/Source/DynamicMesh/Private/Operations/OffsetMeshRegion.cpp
--- a/Engine/Plugins/Runtime/GeometryProcessing/Source/DynamicMesh/Private/Operations/OffsetMeshRegion.cpp
+++ b/Engine/Plugins/Runtime/GeometryProcessing/Source/DynamicMesh/Private/Operations/OffsetMeshRegion.cpp
@@ -12,20 +12,33 @@
 #include "ExplicitUseGeometryMathTypes.h"		// using UE::Geometry::(math types)
 using namespace UE::Geometry;
 
-FOffsetMeshRegion::FOffsetMeshRegion(FDynamicMesh3* mesh) : Mesh(mesh)
+namespace OffsetMeshRegionLocals
 {
-	OffsetPositionFunc = [this](const FVector3d& Position, const FVector3f& Normal, int VertexID)
+	bool EdgesAreParallel(FDynamicMesh3* Mesh, int32 Eid1, int32 Eid2)
 	{
-		return Position + this->DefaultOffsetDistance * (FVector3d)Normal;
-	};
+		FIndex2i Vids1 = Mesh->GetEdgeV(Eid1);
+		FIndex2i Vids2 = Mesh->GetEdgeV(Eid2);
+
+		FVector3d Vec1 = Mesh->GetVertex(Vids1.A) - Mesh->GetVertex(Vids1.B);
+		FVector3d Vec2 = Mesh->GetVertex(Vids2.A) - Mesh->GetVertex(Vids2.B);
+		if (!Vec1.Normalize(KINDA_SMALL_NUMBER) || !Vec2.Normalize(KINDA_SMALL_NUMBER))
+		{
+			// A degenerate edge is parallel enough for our purposes
+			return true;
+		}
+		return FMath::Abs(Vec1.Dot(Vec2)) >= 1 - KINDA_SMALL_NUMBER;
+	}
 }
 
+FOffsetMeshRegion::FOffsetMeshRegion(FDynamicMesh3* mesh) : Mesh(mesh)
+{
+}
 
 bool FOffsetMeshRegion::Apply()
 {
 	FMeshNormals Normals;
 	bool bHaveVertexNormals = Mesh->HasVertexNormals();
-	if (bHaveVertexNormals == false && bUseFaceNormals == false)
+	if (!bHaveVertexNormals && ExtrusionVectorType == EVertexExtrusionVectorType::VertexNormal)
 	{
 		Normals = FMeshNormals(Mesh);
 		Normals.ComputeVertexNormals();
@@ -39,32 +52,18 @@ bool FOffsetMeshRegion::Apply()
 	for (int k = 0; k < RegionComponents.Num(); ++k)
 	{
 		FOffsetInfo& Region = OffsetRegions[k];
-		Region.InitialTriangles = MoveTemp(RegionComponents.Components[k].Indices);
-		bool bRegionOK = false;
+		Region.OffsetTids = MoveTemp(RegionComponents.Components[k].Indices);
 
-		TArray<int32> AllTriangles;
-		FMeshConnectedComponents::GrowToConnectedTriangles(Mesh, Region.InitialTriangles, AllTriangles);
-		if (AllTriangles.Num() == Region.InitialTriangles.Num() && bOffsetFullComponentsAsSolids)
-		{
-			bRegionOK = ApplySolidExtrude(Region, (bHaveVertexNormals) ? nullptr : &Normals);
-		}
-		else 
+		if (bOffsetFullComponentsAsSolids)
 		{
-			bRegionOK = ApplyOffset(Region, (bHaveVertexNormals) ? nullptr : &Normals);
+			TArray<int32> AllTriangles;
+			FMeshConnectedComponents::GrowToConnectedTriangles(Mesh, Region.OffsetTids, AllTriangles);
+			Region.bIsSolid = AllTriangles.Num() == Region.OffsetTids.Num();
 		}
 
-		if ( bRegionOK )
-		{
-			AllModifiedTriangles.Append(Region.InitialTriangles);
-			for (TArray<int32>& RegionTris : Region.StitchTriangles)
-			{
-				AllModifiedTriangles.Append(RegionTris);
-			}
-		}
-		else
-		{
-			bAllOK = false;
-		}
+		bool bRegionOK = ApplyOffset(Region, (bHaveVertexNormals) ? nullptr : &Normals);
+
+		bAllOK = bAllOK && bRegionOK;
 	}
 
 	return bAllOK;
@@ -72,188 +71,227 @@ bool FOffsetMeshRegion::Apply()
 
 }
 
-
-bool FOffsetMeshRegion::ApplySolidExtrude(FOffsetInfo& Region, FMeshNormals* UseNormals)
+bool FOffsetMeshRegion::ApplyOffset(FOffsetInfo& Region, FMeshNormals* UseNormals)
 {
-	FDynamicSubmesh3 SubmeshCalc(Mesh, Region.InitialTriangles);
-	FDynamicMesh3& Submesh = SubmeshCalc.GetSubmesh();
-
-	FExtrudeMesh Extruder(&Submesh);
-	Extruder.ExtrudedPositionFunc = this->OffsetPositionFunc;
-	Extruder.DefaultExtrudeDistance = this->DefaultOffsetDistance;
-	Extruder.UVScaleFactor = this->UVScaleFactor;
-	Extruder.IsPositiveOffset = bIsPositiveOffset;
+	// Store offset groups
+	if (Mesh->HasTriangleGroups())
+	{
+		for (int32 Tid : Region.OffsetTids)
+		{
+			Region.OffsetGroups.AddUnique(Mesh->GetTriangleGroup(Tid));
+		}
+	}
 
-	bool bOK = Extruder.Apply();
+	FMeshRegionBoundaryLoops InitialLoops(Mesh, Region.OffsetTids, false);
+	bool bOK = InitialLoops.Compute();
 	if (bOK == false)
 	{
 		return false;
 	}
 
-	if (ChangeTracker)
+	AllModifiedAndNewTriangles.Append(Region.OffsetTids);
+
+	// Before we start changing triangles, prepare by allocating group IDs that we'll use
+	// for the stitched sides (doing it before changes to the mesh allows user-provided
+	// functions to operate on the original mesh).
+	TArray<TArray<int32>> LoopsEdgeGroups;
+	TArray<int32> NewGroupIDs;
+	LoopsEdgeGroups.SetNum(InitialLoops.Loops.Num());
+	for (int32 i = 0; i < InitialLoops.Loops.Num(); ++i)
 	{
-		ChangeTracker->SaveTriangles(Region.InitialTriangles, true);
-	}
+		TArray<int32>& LoopEids = InitialLoops.Loops[i].Edges;
+		int32 NumEids = LoopEids.Num();
 
-	FDynamicMeshEditor Editor(Mesh);
-	Editor.RemoveTriangles(Region.InitialTriangles, true);
+		if (!ensure(NumEids > 2))
+		{
+			// Shouldn't actually happen because we're extruding triangles
+			continue;
+		}
 
-	FMeshIndexMappings Mappings;
-	Editor.AppendMesh(&Submesh, Mappings);
+		TArray<int32>& CurrentEdgeGroups = LoopsEdgeGroups[i];
+		CurrentEdgeGroups.SetNumUninitialized(NumEids);
+		CurrentEdgeGroups[0] = Mesh->AllocateTriangleGroup();
+		NewGroupIDs.Add(CurrentEdgeGroups[0]);
 
-	// transfer tris and groups back
-	// TODO: loops
-	for (FExtrudeMesh::FExtrusionInfo& ExtrudeRegionInfo : Extruder.Extrusions)
-	{
-		for (TArray<int32>& StitchTriSet : ExtrudeRegionInfo.StitchTriangles)
+		// Propagate the group backwards first so we don't allocate an unnecessary group
+		// at the end and then have to fix it.
+		int32 LastDifferentGroupIndex = NumEids - 1;
+		while (LastDifferentGroupIndex > 0
+			&& LoopEdgesShouldHaveSameGroup(LoopEids[0], LoopEids[LastDifferentGroupIndex]))
 		{
-			for (int32 k = 0; k < StitchTriSet.Num(); ++k)
-			{
-				StitchTriSet[k] = Mappings.GetNewTriangle(StitchTriSet[k]);
-			}
-			Region.StitchTriangles.Add(StitchTriSet);
+			CurrentEdgeGroups[LastDifferentGroupIndex] = CurrentEdgeGroups[0];
+			--LastDifferentGroupIndex;
 		}
-		for (TArray<int32>& StitchGroupSet : ExtrudeRegionInfo.StitchPolygonIDs)
+
+		// Now add new groups forward
+		for (int32 j = 1; j <= LastDifferentGroupIndex; ++j)
 		{
-			for (int32 k = 0; k < StitchGroupSet.Num(); ++k)
+			if (!LoopEdgesShouldHaveSameGroup(LoopEids[j], LoopEids[j - 1]))
 			{
-				StitchGroupSet[k] = Mappings.GetNewGroup(StitchGroupSet[k]);
+				CurrentEdgeGroups[j] = Mesh->AllocateTriangleGroup();
+				NewGroupIDs.Add(CurrentEdgeGroups[j]);
+			}
+			else
+			{
+				CurrentEdgeGroups[j] = CurrentEdgeGroups[j-1];
 			}
-			Region.StitchPolygonIDs.Add(StitchGroupSet);
-		}
-
-		for (int32 GroupID : ExtrudeRegionInfo.OffsetTriGroups)
-		{
-			Region.OffsetGroups.Add(Mappings.GetNewGroup(GroupID));
 		}
 	}
 
-	Region.bIsSolid = true;
-
-	//Region.BaseLoops[LoopIndex].InitializeFromVertices(Mesh, BaseLoopV);
-	//Region.OffsetLoops[LoopIndex].InitializeFromVertices(Mesh, OffsetLoopV);
-	//LoopIndex++;
-
-	return true;
-
-}
-
+	FDynamicMeshEditor Editor(Mesh);
+	TArray<FDynamicMeshEditor::FLoopPairSet> LoopPairs;
 
-bool FOffsetMeshRegion::ApplyOffset(FOffsetInfo& Region, FMeshNormals* UseNormals)
-{
-	FMeshRegionBoundaryLoops InitialLoops(Mesh, Region.InitialTriangles, false);
-	bool bOK = InitialLoops.Compute();
-	if (bOK == false)
+	FDynamicMeshEditResult DuplicateResult;
+	if (Region.bIsSolid)
 	{
-		return false;
-	}
+		// In the solid case, we want to duplicate the region.
+		FMeshIndexMappings IndexMap;
+		Editor.DuplicateTriangles(Region.OffsetTids, IndexMap, DuplicateResult);
 
-	int NumInitialLoops = InitialLoops.GetLoopCount();
+		AllModifiedAndNewTriangles.Append(DuplicateResult.NewTriangles);
 
-	if (ChangeTracker)
-	{
-		ChangeTracker->SaveTriangles(Region.InitialTriangles, true);
-	}
+		// Populate LoopPairs
+		for (FEdgeLoop& BaseLoop : InitialLoops.Loops)
+		{
+			LoopPairs.Add(FDynamicMeshEditor::FLoopPairSet());
+			FDynamicMeshEditor::FLoopPairSet& LoopPair = LoopPairs.Last();
+
+			// Which loops we choose as the outer/inner will determine whether the
+			// sides are stitched inside out or not. The original Tids are the ones
+			// that are offset. In a positive offset, we want old as the "outer" and
+			// new as "inner". In negative offset, we want the reverse to allow our
+			// stitching code to still have the sides face outward.
+			TArray<int32>* OriginalVertsOut = &LoopPair.OuterVertices;
+			TArray<int32>* OriginalEdgesOut = &LoopPair.OuterEdges;
+			TArray<int32>* NewVertsOut = &LoopPair.InnerVertices;
+			TArray<int32>* NewEdgesOut = &LoopPair.InnerEdges;
+			
+			if (bIsPositiveOffset)
+			{
+				Swap(OriginalVertsOut, NewVertsOut);
+				Swap(OriginalEdgesOut, NewEdgesOut);
+			}
 
-	FDynamicMeshEditor Editor(Mesh);
+			*OriginalVertsOut = BaseLoop.Vertices;
+			*OriginalEdgesOut = BaseLoop.Edges;
+			
+			for (int32 Vid : BaseLoop.Vertices)
+			{
+				NewVertsOut->Add(IndexMap.GetNewVertex(Vid));
+			}
 
-	// keep track of offset groups
-	if (Mesh->HasTriangleGroups())
-	{
-		for (int32 gid : Region.InitialTriangles)
-		{
-			Region.OffsetGroups.AddUnique(Mesh->GetTriangleGroup(gid));
+			FEdgeLoop OtherLoop;
+			bOK = ensure(OtherLoop.InitializeFromVertices(Mesh, *NewVertsOut, false));
+			*NewEdgesOut = OtherLoop.Edges;
 		}
 	}
+	else
+	{
+		bOK = Editor.DisconnectTriangles(Region.OffsetTids, LoopPairs, true /*bHandleBoundaryVertices*/);
+	}
 
-	TArray<FDynamicMeshEditor::FLoopPairSet> LoopPairs;
-	bOK = Editor.DisconnectTriangles(Region.InitialTriangles, LoopPairs, true);
 	if (bOK == false)
 	{
 		return false;
 	}
 
-	// offset vertices
 	FMeshVertexSelection SelectionV(Mesh);
-	SelectionV.SelectTriangleVertices(Region.InitialTriangles);
-	if (bUseFaceNormals)
+	SelectionV.SelectTriangleVertices(Region.OffsetTids);
+	TArray<int32> SelectedVids = SelectionV.AsArray();
+
+	// If we need to, assemble the vertex vectors for us to use (before we actually start moving things)
+	TArray<FVector3d> VertexExtrudeVectors;
+	if (ExtrusionVectorType == EVertexExtrusionVectorType::SelectionTriNormalsAngleWeightedAverage
+		|| ExtrusionVectorType == EVertexExtrusionVectorType::SelectionTriNormalsAngleWeightedAdjusted)
 	{
-		TArray<int32> Vertices = SelectionV.AsArray();
-		TSet<int32> TriangleSet(Region.InitialTriangles);
-		int32 NumV = Vertices.Num();
-		TArray<FVector3d> NewPositions;
-		NewPositions.SetNum(NumV);
-		for (int32 k = 0; k < NumV; ++k)
+		VertexExtrudeVectors.SetNumUninitialized(SelectedVids.Num());
+
+		// Used to test which triangles are in selection
+		TSet<int32> TriangleSet(Region.OffsetTids);
+
+		for (int32 i = 0; i < SelectedVids.Num(); ++i)
 		{
-			int32 vid = Vertices[k];
-			FVector3d VertexPos = Mesh->GetVertex(vid);
-			FVector3d AccumV = FVector3d::Zero();
-			int32 Count = 0;
-			for (int32 tid : Mesh->VtxTrianglesItr(vid))
+			int32 Vid = SelectedVids[i];
+			FVector3d ExtrusionVector = FVector3d::Zero();
+
+			// Get angle-weighted normalized average vector
+			for (int32 Tid : Mesh->VtxTrianglesItr(Vid))
 			{
-				if (TriangleSet.Contains(tid))
+				if (TriangleSet.Contains(Tid))
 				{
-					FVector3f TriNormal = (FVector3f)Mesh->GetTriNormal(tid);
-					FVector3d TriNormalOffsetPos = OffsetPositionFunc(VertexPos, TriNormal, vid);
-					AccumV += TriNormalOffsetPos;
-					Count++;
+					double Angle = Mesh->GetTriInternalAngleR(Tid, Mesh->GetTriangle(Tid).IndexOf(Vid));
+					ExtrusionVector += Angle * Mesh->GetTriNormal(Tid);
 				}
 			}
-			NewPositions[k] = (Count == 0) ? VertexPos : (AccumV / (double)Count);
-		}
-		for (int32 k = 0; k < NumV; ++k)
-		{
-			Mesh->SetVertex(Vertices[k], NewPositions[k]);
+			ExtrusionVector.Normalize();
+
+			if (ExtrusionVectorType == EVertexExtrusionVectorType::SelectionTriNormalsAngleWeightedAdjusted)
+			{
+				// Perform an angle-weighted adjustment of the vector length. For each triangle normal, the
+				// length needs to be multiplied by 1/cos(theta) to place the vertex in the plane that it
+				// would be in if the face was moved a unit along triangle normal (where theta is angle of
+				// triangle normal to the current extrusion vector).
+				double AngleSum = 0;
+				double Adjustment = 0;
+				for (int32 Tid : Mesh->VtxTrianglesItr(Vid))
+				{
+					if (TriangleSet.Contains(Tid))
+					{
+						double Angle = Mesh->GetTriInternalAngleR(Tid, Mesh->GetTriangle(Tid).IndexOf(Vid));
+						double CosTheta = Mesh->GetTriNormal(Tid).Dot(ExtrusionVector);
+						Adjustment += Angle / CosTheta;
+
+						// For the average at the end
+						AngleSum += Angle;
+					}
+				}
+				Adjustment /= AngleSum;
+				ExtrusionVector *= Adjustment;
+			}
+
+			VertexExtrudeVectors[i] = ExtrusionVector;
 		}
 	}
-	else
+
+	// Perform the actual vertex displacement.
+	for (int32 i = 0; i < SelectedVids.Num(); ++i)
 	{
-		for (int32 vid : SelectionV)
+		int32 Vid = SelectedVids[i];
+		FVector3d OldPosition = Mesh->GetVertex(Vid);
+		FVector3f ExtrusionVector = FVector3d::Zero();
+
+		switch (ExtrusionVectorType)
 		{
-			FVector3d v = Mesh->GetVertex(vid);
-			FVector3f n = (UseNormals != nullptr) ? (FVector3f)(*UseNormals)[vid] : Mesh->GetVertexNormal(vid);
-			FVector3d newv = OffsetPositionFunc(v, n, vid);
-			Mesh->SetVertex(vid, newv);
+		case EVertexExtrusionVectorType::VertexNormal:
+			ExtrusionVector = Mesh->GetVertexNormal(Vid);
+			break;
+		case EVertexExtrusionVectorType::SelectionTriNormalsAngleWeightedAverage:
+		case EVertexExtrusionVectorType::SelectionTriNormalsAngleWeightedAdjusted:
+			ExtrusionVector = VertexExtrudeVectors[i];
+			break;
 		}
-	}
 
+		FVector3d NewPosition = OffsetPositionFunc(OldPosition, ExtrusionVector, Vid);
+		Mesh->SetVertex(Vid, NewPosition);
+	}
 
+	// Stitch the loops
 
-	// stitch each loop
+	int NumInitialLoops = InitialLoops.GetLoopCount();
 	Region.BaseLoops.SetNum(NumInitialLoops);
 	Region.OffsetLoops.SetNum(NumInitialLoops);
 	Region.StitchTriangles.SetNum(NumInitialLoops);
 	Region.StitchPolygonIDs.SetNum(NumInitialLoops);
 	int32 LoopIndex = 0;
-	for (FDynamicMeshEditor::FLoopPairSet& LoopPair : LoopPairs)
+	for (int32 i = 0; i < LoopPairs.Num(); ++i)
 	{
+		FDynamicMeshEditor::FLoopPairSet& LoopPair = LoopPairs[i];
+		const TArray<int32>& EdgeGroups = LoopsEdgeGroups[i];
+
 		TArray<int32>& BaseLoopV = LoopPair.OuterVertices;
 		TArray<int32>& OffsetLoopV = LoopPair.InnerVertices;
 		int NumLoopV = BaseLoopV.Num();
 
-		// allocate a new group ID for each pair of input group IDs, and build up list of new group IDs along loop
-		TArray<int32> NewGroupIDs;
-		TArray<int32> EdgeGroups;
-		TMap<TPair<int32,int32>, int32> NewGroupsMap;
-		for (int32 k = 0; k < NumLoopV; ++k)
-		{
-			int32 OffsetEdgeID = Mesh->FindEdge(OffsetLoopV[k], OffsetLoopV[(k + 1) % NumLoopV]);
-			int32 OffsetGroupID = Mesh->GetTriangleGroup(Mesh->GetEdgeT(OffsetEdgeID).A);
-
-			// base edge may not exist if we offset entire region. In that case just use single GroupID
-			int32 BaseEdgeID = Mesh->FindEdge(BaseLoopV[k], BaseLoopV[(k + 1) % NumLoopV]);
-			int32 BaseGroupID = (BaseEdgeID >= 0) ? Mesh->GetTriangleGroup(Mesh->GetEdgeT(BaseEdgeID).A) : OffsetGroupID;
-
-			TPair<int32,int32> GroupPair(FMathd::Min(BaseGroupID, OffsetGroupID), FMathd::Max(BaseGroupID, OffsetGroupID));
-			if (NewGroupsMap.Contains(GroupPair) == false)
-			{
-				int32 NewGroupID = Mesh->AllocateTriangleGroup();
-				NewGroupIDs.Add(NewGroupID);
-				NewGroupsMap.Add(GroupPair, NewGroupID);
-			}
-			EdgeGroups.Add(NewGroupsMap[GroupPair]);
-		}
-
 		// stitch the loops
 		FDynamicMeshEditResult StitchResult;
 		Editor.StitchVertexLoopsMinimal(OffsetLoopV, BaseLoopV, StitchResult);
@@ -270,8 +308,9 @@ bool FOffsetMeshRegion::ApplyOffset(FOffsetInfo& Region, FMeshNormals* UseNormal
 		StitchResult.GetAllTriangles(Region.StitchTriangles[LoopIndex]);
 		Region.StitchPolygonIDs[LoopIndex] = NewGroupIDs;
 
+		AllModifiedAndNewTriangles.Append(Region.StitchTriangles[LoopIndex]);
+
 		// for each polygon we created in stitch, set UVs and normals
-		// TODO copied from FExtrudeMesh, doesn't really make sense in this context...
 		if (Mesh->HasAttributes())
 		{
 			float AccumUVTranslation = 0;
@@ -316,7 +355,41 @@ bool FOffsetMeshRegion::ApplyOffset(FOffsetInfo& Region, FMeshNormals* UseNormal
 		LoopIndex++;
 	}
 
+	if (Region.bIsSolid)
+	{
+		if (bIsPositiveOffset)
+		{
+			// Flip the "bottom" of the region to face outwards
+			Editor.ReverseTriangleOrientations(DuplicateResult.NewTriangles, true);
+		}
+		else
+		{
+			Editor.ReverseTriangleOrientations(Region.OffsetTids, true);
+		}
+	}
+
 	return true;
 }
 
+bool FOffsetMeshRegion::EdgesSeparateSameGroupsAndAreColinearAtBorder(FDynamicMesh3* Mesh, 
+	int32 Eid1, int32 Eid2, bool bCheckColinearityAtBorder)
+{
+	FIndex2i Tris1 = Mesh->GetEdgeT(Eid1);
+	FIndex2i Groups1(Mesh->GetTriangleGroup(Tris1.A),
+		Tris1.B == IndexConstants::InvalidID ? IndexConstants::InvalidID : Mesh->GetTriangleGroup(Tris1.B));
+
+	FIndex2i Tris2 = Mesh->GetEdgeT(Eid2);
+	FIndex2i Groups2(Mesh->GetTriangleGroup(Tris2.A), 
+		Tris2.B == IndexConstants::InvalidID ? IndexConstants::InvalidID : Mesh->GetTriangleGroup(Tris2.B));
+
+	if (bCheckColinearityAtBorder
+		&& Groups1.A == Groups2.A
+		&& Groups1.B == IndexConstants::InvalidID
+		&& Groups2.B == IndexConstants::InvalidID)
+	{
+		return OffsetMeshRegionLocals::EdgesAreParallel(Mesh, Eid1, Eid2);
+	}
+	else return (Groups1.A == Groups2.A && Groups1.B == Groups2.B)
+		|| (Groups1.A == Groups2.B && Groups1.B == Groups2.A);
+}
 
diff --git a/Engine/Plugins/Runtime/GeometryProcessing/Source/DynamicMesh/Public/DynamicSubmesh3.h b/Engine/Plugins/Runtime/GeometryProcessing/Source/DynamicMesh/Public/DynamicSubmesh3.h
--- a/Engine/Plugins/Runtime/GeometryProcessing/Source/DynamicMesh/Public/DynamicSubmesh3.h
+++ b/Engine/Plugins/Runtime/GeometryProcessing/Source/DynamicMesh/Public/DynamicSubmesh3.h
@@ -180,6 +180,28 @@ public:
 		}
 	}
 
+	/** @return Group ID in the submesh, mapped from base mesh */
+	inline int MapGroupToSubmesh(int BaseGID) const
+	{
+		checkSlow(bComputeTriMaps);
+		return Mappings.GetGroupMap().GetTo(BaseGID);
+	}
+	/** @return Group ID in the base mesh, mapped from submesh */
+	inline int MapGroupToBaseMesh(int SubGID) const
+	{
+		checkSlow(bComputeTriMaps);
+		return Mappings.GetGroupMap().GetFrom(SubGID);
+	}
+	/** @param GroupIDs Array of group IDs in the base mesh that will all be changed to the corresponding IDs in the submesh */
+	void MapGroupsToSubmesh(TArrayView<int32> GroupIDs) const
+	{
+		checkSlow(bComputeTriMaps);
+		for (int i = 0; i < GroupIDs.Num(); ++i)
+		{
+			GroupIDs[i] = MapGroupToSubmesh(GroupIDs[i]);
+		}
+	}
+
 	/**
 	 * Computes the Submesh object, index mappings corresponding sub to base mesh, and boundary between sub and base mesh
 	 * @param SubTriangles ArrayView of triangle IDs to include in the submesh
diff --git a/Engine/Plugins/Runtime/GeometryProcessing/Source/DynamicMesh/Public/Operations/ExtrudeMesh.h b/Engine/Plugins/Runtime/GeometryProcessing/Source/DynamicMesh/Public/Operations/ExtrudeMesh.h
--- a/Engine/Plugins/Runtime/GeometryProcessing/Source/DynamicMesh/Public/Operations/ExtrudeMesh.h
+++ b/Engine/Plugins/Runtime/GeometryProcessing/Source/DynamicMesh/Public/Operations/ExtrudeMesh.h
@@ -19,6 +19,9 @@ class FDynamicMesh3;
 class FMeshNormals;
 
 /**
+ * Note: FExtrudeMesh might someday be removed. Consider using FOffsetMeshRegion instead,
+ * which can extrude all or part of a mesh.
+ * 
  * FExtrudeMesh implements a full-mesh extrusion of a mesh. This happens in two stages:
  * 1) all triangles of input mesh are duplicated and offset
  * 2) base and offset border loops are stitched together with triangulated quads
diff --git a/Engine/Plugins/Runtime/GeometryProcessing/Source/DynamicMesh/Public/Operations/MeshBoolean.h b/Engine/Plugins/Runtime/GeometryProcessing/Source/DynamicMesh/Public/Operations/MeshBoolean.h
--- a/Engine/Plugins/Runtime/GeometryProcessing/Source/DynamicMesh/Public/Operations/MeshBoolean.h
+++ b/Engine/Plugins/Runtime/GeometryProcessing/Source/DynamicMesh/Public/Operations/MeshBoolean.h
@@ -117,6 +117,21 @@ public:
 	/** All edges created by mesh boolean algorithm. Only populated if bTrackAllNewEdges = true */
 	TSet<int32> AllNewEdges;
 
+	/**
+	 * When doing an operation that merges two meshes (Union, Difference, Intersect)
+	 * and bPopulateSecondMeshGroupMap is true, we populate a map from group ID's in the
+	 * second mesh to group ID's in the result. Useful when the boolean is part of
+	 * another operation that might want to track selection into the result.
+	 * 
+	 * TODO: It would be nice to have a similar triangle map, but it's more tedious to
+	 * implement due to the mesh cuts that happen. We might also someday want something
+	 * similar for the first mesh, where the group ID's will usually stay the same except
+	 * when all triangles of a particular group are removed and the group ID is repurposed
+	 * for the second mesh...
+	 */
+	bool bPopulateSecondMeshGroupMap = false;
+	FIndexMapi SecondMeshGroupMap;
+
 public:
 
 	/**
diff --git a/Engine/Plugins/Runtime/GeometryProcessing/Source/DynamicMesh/Public/Operations/OffsetMeshRegion.h b/Engine/Plugins/Runtime/GeometryProcessing/Source/DynamicMesh/Public/Operations/OffsetMeshRegion.h
--- a/Engine/Plugins/Runtime/GeometryProcessing/Source/DynamicMesh/Public/Operations/OffsetMeshRegion.h
+++ b/Engine/Plugins/Runtime/GeometryProcessing/Source/DynamicMesh/Public/Operations/OffsetMeshRegion.h
@@ -32,6 +32,18 @@ class DYNAMICMESH_API FOffsetMeshRegion
 {
 public:
 
+
+	enum class EVertexExtrusionVectorType
+	{
+		Zero,
+		VertexNormal,
+		// Angle weighted average of the triangles in the selection that contain this vertex (unit length)
+		SelectionTriNormalsAngleWeightedAverage,
+		// Like SelectionTriNormalsAngleWeightedAverage, but with the vertex length adjusted to try to keep
+		// the selection triangles parallel to their original location.
+		SelectionTriNormalsAngleWeightedAdjusted,
+	};
+
 	//
 	// Inputs
 	//
@@ -42,28 +54,42 @@ public:
 	/** The triangle region we are modifying */
 	TArray<int32> Triangles;
 
-	/** This function is called to generate the offset vertex position. Default returns (Position + DefaultOffsetDistance * Normal) */
-	TFunction<FVector3d(const FVector3d&, const FVector3f&, int)> OffsetPositionFunc;
+	/** This function is called to generate the offset vertex position. */
+	TFunction<FVector3d(const FVector3d& Position, const FVector3d& VertexVector, int Vid)> OffsetPositionFunc = 
+		[this](const FVector3d& Position, const FVector3d& VertexVector, int Vid)
+	{
+		return Position + this->DefaultOffsetDistance * (FVector3d)VertexVector;
+	};
 
-	/** If no Offset function is set, we will displace by DefaultOffsetDistance*Normal */
+	/** Used in the default OffsetPositionFunc. */
 	double DefaultOffsetDistance = 1.0;
 
+	/** Determines the type of vector passed in per-vertex to OffsetPositionFunc */
+	EVertexExtrusionVectorType ExtrusionVectorType = EVertexExtrusionVectorType::Zero;
+
+	/** 
+	 * When stitching the extruded portion to the source loop, this determines the grouping of the side quads by 
+	 * passing in successive Eids in the original loop. For instance, always returning true will result in one 
+	 * group for the stitched portion, and always returning true will result in a separate group for every edge
+	 * quad. By default, the same group is given to adjacent edges that separate the same two groups, or to border
+	 * edges that are colinear.
+	 */
+	TFunction<bool(int32 Eid1, int32 Eid2)> LoopEdgesShouldHaveSameGroup =
+		[this](int32 Eid1, int32 Eid2) { return EdgesSeparateSameGroupsAndAreColinearAtBorder(Mesh, Eid1, Eid2, true); };
+
 	/** quads on the stitch loop are planar-projected and scaled by this amount */
 	float UVScaleFactor = 1.0f;
 
 	/** If a sub-region of Triangles is a full connected component, offset into a solid instead of leaving a shell*/
 	bool bOffsetFullComponentsAsSolids = true;
 
-	/** if offset is "negative" (ie negative distance, inset, etc) and inset is an entire mesh region, needs to  */
+	/** 
+	 * When bOffsetFullComponentsAsSolids is true and the extrusion applied by OffsetPositionFunc is negative,
+	 * then fully-connected components need to be turned inside-out. Thus, in this circumstance we need to
+	 * be notified whether the OffsetPositionFunc is applying a positive or negative offset.
+	 */
 	bool bIsPositiveOffset = true;
 
-	/** if true, offset each vertex along each face normal and average result, rather than using estimated vertex normal  */
-	bool bUseFaceNormals = false;
-
-
-	/** If set, change tracker will be updated based on edit */
-	TUniquePtr<FDynamicMeshChangeTracker> ChangeTracker;
-
 	//
 	// Outputs
 	//
@@ -73,14 +99,21 @@ public:
 	 */
 	struct FOffsetInfo
 	{
-		/** Set of triangles for this region */
-		TArray<int32> InitialTriangles;
+		/**
+		 * Set of triangles for this region. These will be the same Tids that were
+		 * originally in the region. 
+		 */
+		TArray<int32> OffsetTids;
+		/** 
+		 * Groups on offset faces. These IDs will be the same as the original groups 
+		 * in the region. 
+		 */
+		TArray<int32> OffsetGroups;
+
 		/** Initial loops on the mesh */
 		TArray<FEdgeLoop> BaseLoops;
 		/** Offset loops on the mesh */
 		TArray<FEdgeLoop> OffsetLoops;
-		/** Groups on offset faces */
-		TArray<int32> OffsetGroups;
 
 		/** Lists of triangle-strip "tubes" that connect each loop-pair */
 		TArray<TArray<int>> StitchTriangles;
@@ -99,7 +132,7 @@ public:
 	/**
 	 * List of all triangles created/modified by this operation
 	 */
-	TArray<int32> AllModifiedTriangles;
+	TArray<int32> AllModifiedAndNewTriangles;
 
 public:
 	FOffsetMeshRegion(FDynamicMesh3* mesh);
@@ -126,12 +159,12 @@ public:
 	 */
 	virtual bool Apply();
 
-
+	// The default function used for LoopEdgesShouldHaveSameGroup. 
+	static bool EdgesSeparateSameGroupsAndAreColinearAtBorder(FDynamicMesh3* Mesh,
+		int32 Eid1, int32 Eid2, bool bCheckColinearityAtBorder);
 protected:
 
 	virtual bool ApplyOffset(FOffsetInfo& Region, FMeshNormals* UseNormals = nullptr);
-
-	virtual bool ApplySolidExtrude(FOffsetInfo& Region, FMeshNormals* UseNormals = nullptr);
 };
 
 } // end namespace UE::Geometry
diff --git a/Engine/Source/Runtime/GeometryCore/Public/DynamicMesh/MeshIndexMappings.h b/Engine/Source/Runtime/GeometryCore/Public/DynamicMesh/MeshIndexMappings.h
--- a/Engine/Source/Runtime/GeometryCore/Public/DynamicMesh/MeshIndexMappings.h
+++ b/Engine/Source/Runtime/GeometryCore/Public/DynamicMesh/MeshIndexMappings.h
@@ -67,6 +67,7 @@ public:
 	inline bool ContainsTriangle(int FromID) const { return TriangleMap.ContainsFrom(FromID); }
 
 	FIndexMapi& GetGroupMap() { return GroupMap; }
+	const FIndexMapi& GetGroupMap() const { return GroupMap; }
 	void SetGroup(int FromID, int ToID) { GroupMap.Add(FromID, ToID); }
 	int GetNewGroup(int FromID) const { return GroupMap.GetTo(FromID); }
 	inline bool ContainsGroup(int FromID) const { return GroupMap.ContainsFrom(FromID); }
