ModelingTools: add new bOnlySurfaceMaterials build option in MeshSceneAdapter to only include StaticMesh sections with MD_Surface material (previously this was hardcoded). Enable this setting in ApproximateActors, disable it in ISMEditorTool.
#rnx

diff --git a/Engine/Plugins/Experimental/MeshModelingToolsetExp/Source/GeometryProcessingAdapters/Private/ApproximateActorsImpl.cpp b/Engine/Plugins/Experimental/MeshModelingToolsetExp/Source/GeometryProcessingAdapters/Private/ApproximateActorsImpl.cpp
--- a/Engine/Plugins/Experimental/MeshModelingToolsetExp/Source/GeometryProcessingAdapters/Private/ApproximateActorsImpl.cpp
+++ b/Engine/Plugins/Experimental/MeshModelingToolsetExp/Source/GeometryProcessingAdapters/Private/ApproximateActorsImpl.cpp
@@ -964,6 +964,7 @@ void FApproximateActorsImpl::GenerateApproximationForActorSet(const TArray<AActo
 	// filter out objects smaller than 10% of voxel size
 	SceneBuildOptions.bFilterTinyObjects = Options.bIgnoreTinyParts;
 	SceneBuildOptions.TinyObjectBoxMaxDimension = Options.TinyPartMaxDimensionMeters;
+	SceneBuildOptions.bOnlySurfaceMaterials = true;
 	SceneBuildOptions.bPrintDebugMessages = Options.bVerbose;
 	{
 		TRACE_CPUPROFILER_EVENT_SCOPE(ApproximateActorsImpl_Generate_BuildScene);
diff --git a/Engine/Plugins/Experimental/MeshModelingToolsetExp/Source/ModelingComponents/Private/Scene/MeshSceneAdapter.cpp b/Engine/Plugins/Experimental/MeshModelingToolsetExp/Source/ModelingComponents/Private/Scene/MeshSceneAdapter.cpp
--- a/Engine/Plugins/Experimental/MeshModelingToolsetExp/Source/ModelingComponents/Private/Scene/MeshSceneAdapter.cpp
+++ b/Engine/Plugins/Experimental/MeshModelingToolsetExp/Source/ModelingComponents/Private/Scene/MeshSceneAdapter.cpp
@@ -482,7 +482,7 @@ static UMaterialInterface* GetStaticMeshMaterialFromSlotName(UStaticMesh* Static
 struct FMeshDescriptionTriangleMeshSurfaceAdapter : public FMeshDescriptionTriangleMeshAdapter
 {
 public:
-	FMeshDescriptionTriangleMeshSurfaceAdapter(const FMeshDescription* MeshIn, UStaticMesh* StaticMeshIn)
+	FMeshDescriptionTriangleMeshSurfaceAdapter(const FMeshDescription* MeshIn, UStaticMesh* StaticMeshIn, bool bOnlySurfaceMaterials)
 		: FMeshDescriptionTriangleMeshAdapter(MeshIn)
 		, NumValidTriangles(0)
 		, bHasInvalidPolyGroup(false)
@@ -490,25 +490,28 @@ public:
 		FStaticMeshConstAttributes MeshDescriptionAttributes(*MeshIn);
 		TPolygonGroupAttributesConstRef<FName> MaterialSlotNames = MeshDescriptionAttributes.GetPolygonGroupMaterialSlotNames();
 
-		for (FPolygonGroupID PolygonGroupID : MeshIn->PolygonGroups().GetElementIDs())
+		if (bOnlySurfaceMaterials)
 		{
-			// Try to find correct material index for the slot name in the MeshDescription. 
-			// Sometimes this data is wrong, in that case fall back to slot 0
-			FName MaterialSlotName = MaterialSlotNames[PolygonGroupID];
-			int32 MaterialIndex = GetStaticMeshMaterialIndexFromSlotName(StaticMeshIn, MaterialSlotName, true);
+			for (FPolygonGroupID PolygonGroupID : MeshIn->PolygonGroups().GetElementIDs())
+			{
+				// Try to find correct material index for the slot name in the MeshDescription. 
+				// Sometimes this data is wrong, in that case fall back to slot 0
+				FName MaterialSlotName = MaterialSlotNames[PolygonGroupID];
+				int32 MaterialIndex = GetStaticMeshMaterialIndexFromSlotName(StaticMeshIn, MaterialSlotName, true);
 
-			UMaterialInterface* MaterialInterface = GetStaticMeshMaterialFromSlotName(StaticMeshIn, MaterialSlotName, true);
-			const UMaterial* Material = MaterialInterface != nullptr ? MaterialInterface->GetMaterial_Concurrent() : nullptr;
+				UMaterialInterface* MaterialInterface = GetStaticMeshMaterialFromSlotName(StaticMeshIn, MaterialSlotName, true);
+				const UMaterial* Material = MaterialInterface != nullptr ? MaterialInterface->GetMaterial_Concurrent() : nullptr;
 
-			bool bValidPolyGroup = Material != nullptr && Material->MaterialDomain == EMaterialDomain::MD_Surface;
-			if (bValidPolyGroup)
-			{
-				ValidPolyGroups.Add(PolygonGroupID);
-				NumValidTriangles += MeshIn->GetNumPolygonGroupTriangles(PolygonGroupID);
-			}
-			else
-			{
-				bHasInvalidPolyGroup = true;
+				bool bValidPolyGroup = Material != nullptr && Material->MaterialDomain == EMaterialDomain::MD_Surface;
+				if (bValidPolyGroup)
+				{
+					ValidPolyGroups.Add(PolygonGroupID);
+					NumValidTriangles += MeshIn->GetNumPolygonGroupTriangles(PolygonGroupID);
+				}
+				else
+				{
+					bHasInvalidPolyGroup = true;
+				}
 			}
 		}
 	}
@@ -572,7 +575,7 @@ public:
 #else
 		SourceMesh = StaticMesh->GetMeshDescription(LODIndex);
 #endif
-		Adapter = MakeUnique<FMeshDescriptionTriangleMeshSurfaceAdapter>(SourceMesh, StaticMesh);
+		Adapter = MakeUnique<FMeshDescriptionTriangleMeshSurfaceAdapter>(SourceMesh, StaticMesh, BuildOptions.bOnlySurfaceMaterials);
 
 		if (SourceMesh && Adapter)
 		{
@@ -904,7 +907,7 @@ void FMeshSceneAdapter::Build(const FMeshSceneAdapterBuildOptions& BuildOptions)
 
 		if (BuildOptions.bPrintDebugMessages)
 		{
-			UE_LOG(LogTemp, Warning, TEXT("[FMeshSceneAdapter] decomposed %d source meshes into %d unique meshes containing %d triangles"), DecomposedSourceMeshCount.load(), DecomposedMeshesCount.load(), AddedTrisCount)
+			UE_LOG(LogGeometry, Warning, TEXT("[FMeshSceneAdapter] decomposed %d source meshes into %d unique meshes containing %d triangles"), DecomposedSourceMeshCount.load(), DecomposedMeshesCount.load(), AddedTrisCount)
 		}
 
 	}
@@ -1432,7 +1435,7 @@ void FMeshSceneAdapter::Build_FullDecompose(const FMeshSceneAdapterBuildOptions&
 
 	if (BuildOptions.bPrintDebugMessages)
 	{
-		UE_LOG(LogTemp, Warning, TEXT("[FMeshSceneAdapter] decomposed %d source meshes used in %d instances (of %d total source meshes with %ld unique triangles), into %d new part meshes used in %d new instances containing %ld unique triangles. Scene has %ld total unique and %ld total instanced. Skipped %d decompositions. Skipped %d tiny components (%d instances, %d total triangles). %d 1-triangle meshes."), 
+		UE_LOG(LogGeometry, Warning, TEXT("[FMeshSceneAdapter] decomposed %d source meshes used in %d instances (of %d total source meshes with %ld unique triangles), into %d new part meshes used in %d new instances containing %ld unique triangles. Scene has %ld total unique and %ld total instanced. Skipped %d decompositions. Skipped %d tiny components (%d instances, %d total triangles). %d 1-triangle meshes."), 
 			   DecomposedSourceMeshCount.load(), SourceInstancesCount.load(), NumInitialSources, NumSourceUniqueTris, 
 			   NewUniqueMeshesCount.load(), NewInstancesCount.load(), AddedUniqueTrisCount, 
 			   NumFinalUniqueTris.load(), InstancedTrisCount,
diff --git a/Engine/Plugins/Experimental/MeshModelingToolsetExp/Source/ModelingComponents/Public/Scene/MeshSceneAdapter.h b/Engine/Plugins/Experimental/MeshModelingToolsetExp/Source/ModelingComponents/Public/Scene/MeshSceneAdapter.h
--- a/Engine/Plugins/Experimental/MeshModelingToolsetExp/Source/ModelingComponents/Public/Scene/MeshSceneAdapter.h
+++ b/Engine/Plugins/Experimental/MeshModelingToolsetExp/Source/ModelingComponents/Public/Scene/MeshSceneAdapter.h
@@ -64,6 +64,7 @@ struct MODELINGCOMPONENTS_API FMeshTypeContainer
  */
 struct FMeshSceneAdapterBuildOptions
 {
+	/** If true, various build log messages and statistics will be written to LogGeometry */
 	bool bPrintDebugMessages = false;
 
 	/** If true, find approximately-planar meshes with a main dimension below DesiredMinThickness and thicken them to DesiredMinThickness  */
@@ -76,7 +77,10 @@ struct FMeshSceneAdapterBuildOptions
 	/** If bFilterTinyObjects is enabled, then a tiny object is identified by having a maximum (transformed) bounding-box dimension below this size */
 	double TinyObjectBoxMaxDimension = 0.001;
 
+	/** If true, only mesh sections that are assigned a valid surface material (ie MaterialDomain::MD_Surface) will be included/processed in the Mesh Scene. This filters out Decals, for example. */
+	bool bOnlySurfaceMaterials = false;
 
+	/** If true, AABBTree and Fast Winding data structures will be built for unique scene meshes */
 	bool bBuildSpatialDataStructures = true;
 };
 
