MeshModelingTools: Make boolean extrude a separate action (push pull) in PolyEdit. Disallow deletion of entire mesh via push. Fix "shells to solids" not working right in pull. Fix selection not reverting correctly on undoing extrude.

#rnx

diff --git a/Engine/Plugins/Experimental/MeshModelingToolsetExp/Source/MeshModelingToolsExp/Private/EditMeshPolygonsTool.cpp b/Engine/Plugins/Experimental/MeshModelingToolsetExp/Source/MeshModelingToolsExp/Private/EditMeshPolygonsTool.cpp
--- a/Engine/Plugins/Experimental/MeshModelingToolsetExp/Source/MeshModelingToolsExp/Private/EditMeshPolygonsTool.cpp
+++ b/Engine/Plugins/Experimental/MeshModelingToolsetExp/Source/MeshModelingToolsExp/Private/EditMeshPolygonsTool.cpp
@@ -864,6 +864,13 @@ void UEditMeshPolygonsTool::OnTick(float DeltaTime)
 		//Interactive operations:
 		case EEditMeshPolygonsToolActions::Extrude:
 		{
+			ExtrudeActivity->ExtrudeProperties->ExtrudeMode = EPolyEditExtrudeMode::MoveAndStitch;
+			StartActivity(ExtrudeActivity);
+			break;
+		}
+		case EEditMeshPolygonsToolActions::PushPull:
+		{
+			ExtrudeActivity->ExtrudeProperties->ExtrudeMode = EPolyEditExtrudeMode::Boolean;
 			StartActivity(ExtrudeActivity);
 			break;
 		}
diff --git a/Engine/Plugins/Experimental/MeshModelingToolsetExp/Source/MeshModelingToolsExp/Private/ToolActivities/PolyEditExtrudeActivity.cpp b/Engine/Plugins/Experimental/MeshModelingToolsetExp/Source/MeshModelingToolsExp/Private/ToolActivities/PolyEditExtrudeActivity.cpp
--- a/Engine/Plugins/Experimental/MeshModelingToolsetExp/Source/MeshModelingToolsExp/Private/ToolActivities/PolyEditExtrudeActivity.cpp
+++ b/Engine/Plugins/Experimental/MeshModelingToolsetExp/Source/MeshModelingToolsExp/Private/ToolActivities/PolyEditExtrudeActivity.cpp
@@ -144,13 +144,6 @@ EToolActivityStartResult UPolyEditExtrudeActivity::Start()
 			NewSelectionGids = Op->ExtrudedFaceNewGids;
 		});
 
-	// Temporarily clear the selection to avoid the highlighting on the preview as we extrude (which 
-	// sometimes highlights incorrect triangles in boolean mode). There's currently not a good way to
-	// disable those secondary triangle buffers without losing the filter function.
-	// The selection gets reset on End(), either to a new selection, or to the old one on cancel.
-	ActiveSelection = ActivityContext->SelectionMechanic->GetActiveSelection();
-	ActivityContext->SelectionMechanic->SetSelection(FGroupTopologySelection(), false);
-
 	SetToolPropertySourceEnabled(ExtrudeProperties, true);
 
 	// Set up the basics of the extrude height mechanic.
@@ -184,6 +177,14 @@ void UPolyEditExtrudeActivity::BeginExtrude()
 	using namespace PolyEditExtrudeActivityLocals;
 	using UE::Geometry::FTransform3d;
 
+	ActiveSelection = ActivityContext->SelectionMechanic->GetActiveSelection();
+
+	// Temporarily clear the selection to avoid the highlighting on the preview as we extrude (which 
+	// sometimes highlights incorrect triangles in boolean mode). There's currently not a good way to
+	// disable those secondary triangle buffers without losing the filter function.
+	// The selection gets reset on ApplyExtrude() to a new selection, or to the old selection in EndInternal.
+	ActivityContext->SelectionMechanic->SetSelection(FGroupTopologySelection(), false);
+
 	TArray<int32> ActiveTriangleSelection;
 	ActivityContext->CurrentTopology->GetSelectedTriangles(ActiveSelection, ActiveTriangleSelection);
 
@@ -257,10 +258,6 @@ EToolActivityEndResult UPolyEditExtrudeActivity::End(EToolShutdownType ShutdownT
 		return EToolActivityEndResult::ErrorDuringEnd;
 	}
 
-	// Remember to reset the selection. Even if we change it to something else when applying the
-	// extrude, the previous selection needs to be valid.
-	ActivityContext->SelectionMechanic->SetSelection(ActiveSelection, false);
-
 	if (ShutdownType == EToolShutdownType::Cancel)
 	{
 		EndInternal();
@@ -286,6 +283,12 @@ EToolActivityEndResult UPolyEditExtrudeActivity::End(EToolShutdownType ShutdownT
 // Does whatever kind of cleanup we need to do to end the activity.
 void UPolyEditExtrudeActivity::EndInternal()
 {
+	if (ActivityContext->SelectionMechanic->GetActiveSelection() != ActiveSelection)
+	{
+		// We haven't reset the selection back from when we cleared it to avoid the highlighting.
+		ActivityContext->SelectionMechanic->SetSelection(ActiveSelection, false);
+	}
+
 	ActivityContext->Preview->ClearOpFactory();
 	ActivityContext->Preview->OnOpCompleted.RemoveAll(this);
 	PatchMesh.Reset();
@@ -301,6 +304,17 @@ void UPolyEditExtrudeActivity::ApplyExtrude()
 
 	const FDynamicMesh3* ResultMesh = ActivityContext->Preview->PreviewMesh->GetMesh();
 
+	if (ResultMesh->TriangleCount() == 0)
+	{
+		ParentTool->GetToolManager()->DisplayMessage(
+			LOCTEXT("OnDeleteAllFailedMessage", "Cannot Delete Entire Mesh"),
+			EToolMessageLevel::UserWarning);
+
+		// Reset the preview
+		ActivityContext->Preview->PreviewMesh->UpdatePreview(ActivityContext->CurrentMesh.Get());
+		return;
+	}
+
 	// Prep for undo.
 	// TODO: It would be nice if we could attach the change tracker to the op and have it actually
 	// track the changed triangles instead of us having to know which ones to save here. This is
@@ -349,9 +363,18 @@ void UPolyEditExtrudeActivity::ApplyExtrude()
 		NewSelection.SelectedGroupIDs.Append(NewSelectionGids);
 	}
 
+	// We need to reset the old selection back before we give the new one, so
+	// that undo reverts back to the correct selection state.
+	if (ActivityContext->SelectionMechanic->GetActiveSelection() != ActiveSelection)
+	{
+		ActivityContext->SelectionMechanic->SetSelection(ActiveSelection, false);
+	}
+
 	// Emit undo  (also updates relevant structures)
 	ActivityContext->EmitCurrentMeshChangeAndUpdate(LOCTEXT("PolyMeshExtrudeChange", "Extrude"),
 		ChangeTracker.EndChange(), NewSelection, true);
+
+	ActiveSelection = NewSelection;
 }
 
 // Gets used to set the direction along which we measure the distance to extrude.
diff --git a/Engine/Plugins/Experimental/MeshModelingToolsetExp/Source/MeshModelingToolsExp/Public/EditMeshPolygonsTool.h b/Engine/Plugins/Experimental/MeshModelingToolsetExp/Source/MeshModelingToolsExp/Public/EditMeshPolygonsTool.h
--- a/Engine/Plugins/Experimental/MeshModelingToolsetExp/Source/MeshModelingToolsExp/Public/EditMeshPolygonsTool.h
+++ b/Engine/Plugins/Experimental/MeshModelingToolsetExp/Source/MeshModelingToolsExp/Public/EditMeshPolygonsTool.h
@@ -95,6 +95,7 @@ enum class EEditMeshPolygonsToolActions
 	NoAction,
 	CancelCurrent,
 	Extrude,
+	PushPull,
 	Inset,
 	Outset,
 	InsertEdge,
@@ -179,10 +180,14 @@ class MESHMODELINGTOOLSEXP_API UEditMeshPolygonsToolActions : public UEditMeshPo
 {
 	GENERATED_BODY()
 public:
-	/** Extrude the current set of selected faces. Click in viewport to confirm extrude height. */
+	/** Extrude the current set of selected faces by moving and stitching them. */
 	UFUNCTION(CallInEditor, Category = FaceEdits, meta = (DisplayName = "Extrude", DisplayPriority = 1))
 	void Extrude() { PostAction(EEditMeshPolygonsToolActions::Extrude); }
 
+	/** Like Extrude, but performed in a boolean way, meaning that the faces can cut away the mesh or bridge mesh parts. */
+	UFUNCTION(CallInEditor, Category = FaceEdits, meta = (DisplayName = "Push/Pull", DisplayPriority = 1))
+	void PushPull() { PostAction(EEditMeshPolygonsToolActions::PushPull); }
+
 	/** Inset the current set of selected faces. Click in viewport to confirm inset distance. */
 	UFUNCTION(CallInEditor, Category = FaceEdits, meta = (DisplayName = "Inset", DisplayPriority = 2))
 	void Inset() { PostAction(EEditMeshPolygonsToolActions::Inset);	}
@@ -191,11 +196,14 @@ public:
 	UFUNCTION(CallInEditor, Category = FaceEdits, meta = (DisplayName = "Outset", DisplayPriority = 3))
 	void Outset() { PostAction(EEditMeshPolygonsToolActions::Outset);	}
 
-	/** Merge the current set of selected faces into a single face. */
+	//~ TODO: Make the Merge and Delete comments visible as tooltips. Currently we can't due to a bug that
+	//~ limits our total tooltip text allotment: UE-124608
+
+	//~ Merge the current set of selected faces into a single face.
 	UFUNCTION(CallInEditor, Category = FaceEdits, meta = (DisplayName = "Merge", DisplayPriority = 4))
 	void Merge() { PostAction(EEditMeshPolygonsToolActions::Merge);	}
 
-	/** Delete the current set of selected faces */
+	//~ Delete the current set of selected faces
 	UFUNCTION(CallInEditor, Category = FaceEdits, meta = (DisplayName = "Delete", DisplayPriority = 4))
 	void Delete() { PostAction(EEditMeshPolygonsToolActions::Delete); }
 
diff --git a/Engine/Plugins/Experimental/MeshModelingToolsetExp/Source/MeshModelingToolsExp/Public/ToolActivities/PolyEditExtrudeActivity.h b/Engine/Plugins/Experimental/MeshModelingToolsetExp/Source/MeshModelingToolsExp/Public/ToolActivities/PolyEditExtrudeActivity.h
--- a/Engine/Plugins/Experimental/MeshModelingToolsetExp/Source/MeshModelingToolsExp/Public/ToolActivities/PolyEditExtrudeActivity.h
+++ b/Engine/Plugins/Experimental/MeshModelingToolsetExp/Source/MeshModelingToolsExp/Public/ToolActivities/PolyEditExtrudeActivity.h
@@ -36,19 +36,19 @@ class MESHMODELINGTOOLSEXP_API UPolyEditExtrudeProperties : public UInteractiveT
 	GENERATED_BODY()
 
 public:
-	UPROPERTY(EditAnywhere, Category = Extrude)
+	// Not visible in details panel because we decided that we would use a separate "Push/Pull" button.
 	EPolyEditExtrudeMode ExtrudeMode = EPolyEditExtrudeMode::MoveAndStitch;
 
 	/** Which way to move vertices during the extrude */
-	UPROPERTY(EditAnywhere, Category = Extrude)
+	UPROPERTY(EditAnywhere, Category = ExtrusionOptions)
 	EPolyEditExtrudeDirectionMode DirectionMode = EPolyEditExtrudeDirectionMode::SelectedTriangleNormalsEven;
 
 	/** What axis to measure the extrusion distance along. When the direction mode is Single Direction, also controls the direction. */
-	UPROPERTY(EditAnywhere, Category = Extrude, AdvancedDisplay)
+	UPROPERTY(EditAnywhere, Category = ExtrusionOptions, AdvancedDisplay)
 	EPolyEditExtrudeDirection MeasureDirection = EPolyEditExtrudeDirection::SelectionNormal;
 
 	/** Controls whether extruding an entire open-border patch should create a solid or an open shell */
-	UPROPERTY(EditAnywhere, Category = Extrude)
+	UPROPERTY(EditAnywhere, Category = ExtrusionOptions)
 	bool bShellsToSolids = true;
 
 	/** 
diff --git a/Engine/Plugins/Experimental/MeshModelingToolsetExp/Source/ModelingOperators/Private/DeformationOps/ExtrudeOp.cpp b/Engine/Plugins/Experimental/MeshModelingToolsetExp/Source/ModelingOperators/Private/DeformationOps/ExtrudeOp.cpp
--- a/Engine/Plugins/Experimental/MeshModelingToolsetExp/Source/ModelingOperators/Private/DeformationOps/ExtrudeOp.cpp
+++ b/Engine/Plugins/Experimental/MeshModelingToolsetExp/Source/ModelingOperators/Private/DeformationOps/ExtrudeOp.cpp
@@ -7,6 +7,7 @@
 #include "DynamicSubmesh3.h"
 #include "Operations/OffsetMeshRegion.h"
 #include "Operations/MeshBoolean.h"
+#include "Selections/MeshConnectedComponents.h"
 
 #include "Util/ProgressCancel.h"
 
@@ -137,9 +138,43 @@ bool FExtrudeOp::BooleanExtrude(FProgressCancel* Progress)
 		return false;
 	}
 
+	// For whole-component extrusions, the boolean operation that we are about to perform typically
+	// results in the removal of the original triangles, resulting in an open shell. If we want
+	// the shell to be a closed solid, we actually just want the extruded portion in the boolean,
+	// with no participation from the original triangles. So we actually have to delete the whole-component
+	// triangles before combining with the extruded piece(s)
+	if (bShellsToSolids)
+	{
+		FMeshConnectedComponents SelectionComponents(ResultMesh.Get());
+		SelectionComponents.FindTrianglesConnectedToSeeds(TriangleSelection);
+
+		TSet<int32> TriangleSelectionSet(TriangleSelection);
+		for (FMeshConnectedComponents::FComponent& Component : SelectionComponents.Components)
+		{
+			bool bWholeComponentIsSelected = true;
+			for (int32 Tid : Component.Indices)
+			{
+				if (!TriangleSelectionSet.Contains(Tid))
+				{
+					bWholeComponentIsSelected = false;
+					break;
+				}
+			}
+
+			// Perform the deletion if needed
+			if (bWholeComponentIsSelected)
+			{
+				for (int32 Tid : Component.Indices)
+				{
+					ResultMesh->RemoveTriangle(Tid);
+				}
+			}
+		}//end iterating through selection components
+	}
+
 	// Now perform a boolean operation with our result.
 	FMeshBoolean::EBooleanOp Op = ExtrudeDistance > 0 ? FMeshBoolean::EBooleanOp::Union : FMeshBoolean::EBooleanOp::Difference;
-	FMeshBoolean MeshBoolean(OriginalMesh.Get(), FTransform3d::Identity(), &Submesh->GetSubmesh(), FTransform3d::Identity(), ResultMesh.Get(), Op);
+	FMeshBoolean MeshBoolean(ResultMesh.Get(), FTransform3d::Identity(), &Submesh->GetSubmesh(), FTransform3d::Identity(), ResultMesh.Get(), Op);
 	MeshBoolean.bPutResultInInputSpace = true;
 	MeshBoolean.bSimplifyAlongNewEdges = true;
 	MeshBoolean.bPopulateSecondMeshGroupMap = true;
