ModelingTools:
Created a new option field Subdivision Type which allows to choose between exsiting loop-style subdivision and a new PN Triangle style subdivision.

#rnx

diff --git a/Engine/Plugins/Experimental/MeshModelingToolsetExp/Source/MeshModelingToolsExp/Private/DisplaceMeshTool.cpp b/Engine/Plugins/Experimental/MeshModelingToolsetExp/Source/MeshModelingToolsExp/Private/DisplaceMeshTool.cpp
--- a/Engine/Plugins/Experimental/MeshModelingToolsetExp/Source/MeshModelingToolsExp/Private/DisplaceMeshTool.cpp
+++ b/Engine/Plugins/Experimental/MeshModelingToolsetExp/Source/MeshModelingToolsExp/Private/DisplaceMeshTool.cpp
@@ -16,6 +16,7 @@
 #include "ProfilingDebugging/ScopedTimers.h"
 #include "MeshDescription.h"
 #include "ModelingToolTargetUtil.h"
+#include "Operations/PNTriangles.h"
 
 // needed to disable normals recalculation on the underlying asset
 #include "AssetUtils/MeshDescriptionUtil.h"
@@ -263,14 +264,15 @@ namespace DisplaceMeshToolLocals{
 	class FSubdivideMeshOp : public FDynamicMeshOperator
 	{
 	public:
-		FSubdivideMeshOp(const FDynamicMesh3& SourceMesh, int SubdivisionsCountIn, TSharedPtr<FIndexedWeightMap, ESPMode::ThreadSafe> WeightMap);
+		FSubdivideMeshOp(const FDynamicMesh3& SourceMesh, ESubdivisionType SubdivisionTypeIn, int SubdivisionsCountIn, TSharedPtr<FIndexedWeightMap, ESPMode::ThreadSafe> WeightMap);
 		void CalculateResult(FProgressCancel* Progress) final;
 	private:
+		ESubdivisionType SubdivisionType;
 		int SubdivisionsCount;
 	};
 
-	FSubdivideMeshOp::FSubdivideMeshOp(const FDynamicMesh3& SourceMesh, int SubdivisionsCountIn, TSharedPtr<FIndexedWeightMap, ESPMode::ThreadSafe> WeightMap)
-		: SubdivisionsCount(SubdivisionsCountIn)
+	FSubdivideMeshOp::FSubdivideMeshOp(const FDynamicMesh3& SourceMesh, ESubdivisionType SubdivisionTypeIn, int SubdivisionsCountIn, TSharedPtr<FIndexedWeightMap, ESPMode::ThreadSafe> WeightMap)
+		: SubdivisionType(SubdivisionTypeIn), SubdivisionsCount(SubdivisionsCountIn)
 	{
 		ResultMesh->Copy(SourceMesh);
 
@@ -300,10 +302,29 @@ namespace DisplaceMeshToolLocals{
 		using namespace DisplaceMeshToolLocals;
 
 		// calculate subdivisions (todo: move to elsewhere)
-		for (int ri = 0; ri < SubdivisionsCount; ri++)
+		if (SubdivisionType == ESubdivisionType::Loop) 
 		{
-			if (ProgressCancel && ProgressCancel->Cancelled()) return;
-			SubdivideMesh(*ResultMesh, ProgressCancel);
+			for (int ri = 0; ri < SubdivisionsCount; ri++)
+			{
+				if (ProgressCancel && ProgressCancel->Cancelled()) return;
+				SubdivideMesh(*ResultMesh, ProgressCancel);
+			} 
+		}
+		else if (SubdivisionType == ESubdivisionType::PNTriangles) 
+		{
+			FPNTriangles PNTriangles(ResultMesh.Get());
+			PNTriangles.Progress = ProgressCancel;
+			PNTriangles.TesselationLevel = SubdivisionsCount;
+
+			if (PNTriangles.Validate() == EOperationValidationResult::Ok)
+			{
+				PNTriangles.Compute(); 
+			}
+		}
+		else 
+		{
+			// Unsupported subdivision type
+			checkNoEntry();
 		}
 	}
 
@@ -311,32 +332,48 @@ namespace DisplaceMeshToolLocals{
 	{
 	public:
 		FSubdivideMeshOpFactory(FDynamicMesh3& SourceMeshIn,
+			ESubdivisionType SubdivisionTypeIn,
 			int SubdivisionsCountIn,
 			TSharedPtr<FIndexedWeightMap, ESPMode::ThreadSafe> WeightMapIn)
-			: SourceMesh(SourceMeshIn), SubdivisionsCount(SubdivisionsCountIn), WeightMap(WeightMapIn)
+			: SourceMesh(SourceMeshIn), SubdivisionType(SubdivisionTypeIn), SubdivisionsCount(SubdivisionsCountIn), WeightMap(WeightMapIn)
 		{
 		}
+
+		void SetSubdivisionType(ESubdivisionType SubdivisionTypeIn);
+		ESubdivisionType GetSubdivisionType() const;
+
 		void SetSubdivisionsCount(int SubdivisionsCountIn);
-		int  GetSubdivisionsCount();
+		int  GetSubdivisionsCount() const;
 
 		void SetWeightMap(TSharedPtr<FIndexedWeightMap, ESPMode::ThreadSafe> WeightMapIn);
 
 		TUniquePtr<FDynamicMeshOperator> MakeNewOperator() final
 		{
-			return MakeUnique<FSubdivideMeshOp>(SourceMesh, SubdivisionsCount, WeightMap);
+			return MakeUnique<FSubdivideMeshOp>(SourceMesh, SubdivisionType, SubdivisionsCount, WeightMap);
 		}
 	private:
 		const FDynamicMesh3& SourceMesh;
+		ESubdivisionType SubdivisionType;
 		int SubdivisionsCount;
 		TSharedPtr<FIndexedWeightMap, ESPMode::ThreadSafe> WeightMap;
 	};
 
+	void FSubdivideMeshOpFactory::SetSubdivisionType(ESubdivisionType SubdivisionTypeIn) 
+	{
+		SubdivisionType = SubdivisionTypeIn;
+	}
+
+	ESubdivisionType FSubdivideMeshOpFactory::GetSubdivisionType() const
+	{
+		return SubdivisionType;
+	}
+
 	void FSubdivideMeshOpFactory::SetSubdivisionsCount(int SubdivisionsCountIn)
 	{
 		SubdivisionsCount = SubdivisionsCountIn;
 	}
 
-	int FSubdivideMeshOpFactory::GetSubdivisionsCount()
+	int FSubdivideMeshOpFactory::GetSubdivisionsCount() const
 	{
 		return SubdivisionsCount;
 	}
@@ -852,7 +889,7 @@ void UDisplaceMeshTool::Setup()
 									} );
 
 	ValidateSubdivisions();
-	Subdivider = MakeUnique<FSubdivideMeshOpFactory>(OriginalMesh, CommonProperties->Subdivisions, ActiveWeightMap);
+	Subdivider = MakeUnique<FSubdivideMeshOpFactory>(OriginalMesh, CommonProperties->SubdivisionType, CommonProperties->Subdivisions, ActiveWeightMap);
 
 	StartComputation();
 
@@ -966,7 +1003,19 @@ void UDisplaceMeshTool::OnPropertyModified(UObject* PropertySet, FProperty* Prop
 	
 		bNeedsDisplaced = true;
 
-		if (PropName == GET_MEMBER_NAME_CHECKED(UDisplaceMeshCommonProperties, Subdivisions))
+		if (PropName == GET_MEMBER_NAME_CHECKED(UDisplaceMeshCommonProperties, SubdivisionType))
+		{
+			if (CommonProperties->SubdivisionType != SubdividerDownCast->GetSubdivisionType())
+			{
+				SubdividerDownCast->SetSubdivisionType(CommonProperties->SubdivisionType);
+				bNeedsSubdivided = true;
+			}
+			else
+			{
+				return;
+			}
+		}
+		else if (PropName == GET_MEMBER_NAME_CHECKED(UDisplaceMeshCommonProperties, Subdivisions))
 		{
 			ValidateSubdivisions();
 			if (CommonProperties->Subdivisions != SubdividerDownCast->GetSubdivisionsCount())
diff --git a/Engine/Plugins/Experimental/MeshModelingToolsetExp/Source/MeshModelingToolsExp/Public/DisplaceMeshTool.h b/Engine/Plugins/Experimental/MeshModelingToolsetExp/Source/MeshModelingToolsExp/Public/DisplaceMeshTool.h
--- a/Engine/Plugins/Experimental/MeshModelingToolsetExp/Source/MeshModelingToolsExp/Public/DisplaceMeshTool.h
+++ b/Engine/Plugins/Experimental/MeshModelingToolsetExp/Source/MeshModelingToolsExp/Public/DisplaceMeshTool.h
@@ -43,6 +43,19 @@ enum class EDisplaceMeshToolDisplaceType : uint8
 	SineWave UMETA(DisplayName = "Sine Wave"),
 };
 
+UENUM() 
+enum class ESubdivisionType : uint8  
+{
+	/** Subdivide the mesh using loop-style subdivision. */
+	Loop UMETA(DisplayName = "Loop"),
+	
+	/** Subdivide the mesh using PN triangles which replace each original flat triangle by a curved shape that is 
+		retriangulated into a number of small subtriangles. The geometry of a PN triangle is defined as one cubic Bezier 
+		patch using control points. The patch matches the point and normal information at the vertices of the original 
+		flat triangle.*/
+	PNTriangles UMETA(DisplayName = "PN Triangles"),
+};
+
 /** The basic set of properties shared by (more or less) all DisplacementTypes. */
 UCLASS()
 class MESHMODELINGTOOLSEXP_API UDisplaceMeshCommonProperties : public UInteractiveToolPropertySet
@@ -64,6 +77,10 @@ public:
 		meta = (EditCondition = "DisplacementType == EDisplaceMeshToolDisplaceType::RandomNoise || DisplacementType == EDisplaceMeshToolDisplaceType::PerlinNoise"))
 	int RandomSeed = 31337;
 
+	/** Type of the  mesh subdivision. */
+	UPROPERTY(EditAnywhere, Category = Options)
+	ESubdivisionType SubdivisionType = ESubdivisionType::Loop;
+
 	/** Number of times to subdivide the mesh before displacing it. */
 	UPROPERTY(EditAnywhere, Category = Options,
 		meta = (UIMin = "0", UIMax = "10", ClampMin = "0", ClampMax = "100"))
