ModelingTools: Add Normal Map transfer support to BakeTexture.

#rnx

diff --git a/Engine/Plugins/Experimental/MeshModelingToolsetExp/Source/MeshModelingToolsExp/Private/BakeMeshAttributeMapsTool.cpp b/Engine/Plugins/Experimental/MeshModelingToolsetExp/Source/MeshModelingToolsExp/Private/BakeMeshAttributeMapsTool.cpp
--- a/Engine/Plugins/Experimental/MeshModelingToolsetExp/Source/MeshModelingToolsExp/Private/BakeMeshAttributeMapsTool.cpp
+++ b/Engine/Plugins/Experimental/MeshModelingToolsetExp/Source/MeshModelingToolsExp/Private/BakeMeshAttributeMapsTool.cpp
@@ -111,9 +111,12 @@ TArray<FString> UBakeMeshAttributeMapsToolProperties::GetUVLayerNamesFunc()
 class FMeshMapBakerOp : public TGenericDataOperator<FMeshMapBaker>
 {
 public:
+	using ImagePtr = TSharedPtr<UE::Geometry::TImageBuilder<FVector4f>, ESPMode::ThreadSafe>;
+	
 	// General bake settings
 	TSharedPtr<UE::Geometry::FDynamicMesh3, ESPMode::ThreadSafe> DetailMesh;
 	TSharedPtr<UE::Geometry::FDynamicMeshAABBTree3, ESPMode::ThreadSafe> DetailSpatial;
+	TSharedPtr<UE::Geometry::TMeshTangents<double>, ESPMode::ThreadSafe> DetailMeshTangents;
 	UE::Geometry::FDynamicMesh3* BaseMesh;
 	TUniquePtr<UE::Geometry::FMeshMapBaker> Baker;
 	UBakeMeshAttributeMapsTool::FBakeCacheSettings BakeCacheSettings;
@@ -127,8 +130,11 @@ public:
 	FMeshPropertyMapSettings PropertySettings;
 	FTexture2DImageSettings TextureSettings;
 
+	// NormalMap settings
+	ImagePtr DetailMeshNormalMap;
+	int32 DetailMeshNormalUVLayer = 0;
+
 	// Texture2DImage & MultiTexture settings
-	using ImagePtr = TSharedPtr<UE::Geometry::TImageBuilder<FVector4f>, ESPMode::ThreadSafe>;
 	const FDynamicMeshUVOverlay* UVOverlay = nullptr;
 	ImagePtr TextureImage;
 	TMap<int32, ImagePtr> MaterialToTextureImageMap;
@@ -147,6 +153,9 @@ public:
 		Baker->SetThickness(BakeCacheSettings.Thickness);
 		Baker->SetMultisampling(BakeCacheSettings.Multisampling);
 		Baker->SetTargetMeshTangents(BaseMeshTangents);
+		Baker->SetDetailMeshTangents(DetailMeshTangents);
+		Baker->SetDetailMeshNormalMap(DetailMeshNormalMap);
+		Baker->SetDetailMeshNormalUVLayer(DetailMeshNormalUVLayer);
 
 		for (const EBakeMapType MapType : ALL_BAKE_MAP_TYPES)
 		{
@@ -355,6 +364,27 @@ void UBakeMeshAttributeMapsTool::Setup()
 	Settings->WatchProperty(Settings->bUseWorldSpace, [this](bool) { bDetailMeshValid = false; bInputsDirty = true; });
 	Settings->WatchProperty(Settings->Thickness, [this](float) { bInputsDirty = true; });
 	Settings->WatchProperty(Settings->Multisampling, [this](EBakeMultisampling) { bInputsDirty = true; });
+
+
+	UToolTarget* DetailTarget = Targets[bIsBakeToSelf ? 0 : 1];
+	IStaticMeshBackedTarget* DetailStaticMeshTarget = Cast<IStaticMeshBackedTarget>(DetailTarget);
+	UStaticMesh* DetailStaticMesh = DetailStaticMeshTarget ? DetailStaticMeshTarget->GetStaticMesh() : nullptr;
+
+	DetailMeshProps = NewObject<UDetailMeshToolProperties>(this);
+	AddToolPropertySource(DetailMeshProps);
+	SetToolPropertySourceEnabled(DetailMeshProps, true);
+	DetailMeshProps->DetailMesh = DetailStaticMesh;
+	DetailMeshProps->DetailMeshNormalMap = nullptr;
+	DetailMeshProps->WatchProperty(DetailMeshProps->DetailNormalUVLayer, [this](int) { bInputsDirty = true; });
+	DetailMeshProps->WatchProperty(DetailMeshProps->DetailMeshNormalMap, [this](UTexture2D*)
+	{
+		// Only invalidate detail mesh if we need to recompute tangents.
+		if (!DetailMeshTangents)
+		{
+			bDetailMeshValid = false;
+		}
+		bInputsDirty = true;
+	});
 	
 
 	NormalMapProps = NewObject<UBakedNormalMapToolProperties>(this);
@@ -449,12 +479,19 @@ TUniquePtr<UE::Geometry::TGenericDataOperator<FMeshMapBaker>> UBakeMeshAttribute
 	Op->BaseMesh = &BaseMesh;
 	Op->BakeCacheSettings = CachedBakeCacheSettings;
 
-	const EBakeMapType RequiresTangents = EBakeMapType::TangentSpaceNormalMap | EBakeMapType::BentNormal;
+	constexpr EBakeMapType RequiresTangents = EBakeMapType::TangentSpaceNormalMap | EBakeMapType::BentNormal;
 	if ((bool)(CachedBakeCacheSettings.BakeMapTypes & RequiresTangents))
 	{
 		Op->BaseMeshTangents = BaseMeshTangents;
 	}
 
+	if (CachedDetailNormalMap)
+	{
+		Op->DetailMeshTangents = DetailMeshTangents;
+		Op->DetailMeshNormalMap = CachedDetailNormalMap;
+		Op->DetailMeshNormalUVLayer = CachedDetailMeshSettings.UVLayer;
+	}
+
 	if ((bool)(CachedBakeCacheSettings.BakeMapTypes & EBakeMapType::TangentSpaceNormalMap))
 	{
 		Op->NormalSettings = CachedNormalMapSettings;
@@ -721,20 +758,32 @@ void UBakeMeshAttributeMapsTool::UpdateDetailMesh()
 {
 	UToolTarget* DetailTarget = Targets[bIsBakeToSelf ? 0 : 1];
 
-	DetailMesh = MakeShared<FDynamicMesh3, ESPMode::ThreadSafe>(UE::ToolTarget::GetDynamicMeshCopy(DetailTarget));
+	const bool bWantMeshTangents = (DetailMeshProps->DetailMeshNormalMap != nullptr);
+	DetailMesh = MakeShared<FDynamicMesh3, ESPMode::ThreadSafe>(UE::ToolTarget::GetDynamicMeshCopy(DetailTarget, bWantMeshTangents));
 
 	if (Settings->bUseWorldSpace && bIsBakeToSelf == false)
 	{
 		using FTransform3d = UE::Geometry::FTransform3d;
-		FTransform3d DetailToWorld = UE::ToolTarget::GetLocalToWorldTransform(DetailTarget);
+		const FTransform3d DetailToWorld = UE::ToolTarget::GetLocalToWorldTransform(DetailTarget);
 		MeshTransforms::ApplyTransform(*DetailMesh, DetailToWorld);
-		FTransform3d WorldToBase = UE::ToolTarget::GetLocalToWorldTransform(Targets[0]);
+		const FTransform3d WorldToBase = UE::ToolTarget::GetLocalToWorldTransform(Targets[0]);
 		MeshTransforms::ApplyTransform(*DetailMesh, WorldToBase.Inverse());
 	}
 	
 	DetailSpatial = MakeShared<FDynamicMeshAABBTree3, ESPMode::ThreadSafe>();
 	DetailSpatial->SetMesh(DetailMesh.Get(), true);
 
+	// Extract tangents if a DetailMesh normal map was provided.
+	if (bWantMeshTangents)
+	{
+		DetailMeshTangents = MakeShared<FMeshTangentsd, ESPMode::ThreadSafe>(DetailMesh.Get());
+		DetailMeshTangents->CopyTriVertexTangents(*DetailMesh);
+	}
+	else
+	{
+		DetailMeshTangents = nullptr;
+	}
+
 	GetTexturesFromDetailMesh(UE::ToolTarget::GetTargetComponent(DetailTarget));
 
 	bInputsDirty = true;
@@ -799,6 +848,8 @@ void UBakeMeshAttributeMapsTool::UpdateResult()
 	// Clear our invalid bitflag to check again for valid inputs.
 	OpState = EBakeOpState::Evaluate;
 
+	OpState |= UpdateResult_DetailNormalMap();
+
 	// Update map type settings
 	if ((bool)(CachedBakeCacheSettings.BakeMapTypes & EBakeMapType::TangentSpaceNormalMap))
 	{
@@ -854,6 +905,43 @@ void UBakeMeshAttributeMapsTool::UpdateResult()
 }
 
 
+EBakeOpState UBakeMeshAttributeMapsTool::UpdateResult_DetailNormalMap()
+{
+	EBakeOpState ResultState = EBakeOpState::Complete;
+	
+	const FDynamicMeshUVOverlay* UVOverlay = DetailMesh->Attributes()->GetUVLayer(DetailMeshProps->DetailNormalUVLayer);
+	if (UVOverlay == nullptr)
+	{
+		GetToolManager()->DisplayMessage(LOCTEXT("InvalidUVWarning", "The Detail Mesh does not have the selected UV layer"), EToolMessageLevel::UserWarning);
+		return EBakeOpState::Invalid;
+	}
+	
+	UTexture2D* DetailMeshNormalMap = DetailMeshProps->DetailMeshNormalMap; 
+	if (DetailMeshNormalMap)
+	{
+		CachedDetailNormalMap = MakeShared<UE::Geometry::TImageBuilder<FVector4f>, ESPMode::ThreadSafe>();
+		if (!UE::AssetUtils::ReadTexture(DetailMeshNormalMap, *CachedDetailNormalMap, bPreferPlatformData))
+		{
+			// Report the failed texture read as a warning, but permit the bake to continue.
+			GetToolManager()->DisplayMessage(LOCTEXT("CannotReadTextureWarning", "Cannot read from the detail normal map texture"), EToolMessageLevel::UserWarning);
+		}
+	}
+	else
+	{
+		CachedDetailNormalMap = nullptr;
+	}
+
+	FDetailMeshSettings DetailMeshSettings;
+	DetailMeshSettings.UVLayer = DetailMeshProps->DetailNormalUVLayer;
+
+	if (!(CachedDetailMeshSettings == DetailMeshSettings))
+	{
+		CachedDetailMeshSettings = DetailMeshSettings;
+		ResultState = EBakeOpState::Evaluate;
+	}
+	return ResultState;
+}
+
 
 EBakeOpState UBakeMeshAttributeMapsTool::UpdateResult_Normal()
 {
diff --git a/Engine/Plugins/Experimental/MeshModelingToolsetExp/Source/MeshModelingToolsExp/Public/BakeMeshAttributeMapsTool.h b/Engine/Plugins/Experimental/MeshModelingToolsetExp/Source/MeshModelingToolsExp/Public/BakeMeshAttributeMapsTool.h
--- a/Engine/Plugins/Experimental/MeshModelingToolsetExp/Source/MeshModelingToolsExp/Public/BakeMeshAttributeMapsTool.h
+++ b/Engine/Plugins/Experimental/MeshModelingToolsetExp/Source/MeshModelingToolsExp/Public/BakeMeshAttributeMapsTool.h
@@ -177,6 +177,9 @@ protected:
 	UPROPERTY()
 	TObjectPtr<UBakeMeshAttributeMapsToolProperties> Settings;
 
+	UPROPERTY()
+	TObjectPtr<UDetailMeshToolProperties> DetailMeshProps;
+
 	UPROPERTY()
 	TObjectPtr<UBakedNormalMapToolProperties> NormalMapProps;
 
@@ -222,6 +225,7 @@ protected:
 
 	TSharedPtr<UE::Geometry::FDynamicMesh3, ESPMode::ThreadSafe> DetailMesh;
 	TSharedPtr<UE::Geometry::FDynamicMeshAABBTree3, ESPMode::ThreadSafe> DetailSpatial;
+	TSharedPtr<UE::Geometry::TMeshTangents<double>, ESPMode::ThreadSafe> DetailMeshTangents;
 	int32 DetailMeshTimestamp = 0;
 	void UpdateDetailMesh();
 	bool bDetailMeshValid = false;
@@ -265,6 +269,10 @@ protected:
 	using CachedMapIndex = TMap<EBakeMapType, int32>;
 	CachedMapIndex CachedMapIndices;
 
+	FDetailMeshSettings CachedDetailMeshSettings;
+	TSharedPtr<UE::Geometry::TImageBuilder<FVector4f>, ESPMode::ThreadSafe> CachedDetailNormalMap;
+	EBakeOpState UpdateResult_DetailNormalMap();
+
 	FNormalMapSettings CachedNormalMapSettings;
 	EBakeOpState UpdateResult_Normal();
 
diff --git a/Engine/Plugins/Experimental/MeshModelingToolsetExp/Source/MeshModelingToolsExp/Public/BakeMeshAttributeToolCommon.h b/Engine/Plugins/Experimental/MeshModelingToolsetExp/Source/MeshModelingToolsExp/Public/BakeMeshAttributeToolCommon.h
--- a/Engine/Plugins/Experimental/MeshModelingToolsetExp/Source/MeshModelingToolsExp/Public/BakeMeshAttributeToolCommon.h
+++ b/Engine/Plugins/Experimental/MeshModelingToolsetExp/Source/MeshModelingToolsExp/Public/BakeMeshAttributeToolCommon.h
@@ -18,6 +18,26 @@ using UE::Geometry::TImageBuilder;
  * Bake tool property sets
  */
 
+
+UCLASS()
+class MESHMODELINGTOOLSEXP_API UDetailMeshToolProperties : public UInteractiveToolPropertySet
+{
+	GENERATED_BODY()
+public:
+	/** The detail mesh to sample */
+	UPROPERTY(VisibleAnywhere, Category = DetailMesh, meta = (TransientToolProperty))
+	TObjectPtr<UStaticMesh> DetailMesh = nullptr;
+
+	/** The detail mesh normal map to sample. If empty, the geometric normals will be used. */
+	UPROPERTY(EditAnywhere, Category = DetailMesh, meta = (TransientToolProperty))
+	TObjectPtr<UTexture2D> DetailMeshNormalMap = nullptr;
+
+	/** UV layer to sample from on the detail mesh */
+	UPROPERTY(EditAnywhere, Category = DetailMesh, meta = (TransientToolProperty))
+	int32 DetailNormalUVLayer = 0;
+};
+
+
 UENUM()
 enum class ENormalMapSpace
 {
@@ -213,6 +233,16 @@ public:
  * Bake tool property settings structs
  */
 
+struct FDetailMeshSettings
+{
+	int32 UVLayer = 0;
+	
+	bool operator==(const FDetailMeshSettings& Other) const
+	{
+		return UVLayer == Other.UVLayer;
+	}
+};
+
 struct FNormalMapSettings
 {
 	FImageDimensions Dimensions;
diff --git a/Engine/Plugins/Runtime/GeometryProcessing/Source/DynamicMesh/Private/Sampling/MeshNormalMapEvaluator.cpp b/Engine/Plugins/Runtime/GeometryProcessing/Source/DynamicMesh/Private/Sampling/MeshNormalMapEvaluator.cpp
--- a/Engine/Plugins/Runtime/GeometryProcessing/Source/DynamicMesh/Private/Sampling/MeshNormalMapEvaluator.cpp
+++ b/Engine/Plugins/Runtime/GeometryProcessing/Source/DynamicMesh/Private/Sampling/MeshNormalMapEvaluator.cpp
@@ -9,24 +9,30 @@ using namespace UE::Geometry;
 
 void FMeshNormalMapEvaluator::Setup(const FMeshBaseBaker& Baker, FEvaluationContext& Context)
 {
-	Context.Evaluate = &EvaluateSample;
+	// Cache data from the baker
+	DetailMesh = Baker.GetDetailMesh();
+	DetailNormalOverlay = Baker.GetDetailMeshNormals();
+	checkSlow(DetailNormalOverlay);
+	const int32 DetailUVLayer = Baker.GetDetailMeshNormalUVLayer();
+	DetailUVOverlay = Baker.GetDetailMeshUVs(DetailUVLayer);
+	checkSlow(DetailUVOverlay);
+	BaseMeshTangents = Baker.GetTargetMeshTangents();
+	DetailMeshTangents = Baker.GetDetailMeshTangents();
+	DetailMeshNormalMap = Baker.GetDetailMeshNormalMap();
+
+	Context.Evaluate = DetailMeshNormalMap ? &EvaluateSample<true> : &EvaluateSample<false>;
 	Context.EvaluateDefault = &EvaluateDefault;
 	Context.EvaluateColor = &EvaluateColor;
 	Context.EvalData = this;
 	Context.AccumulateMode = EAccumulateMode::Add;
 	Context.DataLayout = { EComponents::Float3 };
-
-	// Cache data from the baker
-	DetailMesh = Baker.GetDetailMesh();
-	DetailNormalOverlay = Baker.GetDetailMeshNormals();
-	checkSlow(DetailNormalOverlay);
-	BaseMeshTangents = Baker.GetTargetMeshTangents().Get();
 }
 
+template <bool bUseDetailNormalMap>
 void FMeshNormalMapEvaluator::EvaluateSample(float*& Out, const FCorrespondenceSample& Sample, void* EvalData)
 {
 	FMeshNormalMapEvaluator* Eval = static_cast<FMeshNormalMapEvaluator*>(EvalData);
-	const FVector3f SampleResult = Eval->SampleFunction(Sample);
+	const FVector3f SampleResult = Eval->SampleFunction<bUseDetailNormalMap>(Sample);
 	WriteToBuffer(Out, SampleResult);
 }
 
@@ -45,9 +51,11 @@ void FMeshNormalMapEvaluator::EvaluateColor(const int DataIdx, float*& In, FVect
 	In += 3;
 }
 
+template <bool bUseDetailNormalMap>
 FVector3f FMeshNormalMapEvaluator::SampleFunction(const FCorrespondenceSample& SampleData) const
 {
-	int32 DetailTriID = SampleData.DetailTriID;
+	FVector3f Result = FVector3f::UnitZ();
+	const int32 DetailTriID = SampleData.DetailTriID;
 	if (DetailMesh->IsTriangle(DetailTriID))
 	{
 		// get tangents on base mesh
@@ -61,13 +69,42 @@ FVector3f FMeshNormalMapEvaluator::SampleFunction(const FCorrespondenceSample& S
 		FVector3d DetailNormal;
 		DetailNormalOverlay->GetTriBaryInterpolate<double>(DetailTriID, &SampleData.DetailBaryCoords.X, &DetailNormal.X);
 		Normalize(DetailNormal);
+		
+		if constexpr (bUseDetailNormalMap)
+		{
+			FVector3d DetailTangentX, DetailTangentY;
+			DetailMeshTangents->GetInterpolatedTriangleTangent(
+				SampleData.DetailTriID,
+				SampleData.DetailBaryCoords,
+				DetailTangentX, DetailTangentY);
+			
+			FVector2d DetailUV;
+        	DetailUVOverlay->GetTriBaryInterpolate<double>(DetailTriID, &SampleData.DetailBaryCoords.X, &DetailUV.X);
+			const FVector4f DetailNormalColor4 = SampleNormalMapFunction(DetailUV);
+
+			// Map color space [0,1] to normal space [-1,1]
+			const FVector3f DetailNormalColor(DetailNormalColor4.X, DetailNormalColor4.Y, DetailNormalColor4.Z);
+			const FVector3f DetailNormalTangentSpace = (DetailNormalColor * 2.0f) - FVector3f::One();
+
+			// Convert detail normal tangent space to object space
+			FVector3f DetailNormalObjectSpace = DetailNormalTangentSpace.X * DetailTangentX + DetailNormalTangentSpace.Y * DetailTangentY + DetailNormalTangentSpace.Z * DetailNormal;
+			Normalize(DetailNormalObjectSpace);
+			DetailNormal = DetailNormalObjectSpace;
+		}
 
 		// compute normal in tangent space
-		double dx = DetailNormal.Dot(BaseTangentX);
-		double dy = DetailNormal.Dot(BaseTangentY);
-		double dz = DetailNormal.Dot(SampleData.BaseNormal);
+		const double dx = DetailNormal.Dot(BaseTangentX);
+		const double dy = DetailNormal.Dot(BaseTangentY);
+		const double dz = DetailNormal.Dot(SampleData.BaseNormal);
 
-		return FVector3f((float)dx, (float)dy, (float)dz);
+		Result = FVector3f((float)dx, (float)dy, (float)dz);
 	}
-	return FVector3f::UnitZ();
+	return Result;
+}
+
+FVector4f FMeshNormalMapEvaluator::SampleNormalMapFunction(const FVector2d& UVCoord) const
+{
+	return DetailMeshNormalMap->BilinearSampleUV<float>(UVCoord, FVector4f(0, 0, 0, 1));
 }
+
+
diff --git a/Engine/Plugins/Runtime/GeometryProcessing/Source/DynamicMesh/Private/Sampling/MeshOcclusionMapEvaluator.cpp b/Engine/Plugins/Runtime/GeometryProcessing/Source/DynamicMesh/Private/Sampling/MeshOcclusionMapEvaluator.cpp
--- a/Engine/Plugins/Runtime/GeometryProcessing/Source/DynamicMesh/Private/Sampling/MeshOcclusionMapEvaluator.cpp
+++ b/Engine/Plugins/Runtime/GeometryProcessing/Source/DynamicMesh/Private/Sampling/MeshOcclusionMapEvaluator.cpp
@@ -82,7 +82,7 @@ void FMeshOcclusionMapEvaluator::Setup(const FMeshBaseBaker& Baker, FEvaluationC
 
 	if (WANT_BENT_NORMAL(OcclusionType) && NormalSpace == ESpace::Tangent)
 	{
-		BaseMeshTangents = Baker.GetTargetMeshTangents().Get();
+		BaseMeshTangents = Baker.GetTargetMeshTangents();
 	}
 
 	BiasDotThreshold = FMathd::Cos(FMathd::Clamp(90.0 - BiasAngleDeg, 0.0, 90.0) * FMathd::DegToRad);
diff --git a/Engine/Plugins/Runtime/GeometryProcessing/Source/DynamicMesh/Private/Sampling/MeshPropertyMapEvaluator.cpp b/Engine/Plugins/Runtime/GeometryProcessing/Source/DynamicMesh/Private/Sampling/MeshPropertyMapEvaluator.cpp
--- a/Engine/Plugins/Runtime/GeometryProcessing/Source/DynamicMesh/Private/Sampling/MeshPropertyMapEvaluator.cpp
+++ b/Engine/Plugins/Runtime/GeometryProcessing/Source/DynamicMesh/Private/Sampling/MeshPropertyMapEvaluator.cpp
@@ -9,19 +9,22 @@ using namespace UE::Geometry;
 
 void FMeshPropertyMapEvaluator::Setup(const FMeshBaseBaker& Baker, FEvaluationContext& Context)
 {
-	Context.Evaluate = &EvaluateSample;
-	Context.EvaluateDefault = &EvaluateDefault;
-	Context.EvaluateColor = &EvaluateColor;
-	Context.EvalData = this;
-	Context.AccumulateMode = EAccumulateMode::Add;
-	Context.DataLayout = { EComponents::Float3 };
-
 	// Cache data from the baker
 	DetailMesh = Baker.GetDetailMesh();
 	DetailNormalOverlay = Baker.GetDetailMeshNormals();
 	check(DetailNormalOverlay);
 	DetailUVOverlay = Baker.GetDetailMeshUVs();
 	DetailColorOverlay = Baker.GetDetailMeshColors();
+	DetailMeshTangents = Baker.GetDetailMeshTangents();
+	DetailMeshNormalMap = Baker.GetDetailMeshNormalMap();
+
+	const bool bHasDetailNormalMap = DetailMeshNormalMap != nullptr;
+	Context.Evaluate = bHasDetailNormalMap ? &EvaluateSample<true> : &EvaluateSample<false>;
+	Context.EvaluateDefault = &EvaluateDefault;
+	Context.EvaluateColor = &EvaluateColor;
+	Context.EvalData = this;
+	Context.AccumulateMode = EAccumulateMode::Add;
+	Context.DataLayout = { EComponents::Float3 };
 
 	Bounds = DetailMesh->GetBounds();
 	for (int32 j = 0; j < 3; ++j)
@@ -41,8 +44,14 @@ void FMeshPropertyMapEvaluator::Setup(const FMeshBaseBaker& Baker, FEvaluationCo
 		DefaultValue = PositionToColor(Bounds.Center(), Bounds);
 		break;
 	case EMeshPropertyMapType::FacetNormal:
+		DefaultValue = NormalToColor(FVector3d::UnitZ());
+		break;
 	case EMeshPropertyMapType::Normal:
 		DefaultValue = NormalToColor(FVector3d::UnitZ());
+		if (bHasDetailNormalMap)
+		{
+			DetailUVOverlay = Baker.GetDetailMeshUVs(Baker.GetDetailMeshNormalUVLayer());
+		}
 		break;
 	case EMeshPropertyMapType::UVPosition:
 		DefaultValue = UVToColor(FVector2d::Zero());
@@ -56,10 +65,11 @@ void FMeshPropertyMapEvaluator::Setup(const FMeshBaseBaker& Baker, FEvaluationCo
 	}
 }
 
+template <bool bUseDetailNormalMap>
 void FMeshPropertyMapEvaluator::EvaluateSample(float*& Out, const FCorrespondenceSample& Sample, void* EvalData)
 {
 	FMeshPropertyMapEvaluator* Eval = static_cast<FMeshPropertyMapEvaluator*>(EvalData);
-	FVector3f SampleResult = Eval->SampleFunction(Sample);
+	const FVector3f SampleResult = Eval->SampleFunction<bUseDetailNormalMap>(Sample);
 	WriteToBuffer(Out, SampleResult);
 }
 
@@ -75,6 +85,7 @@ void FMeshPropertyMapEvaluator::EvaluateColor(const int DataIdx, float*& In, FVe
 	In += 3;
 }
 
+template <bool bUseDetailNormalMap>
 FVector3f FMeshPropertyMapEvaluator::SampleFunction(const FCorrespondenceSample& SampleData)
 {
 	FVector3f Color = DefaultValue;
@@ -102,6 +113,29 @@ FVector3f FMeshPropertyMapEvaluator::SampleFunction(const FCorrespondenceSample&
 				FVector3d DetailNormal;
 				DetailNormalOverlay->GetTriBaryInterpolate<double>(DetailTriID, &SampleData.DetailBaryCoords.X, &DetailNormal.X);
 				Normalize(DetailNormal);
+
+				if constexpr (bUseDetailNormalMap)
+				{
+					FVector3d DetailTangentX, DetailTangentY;
+					DetailMeshTangents->GetInterpolatedTriangleTangent(
+						SampleData.DetailTriID,
+						SampleData.DetailBaryCoords,
+						DetailTangentX, DetailTangentY);
+			
+					FVector2d DetailUV;
+					DetailUVOverlay->GetTriBaryInterpolate<double>(DetailTriID, &SampleData.DetailBaryCoords.X, &DetailUV.X);
+					const FVector4f DetailNormalColor4 = SampleNormalMapFunction(DetailUV);
+
+					// Map color space [0,1] to normal space [-1,1]
+					const FVector3f DetailNormalColor(DetailNormalColor4.X, DetailNormalColor4.Y, DetailNormalColor4.Z);
+					const FVector3f DetailNormalTangentSpace = (DetailNormalColor * 2.0f) - FVector3f::One();
+
+					// Convert detail normal tangent space to object space
+					FVector3f DetailNormalObjectSpace = DetailNormalTangentSpace.X * DetailTangentX + DetailNormalTangentSpace.Y * DetailTangentY + DetailNormalTangentSpace.Z * DetailNormal;
+					Normalize(DetailNormalObjectSpace);
+					DetailNormal = DetailNormalObjectSpace;
+				}
+				
 				Color = NormalToColor(DetailNormal);
 			}
 		}
@@ -141,3 +175,9 @@ FVector3f FMeshPropertyMapEvaluator::SampleFunction(const FCorrespondenceSample&
 	return Color;
 }
 
+FVector4f FMeshPropertyMapEvaluator::SampleNormalMapFunction(const FVector2d& UVCoord) const
+{
+	return DetailMeshNormalMap->BilinearSampleUV<float>(UVCoord, FVector4f(0, 0, 0, 1));
+}
+
+
diff --git a/Engine/Plugins/Runtime/GeometryProcessing/Source/DynamicMesh/Public/Sampling/MeshBaseBaker.h b/Engine/Plugins/Runtime/GeometryProcessing/Source/DynamicMesh/Public/Sampling/MeshBaseBaker.h
--- a/Engine/Plugins/Runtime/GeometryProcessing/Source/DynamicMesh/Public/Sampling/MeshBaseBaker.h
+++ b/Engine/Plugins/Runtime/GeometryProcessing/Source/DynamicMesh/Public/Sampling/MeshBaseBaker.h
@@ -6,6 +6,7 @@
 #include "DynamicMesh/DynamicMeshAttributeSet.h"
 #include "DynamicMesh/DynamicMeshAABBTree3.h"
 #include "DynamicMesh/MeshTangents.h"
+#include "Image/ImageBuilder.h"
 
 namespace UE
 {
@@ -40,11 +41,23 @@ public:
 		DetailMesh = Mesh;
 		DetailSpatial = Spatial;
 	}
+	void SetDetailMeshTangents(const TSharedPtr<FMeshTangentsd, ESPMode::ThreadSafe>& Tangents)
+	{
+		DetailMeshTangents = Tangents;
+	}
+	void SetDetailMeshNormalMap(const TSharedPtr<TImageBuilder<FVector4f>, ESPMode::ThreadSafe>& NormalMap)
+	{
+		DetailMeshNormalMap = NormalMap;
+	}
+	void SetDetailMeshNormalUVLayer(int32 Layer)
+	{
+		DetailMeshNormalUVLayer = Layer;
+	}
 	void SetTargetMesh(const FDynamicMesh3* Mesh)
 	{
 		TargetMesh = Mesh;
 	}
-	void SetTargetMeshTangents(const TSharedPtr<FMeshTangentsd, ESPMode::ThreadSafe> Tangents)
+	void SetTargetMeshTangents(const TSharedPtr<FMeshTangentsd, ESPMode::ThreadSafe>& Tangents)
 	{
 		TargetMeshTangents = Tangents;
 	}
@@ -76,9 +89,9 @@ public:
 		check(TargetMesh && TargetMesh->HasAttributes());
 		return TargetMesh->Attributes()->PrimaryNormals();
 	}
-	TSharedPtr<FMeshTangentsd, ESPMode::ThreadSafe> GetTargetMeshTangents() const
+	const FMeshTangentsd* GetTargetMeshTangents() const
 	{
-		return TargetMeshTangents;
+		return TargetMeshTangents.Get();
 	}
 
 	// DetailMesh Getters
@@ -105,6 +118,18 @@ public:
 		check(DetailMesh && DetailMesh->HasAttributes());
 		return DetailMesh->Attributes()->PrimaryColors();
 	}
+	const FMeshTangentsd* GetDetailMeshTangents() const
+	{
+		return DetailMeshTangents.Get();
+	}
+	const TImageBuilder<FVector4f>* GetDetailMeshNormalMap() const
+	{
+		return DetailMeshNormalMap.Get();
+	}
+	int32 GetDetailMeshNormalUVLayer() const
+	{
+		return DetailMeshNormalUVLayer;
+	}
 
 	// Other Getters
 	int32 GetUVLayer() const
@@ -123,6 +148,10 @@ public:
 protected:
 	const FDynamicMesh3* DetailMesh = nullptr;
 	const FDynamicMeshAABBTree3* DetailSpatial = nullptr;
+	TSharedPtr<FMeshTangentsd, ESPMode::ThreadSafe> DetailMeshTangents;
+	TSharedPtr<UE::Geometry::TImageBuilder<FVector4f>, ESPMode::ThreadSafe> DetailMeshNormalMap;
+	int32 DetailMeshNormalUVLayer = 0;
+	
 	const FDynamicMesh3* TargetMesh = nullptr;
 	TSharedPtr<FMeshTangentsd, ESPMode::ThreadSafe> TargetMeshTangents;
 
diff --git a/Engine/Plugins/Runtime/GeometryProcessing/Source/DynamicMesh/Public/Sampling/MeshNormalMapEvaluator.h b/Engine/Plugins/Runtime/GeometryProcessing/Source/DynamicMesh/Public/Sampling/MeshNormalMapEvaluator.h
--- a/Engine/Plugins/Runtime/GeometryProcessing/Source/DynamicMesh/Public/Sampling/MeshNormalMapEvaluator.h
+++ b/Engine/Plugins/Runtime/GeometryProcessing/Source/DynamicMesh/Public/Sampling/MeshNormalMapEvaluator.h
@@ -4,6 +4,7 @@
 
 #include "Sampling/MeshMapEvaluator.h"
 #include "DynamicMesh/MeshTangents.h"
+#include "Image/ImageBuilder.h"
 
 namespace UE
 {
@@ -22,6 +23,7 @@ public:
 	virtual EMeshMapEvaluatorType Type() const override { return EMeshMapEvaluatorType::Normal; }
 	// End FMeshMapEvaluator interface
 
+	template <bool bUseDetailNormalMap>
 	static void EvaluateSample(float*& Out, const FCorrespondenceSample& Sample, void* EvalData);
 
 	static void EvaluateDefault(float*& Out, void* EvalData);
@@ -32,12 +34,18 @@ protected:
 	// Cached data
 	const FDynamicMesh3* DetailMesh = nullptr;
 	const FDynamicMeshNormalOverlay* DetailNormalOverlay = nullptr;
+	const FDynamicMeshUVOverlay* DetailUVOverlay = nullptr;
+	const TMeshTangents<double>* DetailMeshTangents = nullptr;
+	const TImageBuilder<FVector4f>* DetailMeshNormalMap = nullptr;
 	const TMeshTangents<double>* BaseMeshTangents = nullptr;
 
 	FVector3f DefaultNormal = FVector3f::UnitZ();
 
 private:
+	template <bool bUseDetailNormalMap>
 	FVector3f SampleFunction(const FCorrespondenceSample& SampleData) const;
+
+	FVector4f SampleNormalMapFunction(const FVector2d& UVCoord) const;
 };
 
 } // end namespace UE::Geometry
diff --git a/Engine/Plugins/Runtime/GeometryProcessing/Source/DynamicMesh/Public/Sampling/MeshPropertyMapEvaluator.h b/Engine/Plugins/Runtime/GeometryProcessing/Source/DynamicMesh/Public/Sampling/MeshPropertyMapEvaluator.h
--- a/Engine/Plugins/Runtime/GeometryProcessing/Source/DynamicMesh/Public/Sampling/MeshPropertyMapEvaluator.h
+++ b/Engine/Plugins/Runtime/GeometryProcessing/Source/DynamicMesh/Public/Sampling/MeshPropertyMapEvaluator.h
@@ -3,6 +3,8 @@
 #pragma once
 
 #include "Sampling/MeshMapEvaluator.h"
+#include "DynamicMesh/MeshTangents.h"
+#include "Image/ImageBuilder.h"
 
 namespace UE
 {
@@ -34,6 +36,7 @@ public:
 	virtual EMeshMapEvaluatorType Type() const override { return EMeshMapEvaluatorType::Property; }
 	// End FMeshMapEvaluator interface
 
+	template <bool bUseDetailNormalMap>
 	static void EvaluateSample(float*& Out, const FCorrespondenceSample& Sample, void* EvalData);
 
 	static void EvaluateDefault(float*& Out, void* EvalData);
@@ -46,6 +49,8 @@ protected:
 	const FDynamicMeshNormalOverlay* DetailNormalOverlay = nullptr;
 	const FDynamicMeshUVOverlay* DetailUVOverlay = nullptr;
 	const FDynamicMeshColorOverlay* DetailColorOverlay = nullptr;
+	const TMeshTangents<double>* DetailMeshTangents = nullptr;
+	const TImageBuilder<FVector4f>* DetailMeshNormalMap = nullptr;
 	FAxisAlignedBox3d Bounds;
 	FVector3f DefaultValue = FVector3f::Zero();
 
@@ -70,7 +75,10 @@ private:
 		return (FVector3f)FVector3d(X, Y, Z);
 	}
 
+	template <bool bUseDetailNormalMap>
 	FVector3f SampleFunction(const FCorrespondenceSample& Sample);
+
+	FVector4f SampleNormalMapFunction(const FVector2d& UVCoord) const;
 };
 
 } // end namespace UE::Geometry
