ModelingComponents: change the OnDragRectangleFinished event in URectangleMarqueeMechanic to send the FCameraRectangle and bCancelled flag, so that clients can just implement this function in simpler use cases. Add URectangleMarqueeMechanic::bUseExternalClickDragBehavior to allow client to set an external Behavior to operate the marquee-as-IClickDragBehaviorTarget (eg ClickOrDragInputBehavior)

#rnx

diff --git a/Engine/Plugins/Experimental/MeshModelingToolsetExp/Source/ModelingComponents/Private/Mechanics/LatticeControlPointsMechanic.cpp b/Engine/Plugins/Experimental/MeshModelingToolsetExp/Source/ModelingComponents/Private/Mechanics/LatticeControlPointsMechanic.cpp
--- a/Engine/Plugins/Experimental/MeshModelingToolsetExp/Source/ModelingComponents/Private/Mechanics/LatticeControlPointsMechanic.cpp
+++ b/Engine/Plugins/Experimental/MeshModelingToolsetExp/Source/ModelingComponents/Private/Mechanics/LatticeControlPointsMechanic.cpp
@@ -663,7 +663,7 @@ void ULatticeControlPointsMechanic::OnDragRectangleChanged(const FCameraRectangl
 	UpdateDrawables();
 }
 
-void ULatticeControlPointsMechanic::OnDragRectangleFinished()
+void ULatticeControlPointsMechanic::OnDragRectangleFinished(const FCameraRectangle& Rectangle, bool bCancelled)
 {
 	ParentTool->GetToolManager()->BeginUndoTransaction(LatticePointSelectionTransactionText);
 
diff --git a/Engine/Plugins/Experimental/MeshModelingToolsetExp/Source/ModelingComponents/Private/Mechanics/RectangleMarqueeMechanic.cpp b/Engine/Plugins/Experimental/MeshModelingToolsetExp/Source/ModelingComponents/Private/Mechanics/RectangleMarqueeMechanic.cpp
--- a/Engine/Plugins/Experimental/MeshModelingToolsetExp/Source/ModelingComponents/Private/Mechanics/RectangleMarqueeMechanic.cpp
+++ b/Engine/Plugins/Experimental/MeshModelingToolsetExp/Source/ModelingComponents/Private/Mechanics/RectangleMarqueeMechanic.cpp
@@ -137,10 +137,13 @@ void URectangleMarqueeMechanic::Setup(UInteractiveTool* ParentToolIn)
 {
 	UInteractionMechanic::Setup(ParentToolIn);
 
-	ClickDragBehavior = NewObject<UClickDragInputBehavior>(this);
-	ClickDragBehavior->SetDefaultPriority(BasePriority);
-	ClickDragBehavior->Initialize(this);
-	ParentTool->AddInputBehavior(ClickDragBehavior, this);
+	if (bUseExternalClickDragBehavior == false)
+	{
+		ClickDragBehavior = NewObject<UClickDragInputBehavior>(this);
+		ClickDragBehavior->SetDefaultPriority(BasePriority);
+		ClickDragBehavior->Initialize(this);
+		ParentTool->AddInputBehavior(ClickDragBehavior, this);
+	}
 	SetIsEnabled(true);
 }
 
@@ -212,6 +215,7 @@ void URectangleMarqueeMechanic::OnClickDrag(const FInputDeviceRay& DragPos)
 	FRay DragCurrentWorldRay = DragPos.WorldRay;
 
 	FCameraRectangle Rectangle(CachedCameraState, DragStartWorldRay, DragCurrentWorldRay);
+	LastRectangle = Rectangle;
 
 	OnDragRectangleChanged.Broadcast(Rectangle);
 }
@@ -219,13 +223,13 @@ void URectangleMarqueeMechanic::OnClickDrag(const FInputDeviceRay& DragPos)
 void URectangleMarqueeMechanic::OnClickRelease(const FInputDeviceRay& ReleasePos)
 {
 	bIsDragging = false;
-	OnDragRectangleFinished.Broadcast();
+	OnDragRectangleFinished.Broadcast(LastRectangle, false);
 }
 
 void URectangleMarqueeMechanic::OnTerminateDragSequence()
 {
 	bIsDragging = false;
-	OnDragRectangleFinished.Broadcast();
+	OnDragRectangleFinished.Broadcast(LastRectangle, true);
 }
 
 
diff --git a/Engine/Plugins/Experimental/MeshModelingToolsetExp/Source/ModelingComponents/Private/Selection/PolygonSelectionMechanic.cpp b/Engine/Plugins/Experimental/MeshModelingToolsetExp/Source/ModelingComponents/Private/Selection/PolygonSelectionMechanic.cpp
--- a/Engine/Plugins/Experimental/MeshModelingToolsetExp/Source/ModelingComponents/Private/Selection/PolygonSelectionMechanic.cpp
+++ b/Engine/Plugins/Experimental/MeshModelingToolsetExp/Source/ModelingComponents/Private/Selection/PolygonSelectionMechanic.cpp
@@ -642,7 +642,7 @@ void UPolygonSelectionMechanic::OnDragRectangleChanged(const FCameraRectangle& C
 	}
 }
 
-void UPolygonSelectionMechanic::OnDragRectangleFinished()
+void UPolygonSelectionMechanic::OnDragRectangleFinished(const FCameraRectangle& Rectangle, bool bCancelled)
 {
 	if (!(PersistentSelection == PreDragPersistentSelection))
 	{
diff --git a/Engine/Plugins/Experimental/MeshModelingToolsetExp/Source/ModelingComponents/Public/Mechanics/LatticeControlPointsMechanic.h b/Engine/Plugins/Experimental/MeshModelingToolsetExp/Source/ModelingComponents/Public/Mechanics/LatticeControlPointsMechanic.h
--- a/Engine/Plugins/Experimental/MeshModelingToolsetExp/Source/ModelingComponents/Public/Mechanics/LatticeControlPointsMechanic.h
+++ b/Engine/Plugins/Experimental/MeshModelingToolsetExp/Source/ModelingComponents/Public/Mechanics/LatticeControlPointsMechanic.h
@@ -176,7 +176,7 @@ protected:
 	TArray<int32> CurrentDragSelection;
 	void OnDragRectangleStarted();
 	void OnDragRectangleChanged(const FCameraRectangle& Rectangle);
-	void OnDragRectangleFinished();
+	void OnDragRectangleFinished(const FCameraRectangle& Rectangle, bool bCancelled);
 
 	// We need the selected point start positions so we can move multiple points appropriately.
 	TMap<int32, FVector3d> SelectedPointStartPositions;
diff --git a/Engine/Plugins/Experimental/MeshModelingToolsetExp/Source/ModelingComponents/Public/Mechanics/RectangleMarqueeMechanic.h b/Engine/Plugins/Experimental/MeshModelingToolsetExp/Source/ModelingComponents/Public/Mechanics/RectangleMarqueeMechanic.h
--- a/Engine/Plugins/Experimental/MeshModelingToolsetExp/Source/ModelingComponents/Public/Mechanics/RectangleMarqueeMechanic.h
+++ b/Engine/Plugins/Experimental/MeshModelingToolsetExp/Source/ModelingComponents/Public/Mechanics/RectangleMarqueeMechanic.h
@@ -19,8 +19,9 @@ class UClickDragInputBehavior;
 ///		- a 2D basis for coordinates in this plane, and
 ///		- the corners of a rectangle contained in this plane, in this 2D basis
 /// 
-struct FCameraRectangle
+struct MODELINGCOMPONENTS_API FCameraRectangle
 {
+	FCameraRectangle() {};
 	FCameraRectangle(const FViewCameraState& CachedCameraState,
 					 const FRay& DragStartWorldRay,
 					 const FRay& DragEndWorldRay);
@@ -46,9 +47,6 @@ struct FCameraRectangle
 		float V = FVector::DotProduct(Point - CameraPlane.GetOrigin(), VBasisVector);
 		return FVector2D{ U,V };
 	}
-
-private:
-	FCameraRectangle() {};
 };
 
 
@@ -67,6 +65,17 @@ class MODELINGCOMPONENTS_API URectangleMarqueeMechanic : public UInteractionMech
 {
 	GENERATED_BODY()
 
+public:
+	/** 
+	 * If true, then the URectangleMarqueeMechanic will not create an internal UClickDragInputBehavior in ::Setup(), allowing
+	 * the client to control the marquee with an external InputBehavior that uses the marquee mechanic as it's IClickDragBehaviorTarget.
+	 * For instance, this allows the mechanic to be hooked in as the drag component of a USingleClickOrDragInputBehavior.
+	 * Must be configured before calling ::Setup()
+	 */
+	UPROPERTY()
+	bool bUseExternalClickDragBehavior = false;
+
+
 public:
 
 	// UInteractionMechanic
@@ -82,9 +91,9 @@ public:
 	 * Sets the base priority so that users can make sure that their own behaviors are higher
 	 * priority. The mechanic will not use any priority value higher than this.
 	 * Mechanics could use lower priorities (and their range could be inspected with
-	 * GetPriorityRange), but marquee mechanic doesn't.
+	 * GetPriorityRange), but marquee mechanic doesn't. 
 	 *
-	 * Can be called before or after Setup().
+	 * Can be called before or after Setup(). Only relevant if bUseExternalClickDragBehavior is false
 	 */
 	void SetBasePriority(const FInputCapturePriority& Priority);
 
@@ -110,8 +119,11 @@ public:
 
 	/**
 	 * Called once the user lets go of the mouse button after dragging out a rectangle.
+	 * The last dragged rectangle is passed here so that clients can choose to just implement this function in simple cases. 
+	 * bCancelled flag is true when the drag finishes due to a disabling of the mechanic or due to a TerminateDragSequence call, rather than a normal drag completion.
 	 */
-	FSimpleMulticastDelegate OnDragRectangleFinished;
+	DECLARE_MULTICAST_DELEGATE_TwoParams(OnDragRectangleFinishedEvent, const FCameraRectangle&, bool bCancelled);
+	OnDragRectangleFinishedEvent OnDragRectangleFinished;
 	
 protected:
 	UPROPERTY()
@@ -134,4 +146,5 @@ private:
 	FVector2D DragStartScreenPosition;
 	FRay DragStartWorldRay;
 	FVector2D DragCurrentScreenPosition;
+	FCameraRectangle LastRectangle;
 };
diff --git a/Engine/Plugins/Experimental/MeshModelingToolsetExp/Source/ModelingComponents/Public/Selection/PolygonSelectionMechanic.h b/Engine/Plugins/Experimental/MeshModelingToolsetExp/Source/ModelingComponents/Public/Selection/PolygonSelectionMechanic.h
--- a/Engine/Plugins/Experimental/MeshModelingToolsetExp/Source/ModelingComponents/Public/Selection/PolygonSelectionMechanic.h
+++ b/Engine/Plugins/Experimental/MeshModelingToolsetExp/Source/ModelingComponents/Public/Selection/PolygonSelectionMechanic.h
@@ -312,7 +312,7 @@ protected:
 	// These get bound to marquee mechanic delegates.
 	virtual void OnDragRectangleStarted();
 	virtual void OnDragRectangleChanged(const FCameraRectangle& CurrentRectangle);
-	virtual void OnDragRectangleFinished();
+	virtual void OnDragRectangleFinished(const FCameraRectangle& Rectangle, bool bCancelled);
 
 	virtual void UpdateMarqueeEnabled();
 
