ModelingComponents: add MeshSceneAdapter::ParallelMeshVertexEnumeration, allows clients to (relatively) efficiently run computations across all the mesh vertices of the in-scene meshes (ie instances are processed multiple times). This can be expensive!
#rnx

diff --git a/Engine/Plugins/Experimental/MeshModelingToolsetExp/Source/ModelingComponents/Private/Scene/MeshSceneAdapter.cpp b/Engine/Plugins/Experimental/MeshModelingToolsetExp/Source/ModelingComponents/Private/Scene/MeshSceneAdapter.cpp
--- a/Engine/Plugins/Experimental/MeshModelingToolsetExp/Source/ModelingComponents/Private/Scene/MeshSceneAdapter.cpp
+++ b/Engine/Plugins/Experimental/MeshModelingToolsetExp/Source/ModelingComponents/Private/Scene/MeshSceneAdapter.cpp
@@ -384,22 +384,32 @@ public:
 	}
 
 
-	virtual void ProcessVerticesInWorld(TFunctionRef<void(const FVector3d&)> ProcessFunc, const FTransformSequence3d& LocalToWorldTransform) override
+	virtual bool ProcessVerticesInWorld(TFunctionRef<bool(const FVector3d&)> ProcessFunc, const FTransformSequence3d& LocalToWorldTransform) override
 	{
+		bool bContinue = true;
 		if (bHasBakedTransform)
 		{
 			for (FVector3d P : Mesh.VerticesItr())
 			{
-				ProcessFunc(P);
+				bContinue = ProcessFunc(P);
+				if (!bContinue)
+				{
+					break;
+				}
 			}
 		}
 		else
 		{
 			for (FVector3d P : Mesh.VerticesItr())
 			{
-				ProcessFunc(LocalToWorldTransform.TransformPosition(P));
+				bContinue = ProcessFunc(LocalToWorldTransform.TransformPosition(P));
+				if (!bContinue)
+				{
+					break;
+				}
 			}
 		}
+		return bContinue;
 	}
 
 
@@ -647,16 +657,18 @@ public:
 	}
 
 
-	virtual void ProcessVerticesInWorld(TFunctionRef<void(const FVector3d&)> ProcessFunc, const FTransformSequence3d& LocalToWorldTransform) override
+	virtual bool ProcessVerticesInWorld(TFunctionRef<bool(const FVector3d&)> ProcessFunc, const FTransformSequence3d& LocalToWorldTransform) override
 	{
+		bool bContinue = true;
 		int32 NumVertices = (SourceMesh) ? Adapter->VertexCount() : 0;
-		for (int32 vi = 0; vi < NumVertices; ++vi)
+		for (int32 vi = 0; vi < NumVertices && bContinue; ++vi)
 		{
 			if (Adapter->IsVertex(vi))
 			{
-				ProcessFunc(LocalToWorldTransform.TransformPosition(Adapter->GetVertex(vi)));
+				bContinue = ProcessFunc(LocalToWorldTransform.TransformPosition(Adapter->GetVertex(vi)));
 			}
 		}
+		return bContinue;
 	}
 
 	virtual void AppendMesh(FDynamicMesh3& AppendTo, const FTransformSequence3d& TransformSeq) override
@@ -1821,6 +1833,7 @@ void FMeshSceneAdapter::GenerateBaseClosingMesh(double BaseHeight, double Extrud
 			{
 				LocalHullPoints.Add(FVector2d(WorldPos.X, WorldPos.Y));
 			}
+			return true;
 		}, ChildMesh->WorldTransform);
 
 		if (LocalHullPoints.Num() > 0)
@@ -1913,3 +1926,40 @@ void FMeshSceneAdapter::GenerateBaseClosingMesh(double BaseHeight, double Extrud
 
 }
 
+
+
+void FMeshSceneAdapter::ParallelMeshVertexEnumeration(
+	TFunctionRef<bool(int32 NumMeshes)> InitializeFunc,
+	TFunctionRef<bool(int32 MeshIndex, AActor* SourceActor, const FActorChildMesh* ChildMeshInfo, const FAxisAlignedBox3d& WorldBounds)> MeshFilterFunc,
+	TFunctionRef<bool(int32 MeshIndex, AActor* SourceActor, const FActorChildMesh* ChildMeshInfo, const FVector3d& WorldPos)> PerVertexFunc,
+	bool bForceSingleThreaded )
+{
+	int32 N = SortedSpatials.Num();
+	if (InitializeFunc(N) == false)
+	{
+		return;
+	}
+
+	bool bSingleThread = bForceSingleThreaded || (CVarMeshSceneAdapterDisableMultiThreading.GetValueOnAnyThread() != 0);
+
+	ParallelFor(N, [&](int32 Index)
+	{
+		const FSpatialCacheInfo& CacheInfo = SortedSpatials[Index];
+
+		bool bContinue = MeshFilterFunc(Index, CacheInfo.Actor->SourceActor, CacheInfo.ChildMesh, CacheInfo.Bounds);
+		if (!bContinue)
+		{
+			return;
+		}
+
+		FTransformSequence3d WorldTransform = CacheInfo.ChildMesh->WorldTransform;
+		CacheInfo.Spatial->ProcessVerticesInWorld([&](const FVector3d& WorldPos)
+		{
+			bool bContinue = PerVertexFunc(Index, CacheInfo.Actor->SourceActor, CacheInfo.ChildMesh, WorldPos);
+			return bContinue;
+		}, WorldTransform);
+
+	}, bSingleThread ? EParallelForFlags::ForceSingleThread : EParallelForFlags::Unbalanced );
+
+
+}
\ No newline at end of file
diff --git a/Engine/Plugins/Experimental/MeshModelingToolsetExp/Source/ModelingComponents/Public/Scene/MeshSceneAdapter.h b/Engine/Plugins/Experimental/MeshModelingToolsetExp/Source/ModelingComponents/Public/Scene/MeshSceneAdapter.h
--- a/Engine/Plugins/Experimental/MeshModelingToolsetExp/Source/ModelingComponents/Public/Scene/MeshSceneAdapter.h
+++ b/Engine/Plugins/Experimental/MeshModelingToolsetExp/Source/ModelingComponents/Public/Scene/MeshSceneAdapter.h
@@ -142,8 +142,8 @@ public:
 	/** Collect a set of seed points from this Mesh, mapped through LocalToWorldFunc to world space. Must be callable in parallel from any thread. */
 	virtual void CollectSeedPoints(TArray<FVector3d>& WorldPoints, TFunctionRef<FVector3d(const FVector3d&)> LocalToWorldFunc) = 0;
 
-	/** apply ProcessFunc to each vertex in world space */
-	virtual void ProcessVerticesInWorld(TFunctionRef<void(const FVector3d&)> ProcessFunc, const FTransformSequence3d& LocalToWorldTransform) = 0;
+	/** apply ProcessFunc to each vertex in world space. Return false from ProcessFunc to terminate iteration. Returns true if iteration completed. */
+	virtual bool ProcessVerticesInWorld(TFunctionRef<bool(const FVector3d&)> ProcessFunc, const FTransformSequence3d& LocalToWorldTransform) = 0;
 
 	/** Append the geometry represented by this wrapper to the accumulated AppendTo mesh, under the given world transform */
 	virtual void AppendMesh(FDynamicMesh3& AppendTo, const FTransformSequence3d& TransformSeq) = 0;
@@ -327,7 +327,25 @@ public:
 	/** @return world-space bounding box for the Scene Mesh of the specified Component and ComponentIndex. Requires that SpatialEvaluationCache is available. */
 	virtual FAxisAlignedBox3d GetMeshBoundingBox(UActorComponent* Component, int32 ComponentIndex = -1);
 
-
+	/**
+	 * Run a custom query across all world-transformed scene mesh vertices. Requires that SpatialEvaluationCache is available.
+	 * The internal loop is run as:
+	 *  ParallelFor(SceneMesh, {
+	 *    if ( MeshFilterFunc(SceneMesh) ) {
+	 *      for ( FVector3d Vtx in Scene Mesh ) { PerVertexFunc(Vtx) }
+	 *    }
+	 *  }
+	 * 
+	 * @param InitializeFunc called with the total number of scene meshes. MeshIndex in later queries will be < NumMeshes  Return false to terminate query.
+	 * @param MeshFilterFunc called once for each scene mesh, with unique MeshIndex per scene mesh. Return false to skip this scene mesh.
+	 * @param PerVertexFunc called for each worldspace-transformed scene mesh vertex. Return false to terminate enumeration over this scene mesh.
+	 */
+	virtual void ParallelMeshVertexEnumeration(
+		TFunctionRef<bool(int32 NumMeshes)> InitializeFunc,
+		TFunctionRef<bool(int32 MeshIndex, AActor* SourceActor, const FActorChildMesh* ChildMeshInfo, const FAxisAlignedBox3d& WorldBounds)> MeshFilterFunc,
+		TFunctionRef<bool(int32 MeshIndex, AActor* SourceActor, const FActorChildMesh* ChildMeshInfo, const FVector3d& WorldPos)> PerVertexFunc,
+		bool bForceSingleThreaded = false
+	);
 
 protected:
 	// top-level list of ActorAdapters, which represent each Actor and set of Components
