MeshModelingTools: Make PivotAct tool reposition the pivot actor if it is the only thing selected. Also fix undo/redo in PivotAct tool.

#rnx

diff --git a/Engine/Plugins/Experimental/MeshModelingToolsetExp/Source/MeshModelingToolsEditorOnly/Private/AddPivotActorTool.cpp b/Engine/Plugins/Experimental/MeshModelingToolsetExp/Source/MeshModelingToolsEditorOnly/Private/AddPivotActorTool.cpp
--- a/Engine/Plugins/Experimental/MeshModelingToolsetExp/Source/MeshModelingToolsEditorOnly/Private/AddPivotActorTool.cpp
+++ b/Engine/Plugins/Experimental/MeshModelingToolsetExp/Source/MeshModelingToolsEditorOnly/Private/AddPivotActorTool.cpp
@@ -6,6 +6,7 @@
 #include "BaseGizmos/TransformGizmo.h"
 #include "BaseGizmos/TransformGizmoUtil.h"
 #include "BaseGizmos/TransformProxy.h"
+#include "Components/BillboardComponent.h"
 #include "Mechanics/DragAlignmentMechanic.h"
 #include "ModelingToolTargetUtil.h"
 #include "Selection/ToolSelectionUtil.h"
@@ -16,6 +17,21 @@
 
 using namespace UE::Geometry;
 
+namespace AddPivotActorToolLocals
+{
+	UBillboardComponent* GetPivotBillboard(AActor* PivotActor)
+	{
+		for (const TObjectPtr<USceneComponent>& Child : PivotActor->GetRootComponent()->GetAttachChildren())
+		{
+			if (UBillboardComponent* Billboard = Cast<UBillboardComponent>(Child))
+			{
+				return Billboard;
+			}
+		}
+		return nullptr;
+	}
+}
+
 const FToolTargetTypeRequirements& UAddPivotActorToolBuilder::GetTargetRequirements() const
 {
 	static FToolTargetTypeRequirements TypeRequirements(
@@ -26,6 +42,8 @@ const FToolTargetTypeRequirements& UAddPivotActorToolBuilder::GetTargetRequireme
 
 bool UAddPivotActorToolBuilder::CanBuildTool(const FToolBuilderState& SceneState) const
 {
+	using namespace AddPivotActorToolLocals;
+
 	// There are some limitations for when we can use this tool. 
 	// 1. We operate on the actor, not component level.
 	//   TODO: Is there a good way to operate on a sub-actor level? Or should we be checking that
@@ -36,6 +54,12 @@ bool UAddPivotActorToolBuilder::CanBuildTool(const FToolBuilderState& SceneState
 	// 3. All of the actors need to be marked as movable because non-movable items can't be nested
 	//   under a movable one.
 
+	if (SceneState.SelectedActors.Num() == 1 && ExactCast<AActor>(SceneState.SelectedActors[0]) 
+		&& GetPivotBillboard(SceneState.SelectedActors[0]))
+	{
+		return true;
+	}
+
 	TSet<AActor*> ParentActors;
 	bool bAllActorsMovable = true;
 
@@ -55,35 +79,67 @@ bool UAddPivotActorToolBuilder::CanBuildTool(const FToolBuilderState& SceneState
 UInteractiveTool* UAddPivotActorToolBuilder::BuildTool(const FToolBuilderState& SceneState) const
 {
 	UAddPivotActorTool* NewTool = NewObject<UAddPivotActorTool>(SceneState.ToolManager);
-	NewTool->SetTargets(SceneState.TargetManager->BuildAllSelectedTargetable(
-		SceneState, GetTargetRequirements()));
-	NewTool->SetWorld(SceneState.World);
+
+	if (SceneState.SelectedActors.Num() == 1 && ExactCast<AActor>(SceneState.SelectedActors[0]))
+	{
+		NewTool->SetPivotRepositionMode(SceneState.SelectedActors[0]);
+	}
+	else
+	{
+		NewTool->SetTargets(SceneState.TargetManager->BuildAllSelectedTargetable(
+			SceneState, GetTargetRequirements()));
+		NewTool->SetWorld(SceneState.World);
+	}
 
 	return NewTool;
 }
 
 void UAddPivotActorTool::Setup()
 {
-	GetToolManager()->DisplayMessage(LOCTEXT("OnStartTool", 
-		"Adds an empty actor as the parent of the selected actors. Use gizmo to choose where/how "
-		"the empty actor is placed. Hold Ctrl to snap to items in scene."), 
-		EToolMessageLevel::UserNotification);
-
-	// Figure out where to start the gizmo. The location will be the average,
-	// and the rotation will either be identity, or the target's if there is
-	// only one.
-	FVector3d StartTranslation = FVector3d::Zero();
-	for (TObjectPtr<UToolTarget> Target : Targets)
+	using namespace AddPivotActorToolLocals;
+
+	FTransform StartTransform = FTransform::Identity;
+	if (ExistingPivotActor.IsValid())
 	{
-		StartTranslation += UE::ToolTarget::GetLocalToWorldTransform(Target).GetTranslation();
-	}
-	StartTranslation /= Targets.Num();
+		GetToolManager()->DisplayMessage(LOCTEXT("OnStartToolEdit",
+			"Modifies the position of the pivot actor while keeping children in place. "
+			"Hold Ctrl to snap to items in scene."),
+			EToolMessageLevel::UserNotification);
+
+		StartTransform = ExistingPivotActor->GetTransform();
+		ExistingPivotOriginalTransform = StartTransform;
 
-	FTransform StartTransform(StartTranslation);
-	if (Targets.Num() == 1)
+		// Hide the pivot sprite.
+		UBillboardComponent* Billboard = GetPivotBillboard(ExistingPivotActor.Get());
+		if (ensure(Billboard))
+		{
+			Billboard->SetVisibility(false);
+		}
+	}
+	else
 	{
-		StartTransform.SetRotation(FQuat(
-			UE::ToolTarget::GetLocalToWorldTransform(Targets[0]).GetRotation()));
+		GetToolManager()->DisplayMessage(LOCTEXT("OnStartToolAdd",
+			"Adds an empty actor as the parent of the selected actors. Use gizmo to choose where/how "
+			"the empty actor is placed. Hold Ctrl to snap to items in scene."),
+			EToolMessageLevel::UserNotification);
+
+		// Figure out where to start the gizmo. The location will be the average,
+		// and the rotation will either be identity, or the target's if there is
+		// only one.
+		FVector3d StartTranslation = FVector3d::Zero();
+		for (TObjectPtr<UToolTarget> Target : Targets)
+		{
+			StartTranslation += UE::ToolTarget::GetLocalToWorldTransform(Target).GetTranslation();
+		}
+		StartTranslation /= Targets.Num();
+
+		StartTransform.SetTranslation(StartTranslation);
+
+		if (Targets.Num() == 1)
+		{
+			StartTransform.SetRotation(FQuat(
+				UE::ToolTarget::GetLocalToWorldTransform(Targets[0]).GetRotation()));
+		}
 	}
 
 	// Set up the gizmo.
@@ -101,49 +157,99 @@ void UAddPivotActorTool::Setup()
 
 void UAddPivotActorTool::Shutdown(EToolShutdownType ShutdownType)
 {
-	if (ShutdownType == EToolShutdownType::Accept)
+	using namespace AddPivotActorToolLocals;
+
+	if (ExistingPivotActor.IsValid())
 	{
-		GetToolManager()->BeginUndoTransaction(LOCTEXT("AddPivotActorToolTransactionName", "Add Empty Actor"));
-
-		// Create an empty actor at the location of the gizmo. The way we do it here, using this factory, is
-		// editor-only.
-		UActorFactoryEmptyActor* EmptyActorFactory = NewObject<UActorFactoryEmptyActor>();
-		FAssetData AssetData(EmptyActorFactory->GetDefaultActorClass(FAssetData()));
-		FActorSpawnParameters SpawnParams;
-		SpawnParams.Name = TEXT("PivotActor");
-		SpawnParams.NameMode = FActorSpawnParameters::ESpawnActorNameMode::Requested;
-		AActor* NewActor = EmptyActorFactory->CreateActor(AssetData.GetAsset(), 
-			TargetWorld->GetCurrentLevel(), 
-			TransformProxy->GetTransform(), 
-			SpawnParams);
-
-		// Grab the first selected target. It will have the same parent as the other ones, so
-		// we'll use it to figure out the empty actor's parent.
-		AActor* TargetActor = UE::ToolTarget::GetTargetActor(Targets[0]);
-
-		// This is also editor-only: it's the label that shows up in the hierarchy
-		NewActor->SetActorLabel(Targets.Num() == 1 ? TargetActor->GetActorLabel() + TEXT("_Pivot")
-			: TEXT("Pivot"));
-
-		// Attach the empty actor in the correct place in the hierarchy
-		if (AActor* ParentActor = TargetActor->GetAttachParentActor())
+		// Make the sprite visible again.
+		UBillboardComponent* Billboard = GetPivotBillboard(ExistingPivotActor.Get());
+		if (Billboard)
 		{
-			NewActor->AttachToActor(ParentActor, FAttachmentTransformRules(EAttachmentRule::KeepWorld, false));
+			Billboard->SetVisibility(true);
 		}
-		for (TObjectPtr<UToolTarget> Target : Targets)
+	}
+
+	if (ShutdownType == EToolShutdownType::Accept)
+	{
+		if (ExistingPivotActor.IsValid())
 		{
-			TargetActor = UE::ToolTarget::GetTargetActor(Target); 
-			TargetActor->AttachToActor(NewActor, FAttachmentTransformRules(EAttachmentRule::KeepWorld, false));
+			USceneComponent* PivotComponent = ExistingPivotActor->GetRootComponent();
+			FTransform GizmoTransform = TransformProxy->GetTransform();
+
+			GetToolManager()->BeginUndoTransaction(LOCTEXT("MovePivotTransactionName", "Move Pivot Actor"));
+
+			// Move the billboard (the rendered pivot) to the new location (i.e. don't adjust it's relative transform),
+			// but keep all the other children in their current location.
+			TArray<USceneComponent*> Children;
+			PivotComponent->GetChildrenComponents(false, Children);
+			UBillboardComponent* Billboard = GetPivotBillboard(ExistingPivotActor.Get());
+			for (USceneComponent* Child : Children)
+			{
+				if (Child == Billboard)
+				{
+					continue;
+				}
+				Child->Modify();
+
+				FTransform NewTransform = Child->GetComponentToWorld();
+				NewTransform.SetToRelativeTransform(GizmoTransform);
+				Child->SetRelativeTransform(NewTransform);
+			}
+
+			// Move the pivot itself. Note that this needs to happen after moving the children above, because
+			// otherwise the above GetComponentToWorld call won't get us the original location.
+			PivotComponent->Modify();
+			PivotComponent->SetWorldTransform(GizmoTransform);
+
+			GetToolManager()->EndUndoTransaction();
 		}
-		ToolSelectionUtil::SetNewActorSelection(GetToolManager(), NewActor);
+		else // If in "add" mode
+		{
+			GetToolManager()->BeginUndoTransaction(LOCTEXT("AddPivotActorTransactionName", "Add Pivot Actor"));
+
+			// Create an empty actor at the location of the gizmo. The way we do it here, using this factory, is
+			// editor-only.
+			UActorFactoryEmptyActor* EmptyActorFactory = NewObject<UActorFactoryEmptyActor>();
+			FAssetData AssetData(EmptyActorFactory->GetDefaultActorClass(FAssetData()));
+			FActorSpawnParameters SpawnParams;
+			SpawnParams.Name = TEXT("PivotActor");
+			SpawnParams.NameMode = FActorSpawnParameters::ESpawnActorNameMode::Requested;
+			AActor* NewActor = EmptyActorFactory->CreateActor(AssetData.GetAsset(),
+				TargetWorld->GetCurrentLevel(),
+				TransformProxy->GetTransform(),
+				SpawnParams);
+
+			// Grab the first selected target. It will have the same parent as the other ones, so
+			// we'll use it to figure out the empty actor's parent.
+			AActor* TargetActor = UE::ToolTarget::GetTargetActor(Targets[0]);
 
-		GetToolManager()->EndUndoTransaction();
+			// This is also editor-only: it's the label that shows up in the hierarchy
+			NewActor->SetActorLabel(Targets.Num() == 1 ? TargetActor->GetActorLabel() + TEXT("_Pivot")
+				: TEXT("Pivot"));
+
+			// Attach the empty actor in the correct place in the hierarchy. This can be done in a non-editor-only
+			// way, but it's easier to use the editor's function to do it so that undo/redo and level saving are 
+			// properly done.
+			if (AActor* ParentActor = TargetActor->GetAttachParentActor())
+			{
+				GEditor->ParentActors(ParentActor, NewActor, NAME_None);
+			}
+			for (TObjectPtr<UToolTarget> Target : Targets)
+			{
+				TargetActor = UE::ToolTarget::GetTargetActor(Target);
+				GEditor->ParentActors(NewActor, TargetActor, NAME_None);
+			}
+			ToolSelectionUtil::SetNewActorSelection(GetToolManager(), NewActor);
+
+			GetToolManager()->EndUndoTransaction();
+		}
 	}
 
 	GetToolManager()->GetPairedGizmoManager()->DestroyAllGizmosByOwner(this);
 	TransformProxy = nullptr;
 	TransformGizmo = nullptr;
 	DragAlignmentMechanic->Shutdown();
+	ExistingPivotActor.Reset();
 }
 
 void UAddPivotActorTool::Render(IToolsContextRenderAPI* RenderAPI)
diff --git a/Engine/Plugins/Experimental/MeshModelingToolsetExp/Source/MeshModelingToolsEditorOnly/Public/AddPivotActorTool.h b/Engine/Plugins/Experimental/MeshModelingToolsetExp/Source/MeshModelingToolsEditorOnly/Public/AddPivotActorTool.h
--- a/Engine/Plugins/Experimental/MeshModelingToolsetExp/Source/MeshModelingToolsEditorOnly/Public/AddPivotActorTool.h
+++ b/Engine/Plugins/Experimental/MeshModelingToolsetExp/Source/MeshModelingToolsEditorOnly/Public/AddPivotActorTool.h
@@ -37,6 +37,7 @@ class MESHMODELINGTOOLSEDITORONLY_API UAddPivotActorTool : public UMultiSelectio
 public:
 
 	virtual void SetWorld(UWorld* World) { TargetWorld = World; }
+	virtual void SetPivotRepositionMode(AActor* PivotActor) { ExistingPivotActor = PivotActor; }
 
 	virtual void Setup() override;
 	virtual void Shutdown(EToolShutdownType ShutdownType) override;
@@ -59,4 +60,7 @@ protected:
 	TObjectPtr<UTransformProxy> TransformProxy = nullptr;
 
 	TObjectPtr<UWorld> TargetWorld;
+	TWeakObjectPtr<AActor> ExistingPivotActor = nullptr;
+
+	FTransform ExistingPivotOriginalTransform;
 };
