MeshModelingTools: Make the background compute object status more specific to handle the case when a computation is completed but known to be dirty. Make HasValidResult specifically mean a valid final result.

#rnx

diff --git a/Engine/Plugins/Experimental/MeshModelingToolsetExp/Source/ModelingComponents/Private/MeshOpPreviewHelpers.cpp b/Engine/Plugins/Experimental/MeshModelingToolsetExp/Source/ModelingComponents/Private/MeshOpPreviewHelpers.cpp
--- a/Engine/Plugins/Experimental/MeshModelingToolsetExp/Source/ModelingComponents/Private/MeshOpPreviewHelpers.cpp
+++ b/Engine/Plugins/Experimental/MeshModelingToolsetExp/Source/ModelingComponents/Private/MeshOpPreviewHelpers.cpp
@@ -111,7 +111,8 @@ void UMeshOpPreviewWithBackgroundCompute::UpdateResults()
 	LastComputeStatus = BackgroundCompute ? BackgroundCompute->CheckStatus()
 		: EBackgroundComputeTaskStatus::NotComputing;
 
-	if (LastComputeStatus == EBackgroundComputeTaskStatus::NewResultAvailable)
+	if (LastComputeStatus == EBackgroundComputeTaskStatus::ValidResultAvailable
+		|| (bAllowDirtyResultUpdates && LastComputeStatus == EBackgroundComputeTaskStatus::DirtyResultAvailable))
 	{
 		TUniquePtr<FDynamicMeshOperator> MeshOp = BackgroundCompute->ExtractResult();
 		OnOpCompleted.Broadcast(MeshOp.Get());
@@ -127,7 +128,7 @@ void UMeshOpPreviewWithBackgroundCompute::UpdateResults()
 		bMeshInitialized = true;
 
 		PreviewMesh->SetVisible(bVisible);
-		bResultValid = true;
+		bResultValid = (LastComputeStatus == EBackgroundComputeTaskStatus::ValidResultAvailable);
 
 		OnMeshUpdated.Broadcast(this);
 	}
diff --git a/Engine/Plugins/Experimental/MeshModelingToolsetExp/Source/ModelingComponents/Public/MeshOpPreviewHelpers.h b/Engine/Plugins/Experimental/MeshModelingToolsetExp/Source/ModelingComponents/Public/MeshOpPreviewHelpers.h
--- a/Engine/Plugins/Experimental/MeshModelingToolsetExp/Source/ModelingComponents/Public/MeshOpPreviewHelpers.h
+++ b/Engine/Plugins/Experimental/MeshModelingToolsetExp/Source/ModelingComponents/Public/MeshOpPreviewHelpers.h
@@ -216,6 +216,14 @@ public:
 	TObjectPtr<UMaterialInterface> SecondaryMaterial = nullptr;
 
 	UWorld* PreviewWorld;
+
+	/**
+	 * When true, the preview mesh is allowed to be temporarily updated using results that we know
+	 * are dirty (i.e., the preview was invalidated, but a result became available before the operation
+	 * was restarted, so we can at least show that while we wait for the new result). The change 
+	 * notifications will be fired as normal for these results, but HasValidResult will return false.
+	 */
+	bool bAllowDirtyResultUpdates = true;
 protected:
 	// state flag, if true then we have valid result
 	bool bResultValid = false;
@@ -374,7 +382,7 @@ protected:
 		if (BackgroundCompute)
 		{
 			UE::Geometry::EBackgroundComputeTaskStatus Status = BackgroundCompute->CheckStatus();
-			if (Status == UE::Geometry::EBackgroundComputeTaskStatus::NewResultAvailable)
+			if (Status == UE::Geometry::EBackgroundComputeTaskStatus::ValidResultAvailable)
 			{
 				TUniquePtr<OperatorType> ResultOp = BackgroundCompute->ExtractResult();
 				OnOpCompleted.Broadcast(ResultOp.Get());
diff --git a/Engine/Plugins/Experimental/MeshModelingToolsetExp/Source/ModelingOperators/Public/BackgroundModelingComputeSource.h b/Engine/Plugins/Experimental/MeshModelingToolsetExp/Source/ModelingOperators/Public/BackgroundModelingComputeSource.h
--- a/Engine/Plugins/Experimental/MeshModelingToolsetExp/Source/ModelingOperators/Public/BackgroundModelingComputeSource.h
+++ b/Engine/Plugins/Experimental/MeshModelingToolsetExp/Source/ModelingOperators/Public/BackgroundModelingComputeSource.h
@@ -73,13 +73,21 @@ protected:
 enum class EBackgroundComputeTaskStatus
 {
 	/** Computation of a result has finished and is waiting to be returned */
-	NewResultAvailable = 0,
+	ValidResultAvailable,
+
+	/**
+	 * We have a result available, but a recompute has been requested. The status is not yet 
+	 * InProgress only because there is a delay of CancelActiveOpDelaySeconds between the
+	 * request and operation restart, and we have not yet restarted.
+	 */
+	 DirtyResultAvailable,
+
 	/** Last active computation was canceled and nothing new has happend yet*/
-	Aborted = 1,
+	Aborted,
 	/** Computation is currently running */
-	InProgress = 2,
+	InProgress,
 	/** Not running active computation, and last result has already been returned, so no new results to report */
-	NotComputing = 3
+	NotComputing
 };
 
 
@@ -259,21 +267,23 @@ void TBackgroundModelingComputeSource<OpType, OpTypeFactory>::NotifyActiveComput
 template<typename OpType, typename OpTypeFactory>
 EBackgroundComputeTaskStatus TBackgroundModelingComputeSource<OpType, OpTypeFactory>::CheckStatus() const
 {
-	if (ActiveBackgroundTask != nullptr)
+	if (ActiveBackgroundTask == nullptr)
 	{
-		if (ActiveBackgroundTask->IsDone())
-		{
-			return ActiveBackgroundTask->GetTask().IsAborted() ?
-				EBackgroundComputeTaskStatus::Aborted : EBackgroundComputeTaskStatus::NewResultAvailable;
-		}
-		else
-		{
-			return EBackgroundComputeTaskStatus::InProgress;
-		}
+		return EBackgroundComputeTaskStatus::NotComputing;
+	}
+	else if (!ActiveBackgroundTask->IsDone())
+	{
+		return EBackgroundComputeTaskStatus::InProgress;
+	}
+	else if (ActiveBackgroundTask->GetTask().IsAborted())
+	{
+		return EBackgroundComputeTaskStatus::Aborted;
 	}
 	else
 	{
-		return EBackgroundComputeTaskStatus::NotComputing;
+		// Task is done and not aborted, but we may be waiting to start a new one
+		return TaskState == EBackgroundComputeTaskState::WaitingToCancel ? EBackgroundComputeTaskStatus::DirtyResultAvailable
+			: EBackgroundComputeTaskStatus::ValidResultAvailable;
 	}
 }
 
