ModelingTools: replace StorableSelection mechanism with new UPersistentMeshSelection and UPersistentMeshSelectionManager context object.
- UPersistentMeshSelection is largely a port of UGroupTopologyStorableSelection, with the actual selection data moved to FGenericMeshSelection
- UPersistentMeshSelectionManager is meant to be used as a ContextObject in a ToolsContext, also contains utility functions to register/unregister/find context object
- Selection is no longer passed as part of FToolBuilderState. Instead Tools access the selection via ContextObject
- UPersistentMeshSelectionManager currently supports a single active selection. Selection changes are transacted via an FChange.
- When not inside a Tool, Selection is visualized with a pink border outline. Currently vertex selection is not visualized.
- Usage model is that on Tool Accept/Complete, prior to Tool Shutdown, any existing selection is cleared, and Tool may set a new Output selection which becomes the active selection
- StoredMeshSelectionUtil.h has utility functions to get current selection, set output selection, and clear it
- SingleSelectionMeshEditingTool already made input selection available to subclasses, port that capability to new architecture
- convert EditMeshPolygonsTool to be a SingleSelectionMeshEditingTool, use provided Input selection and set Output selection as appropriate

- ModelingToolsEditorMode clears active selection where appropriate, eg if selected object changes. However the behavior here will need further improvement, currently relies on questionable event handling from the TypedElement system

#rnx

diff --git a/Engine/Plugins/Editor/ModelingToolsEditorMode/Source/ModelingToolsEditorMode/ModelingToolsEditorMode.Build.cs b/Engine/Plugins/Editor/ModelingToolsEditorMode/Source/ModelingToolsEditorMode/ModelingToolsEditorMode.Build.cs
--- a/Engine/Plugins/Editor/ModelingToolsEditorMode/Source/ModelingToolsEditorMode/ModelingToolsEditorMode.Build.cs
+++ b/Engine/Plugins/Editor/ModelingToolsEditorMode/Source/ModelingToolsEditorMode/ModelingToolsEditorMode.Build.cs
@@ -49,6 +49,7 @@ public class ModelingToolsEditorMode : ModuleRules
 				"StatusBar",
 				"EditorStyle",
                 "Projects",
+				"TypedElementRuntime",
 				"MeshDescription",
 				"StaticMeshDescription",
                 "InteractiveToolsFramework",
diff --git a/Engine/Plugins/Editor/ModelingToolsEditorMode/Source/ModelingToolsEditorMode/Private/ModelingToolsEditorMode.cpp b/Engine/Plugins/Editor/ModelingToolsEditorMode/Source/ModelingToolsEditorMode/Private/ModelingToolsEditorMode.cpp
--- a/Engine/Plugins/Editor/ModelingToolsEditorMode/Source/ModelingToolsEditorMode/Private/ModelingToolsEditorMode.cpp
+++ b/Engine/Plugins/Editor/ModelingToolsEditorMode/Source/ModelingToolsEditorMode/Private/ModelingToolsEditorMode.cpp
@@ -2,7 +2,6 @@
 
 #include "ModelingToolsEditorMode.h"
 #include "InteractiveTool.h"
-#include "InteractiveToolsSelectionStoreSubsystem.h"
 #include "ModelingToolsEditorModeToolkit.h"
 #include "Toolkits/ToolkitManager.h"
 #include "ToolTargetManager.h"
@@ -17,6 +16,8 @@
 #include "EditorViewportClient.h"
 #include "EngineAnalytics.h"
 #include "BaseGizmos/TransformGizmoUtil.h"
+#include "Selection/PersistentMeshSelectionManager.h"
+#include "Selection/StoredMeshSelectionUtil.h"
 
 #include "Features/IModularFeatures.h"
 #include "ModelingModeToolExtensions.h"
@@ -158,12 +159,8 @@ bool UModelingToolsEditorMode::ProcessEditDelete()
 		return true;
 	}
 
-	// If we didn't skip deletion, then we're probably deleting something, so it seems fair to clear
-	// the selection.
-	if (UInteractiveToolsSelectionStoreSubsystem* ToolSelectionStore = GEngine->GetEngineSubsystem<UInteractiveToolsSelectionStoreSubsystem>())
-	{
-		ToolSelectionStore->ClearStoredSelection();
-	}
+	// clear any active selection
+	UE::Geometry::ClearActiveToolSelection(GetToolManager());
 
 	return false;
 }
@@ -179,16 +176,19 @@ bool UModelingToolsEditorMode::ProcessEditCut()
 		return true;
 	}
 
-	// If we're doing a cut, we should clear the tool selection.
-	if (UInteractiveToolsSelectionStoreSubsystem* ToolSelectionStore = GEngine->GetEngineSubsystem<UInteractiveToolsSelectionStoreSubsystem>())
-	{
-		ToolSelectionStore->ClearStoredSelection();
-	}
+	// clear any active selection
+	UE::Geometry::ClearActiveToolSelection(GetToolManager());
 
 	return false;
 }
 
 
+void UModelingToolsEditorMode::ActorSelectionChangeNotify()
+{
+	// would like to clear selection here, but this is called multiple times, including after a transaction when
+	// we cannot identify that the selection should not be cleared
+}
+
 
 bool UModelingToolsEditorMode::CanAutoSave() const
 {
@@ -315,6 +315,8 @@ void UModelingToolsEditorMode::Enter()
 
 	// register gizmo helper
 	UE::TransformGizmoUtil::RegisterTransformGizmoContextObject(ToolsContext.Get());
+	// register selection manager
+	UE::Geometry::RegisterPersistentMeshSelectionManager(ToolsContext.Get());
 
 	// disable HitProxy rendering, it is not used in Modeling Mode and adds overhead to Render() calls
 	ToolsContext.Get()->SetEnableRenderingDuringHitProxyPass(false);
@@ -693,10 +695,33 @@ void UModelingToolsEditorMode::Enter()
 		ModelingToolkit->InitializeAfterModeSetup();
 	}
 
+	// Need to know about selection changes to clear mesh selections, however we do not want to clear the
+	// mesh selection after a selection change due to transactions, as this may clear a selection we just restored.
+	// Unfortunately most routes to find out about selection changes don't allow for this. Currently this
+	// OnPreChange event is the only one that appears to provide the desired behavior, however it is likely
+	// that this is not going to be reliable...
+	SelectionModifiedEventHandle = GetModeManager()->GetSelectedActors()->GetElementSelectionSet()->OnPreChange().AddLambda( [&](const UTypedElementSelectionSet*) 
+	{ 
+		if (GIsTransacting == 0)
+		{
+			UE::Geometry::ClearActiveToolSelection(GetToolManager());
+		}
+	});
 }
 
 void UModelingToolsEditorMode::Exit()
 {
+	// clear any active selection
+	UE::Geometry::ClearActiveToolSelection(GetToolManager());
+	// deregister selection manager
+	UE::Geometry::DeregisterPersistentMeshSelectionManager(ToolsContext.Get());
+
+	// stop listening to selection changes. On Editor Shutdown, some of these values become null, which will result in an ensure/crash
+	if (UObjectInitialized() && GetModeManager() && GetModeManager()->GetSelectedActors() && GetModeManager()->GetSelectedActors()->GetElementSelectionSet() )
+	{
+		GetModeManager()->GetSelectedActors()->GetElementSelectionSet()->OnPreChange().Remove(SelectionModifiedEventHandle);
+	}
+
 	//
 	// Engine Analytics
 	//
@@ -775,7 +800,10 @@ void UModelingToolsEditorMode::BindCommands()
 
 	CommandList->MapAction(
 		ToolManagerCommands.AcceptActiveTool,
-		FExecuteAction::CreateLambda([this]() { ToolsContext->EndTool(EToolShutdownType::Accept); }),
+		FExecuteAction::CreateLambda([this]() { 
+			UE::Geometry::ClearActiveToolSelection(GetToolManager());
+			ToolsContext->EndTool(EToolShutdownType::Accept); 
+		}),
 		FCanExecuteAction::CreateLambda([this]() { return ToolsContext->CanAcceptActiveTool(); }),
 		FGetActionCheckState(),
 		FIsActionButtonVisible::CreateLambda([this]() {return ToolsContext->ActiveToolHasAccept(); }),
@@ -840,6 +868,9 @@ void UModelingToolsEditorMode::AcceptActiveToolActionOrTool()
 		}
 	}
 
+	// clear existing selection
+	UE::Geometry::ClearActiveToolSelection(GetToolManager());
+
 	const EToolShutdownType ShutdownType = ToolsContext->CanAcceptActiveTool() ? EToolShutdownType::Accept : EToolShutdownType::Completed;
 	ToolsContext->EndTool(ShutdownType);
 }
diff --git a/Engine/Plugins/Editor/ModelingToolsEditorMode/Source/ModelingToolsEditorMode/Public/ModelingToolsEditorMode.h b/Engine/Plugins/Editor/ModelingToolsEditorMode/Source/ModelingToolsEditorMode/Public/ModelingToolsEditorMode.h
--- a/Engine/Plugins/Editor/ModelingToolsEditorMode/Source/ModelingToolsEditorMode/Public/ModelingToolsEditorMode.h
+++ b/Engine/Plugins/Editor/ModelingToolsEditorMode/Source/ModelingToolsEditorMode/Public/ModelingToolsEditorMode.h
@@ -32,6 +32,8 @@ public:
 
 	virtual void Tick(FEditorViewportClient* ViewportClient, float DeltaTime) override;
 
+	virtual void ActorSelectionChangeNotify() override;
+
 	virtual bool ShouldDrawWidget() const override;
 	virtual bool ProcessEditDelete() override;
 	virtual bool ProcessEditCut() override;
@@ -64,6 +66,7 @@ protected:
 
 	FDelegateHandle MeshCreatedEventHandle;
 	FDelegateHandle TextureCreatedEventHandle;
+	FDelegateHandle SelectionModifiedEventHandle;
 
 	TUniquePtr<FStylusStateTracker> StylusStateTracker;
 
diff --git a/Engine/Plugins/Experimental/MeshModelingToolsetExp/Source/MeshModelingToolsExp/Private/EditMeshPolygonsTool.cpp b/Engine/Plugins/Experimental/MeshModelingToolsetExp/Source/MeshModelingToolsExp/Private/EditMeshPolygonsTool.cpp
--- a/Engine/Plugins/Experimental/MeshModelingToolsetExp/Source/MeshModelingToolsExp/Private/EditMeshPolygonsTool.cpp
+++ b/Engine/Plugins/Experimental/MeshModelingToolsetExp/Source/MeshModelingToolsExp/Private/EditMeshPolygonsTool.cpp
@@ -25,7 +25,8 @@
 #include "MeshRegionBoundaryLoops.h"
 #include "ModelingToolTargetUtil.h" // UE::ToolTarget:: functions
 #include "Operations/SimpleHoleFiller.h"
-#include "Selection/GroupTopologyStorableSelection.h"
+#include "Selection/PersistentMeshSelection.h"
+#include "Selection/StoredMeshSelectionUtil.h"
 #include "Selection/PolygonSelectionMechanic.h"
 #include "Selections/MeshConnectedComponents.h"
 #include "TargetInterfaces/MaterialProvider.h"
@@ -65,67 +66,46 @@ namespace EditMeshPolygonsToolLocals
  * ToolBuilder
  */
 
-const FToolTargetTypeRequirements& UEditMeshPolygonsToolBuilder::GetTargetRequirements() const
-{
-	static FToolTargetTypeRequirements TypeRequirements({
-		UMaterialProvider::StaticClass(),
-
-		// TODO: In reality we don't actually need a mesh description- we need anything that we
-		// can get a dynamic mesh out of. Unfortunately we have not yet settled on an interface that
-		// provides this- DynamicMeshProvider (which currently gives a dynamic mesh shared ptr) is 
-		// likely to be changed or replaced with something that gives a UDynamicMesh instead. For now
-		// we ask for a mesh description as a hack, since all the targets we currently use are able to
-		// provide both it and a dynamic mesh.
-		UMeshDescriptionCommitter::StaticClass(),
-		UMeshDescriptionProvider::StaticClass(),
-		
-		UPrimitiveComponentBackedTarget::StaticClass()
-		});
-	return TypeRequirements;
-}
 
-bool UEditMeshPolygonsToolBuilder::CanBuildTool(const FToolBuilderState& SceneState) const
+USingleSelectionMeshEditingTool* UEditMeshPolygonsToolBuilder::CreateNewTool(const FToolBuilderState& SceneState) const
 {
-	return SceneState.TargetManager->CountSelectedAndTargetable(SceneState, GetTargetRequirements()) == 1;
+	return NewObject<UEditMeshPolygonsTool>(SceneState.ToolManager);
 }
 
-UInteractiveTool* UEditMeshPolygonsToolBuilder::BuildTool(const FToolBuilderState& SceneState) const
+void UEditMeshPolygonsToolBuilder::InitializeNewTool(USingleSelectionMeshEditingTool* Tool, const FToolBuilderState& SceneState) const
 {
-	UEditMeshPolygonsTool* EditPolygonsTool = NewObject<UEditMeshPolygonsTool>(SceneState.ToolManager);
+	USingleSelectionMeshEditingToolBuilder::InitializeNewTool(Tool, SceneState);
+	UEditMeshPolygonsTool* EditPolygonsTool = CastChecked<UEditMeshPolygonsTool>(Tool);
 	if (bTriangleMode)
 	{
 		EditPolygonsTool->EnableTriangleMode();
 	}
-	EditPolygonsTool->SetWorld(SceneState.World);
-	EditPolygonsTool->SetTarget(SceneState.TargetManager->BuildFirstSelectedTargetable(SceneState, GetTargetRequirements()));
-
-	// This passes in nullptr if the stored selection is not the right type. The tool
-	// will figure out whether the selection is still relevant.
-	EditPolygonsTool->SetStoredToolSelection(Cast<UGroupTopologyStorableSelection>(SceneState.StoredToolSelection));
-
-	return EditPolygonsTool;
 }
 
 
-UInteractiveTool* UEditMeshPolygonsActionModeToolBuilder::BuildTool(const FToolBuilderState& SceneState) const
+
+void UEditMeshPolygonsActionModeToolBuilder::InitializeNewTool(USingleSelectionMeshEditingTool* Tool, const FToolBuilderState& SceneState) const
 {
-	UEditMeshPolygonsTool* Tool = Cast<UEditMeshPolygonsTool>(UEditMeshPolygonsToolBuilder::BuildTool(SceneState));
+	UEditMeshPolygonsToolBuilder::InitializeNewTool(Tool, SceneState);
+	UEditMeshPolygonsTool* EditPolygonsTool = CastChecked<UEditMeshPolygonsTool>(Tool);
+
 	EEditMeshPolygonsToolActions UseAction = StartupAction;
-	Tool->PostSetupFunction = [UseAction](UEditMeshPolygonsTool* PolyTool)
+	EditPolygonsTool->PostSetupFunction = [UseAction](UEditMeshPolygonsTool* PolyTool)
 	{
 		PolyTool->SetToSelectionModeInterface();
 		PolyTool->RequestAction(UseAction);
 	};
-
-	return Tool;
 }
 
 
-UInteractiveTool* UEditMeshPolygonsSelectionModeToolBuilder::BuildTool(const FToolBuilderState& SceneState) const
+
+void UEditMeshPolygonsSelectionModeToolBuilder::InitializeNewTool(USingleSelectionMeshEditingTool* Tool, const FToolBuilderState& SceneState) const
 {
-	UEditMeshPolygonsTool* Tool = Cast<UEditMeshPolygonsTool>(UEditMeshPolygonsToolBuilder::BuildTool(SceneState));
+	UEditMeshPolygonsToolBuilder::InitializeNewTool(Tool, SceneState);
+	UEditMeshPolygonsTool* EditPolygonsTool = CastChecked<UEditMeshPolygonsTool>(Tool);
+
 	EEditMeshPolygonsToolSelectionMode UseMode = SelectionMode;
-	Tool->PostSetupFunction = [UseMode](UEditMeshPolygonsTool* PolyTool)
+	EditPolygonsTool->PostSetupFunction = [UseMode](UEditMeshPolygonsTool* PolyTool)
 	{
 		PolyTool->SetToSelectionModeInterface();
 
@@ -159,10 +139,9 @@ UInteractiveTool* UEditMeshPolygonsSelectionModeToolBuilder::BuildTool(const FTo
 			break;
 		}
 	};
-
-	return Tool;
 }
 
+
 void UEditMeshPolygonsTool::SetToSelectionModeInterface()
 {
 	if (EditActions) SetToolPropertySourceEnabled(EditActions, false);
@@ -322,9 +301,9 @@ void UEditMeshPolygonsTool::Setup()
 
 	// Have to load selection after initializing the selection mechanic since we need to have
 	// the topology built.
-	if (IsStoredToolSelectionUsable(StoredToolSelection))
+	if (HasInputSelection() && IsToolInputSelectionUsable(GetInputSelection()))
 	{
-		SelectionMechanic->LoadStorableSelection(*StoredToolSelection);
+		SelectionMechanic->LoadSelection(*GetInputSelection());
 	}
 
 	bSelectionStateDirty = SelectionMechanic->HasSelection();
@@ -429,7 +408,7 @@ void UEditMeshPolygonsTool::Setup()
 	}
 }
 
-bool UEditMeshPolygonsTool::IsStoredToolSelectionUsable(const UGroupTopologyStorableSelection* StoredSelection)
+bool UEditMeshPolygonsTool::IsToolInputSelectionUsable(const UPersistentMeshSelection* InputSelectionIn)
 {
 	// TODO: We currently don't support persistent selection on volume brushes because
 	// a conversion back to a brush involves a simplification step that may make the 
@@ -438,12 +417,12 @@ bool UEditMeshPolygonsTool::IsStoredToolSelectionUsable(const UGroupTopologyStor
 	// identifying info) will change.
 	return !Cast<UBrushComponent>(UE::ToolTarget::GetTargetComponent(Target))
 
-		&& StoredSelection
-		&& StoredSelection->IdentifyingInfo.TopologyType == (bTriangleMode ? 
-			UGroupTopologyStorableSelection::ETopologyType::FTriangleGroupTopology
-			: UGroupTopologyStorableSelection::ETopologyType::FGroupTopology)
-		&& StoredSelection->IdentifyingInfo.ComponentTarget == UE::ToolTarget::GetTargetComponent(Target)
-		&& !StoredSelection->IsEmpty();
+		&& InputSelectionIn
+		&& InputSelectionIn->GetSelectionType() == (bTriangleMode ?
+			FGenericMeshSelection::ETopologyType::FTriangleGroupTopology
+			: FGenericMeshSelection::ETopologyType::FGroupTopology)
+		&& InputSelectionIn->GetTargetComponent() == UE::ToolTarget::GetTargetComponent(Target)
+		&& !InputSelectionIn->IsEmpty();
 }
 
 void UEditMeshPolygonsTool::Shutdown(EToolShutdownType ShutdownType)
@@ -475,7 +454,7 @@ void UEditMeshPolygonsTool::Shutdown(EToolShutdownType ShutdownType)
 
 		if (ShutdownType == EToolShutdownType::Accept)
 		{
-			UGroupTopologyStorableSelection* NewStoredToolSelection = nullptr;
+			UPersistentMeshSelection* OutputSelection = nullptr;
 			FCompactMaps CompactMaps;
 
 			// Prep if we have a selection to store. We don't support storing selections for volumes
@@ -483,11 +462,13 @@ void UEditMeshPolygonsTool::Shutdown(EToolShutdownType ShutdownType)
 			if (!SelectionMechanic->GetActiveSelection().IsEmpty()
 				&& !Cast<UBrushComponent>(UE::ToolTarget::GetTargetComponent(Target)))
 			{
-				NewStoredToolSelection = NewObject<UGroupTopologyStorableSelection>();
-				NewStoredToolSelection->IdentifyingInfo.ComponentTarget = UE::ToolTarget::GetTargetComponent(Target);
-				NewStoredToolSelection->IdentifyingInfo.TopologyType = (bTriangleMode ?
-					UGroupTopologyStorableSelection::ETopologyType::FTriangleGroupTopology
-					: UGroupTopologyStorableSelection::ETopologyType::FGroupTopology);
+				OutputSelection = NewObject<UPersistentMeshSelection>();
+				FGenericMeshSelection NewSelection;
+				NewSelection.SourceComponent = UE::ToolTarget::GetTargetComponent(Target);
+				NewSelection.TopologyType = (bTriangleMode ?
+					FGenericMeshSelection::ETopologyType::FTriangleGroupTopology
+					: FGenericMeshSelection::ETopologyType::FGroupTopology);
+				OutputSelection->SetSelection(NewSelection);
 			}
 
 			bool bModifiedTopology = (ModifiedTopologyCounter > 0);
@@ -496,33 +477,21 @@ void UEditMeshPolygonsTool::Shutdown(EToolShutdownType ShutdownType)
 			if (bModifiedTopology)
 			{
 				// Store the compact maps if we have a selection that we need to update
-				CurrentMesh->CompactInPlace(NewStoredToolSelection ? &CompactMaps : nullptr);
+				CurrentMesh->CompactInPlace(OutputSelection ? &CompactMaps : nullptr);
 			}
 
 			// Finish prepping the stored selection
-			if (NewStoredToolSelection)
+			if (OutputSelection)
 			{
-				SelectionMechanic->GetStorableSelection(*NewStoredToolSelection, bModifiedTopology ? &CompactMaps : nullptr);
+				SelectionMechanic->GetSelection(*OutputSelection, bModifiedTopology ? &CompactMaps : nullptr);
 			}
 
 			// Bake CurrentMesh back to target inside an undo transaction
 			GetToolManager()->BeginUndoTransaction(LOCTEXT("EditMeshPolygonsToolTransactionName", "Deform Mesh"));
 			UE::ToolTarget::CommitDynamicMeshUpdate(Target, *CurrentMesh, bModifiedTopology);
 
-			// The stored selection change should go into this transaction as well.
-			// If we're keeping the same selection, we still need to store it back, though we could do it outside
-			// the transaction if we wanted to (but no real reason to). We do want to keep the same UObject if
-			// the selection is the same though, since it's probably getting kept alive by the undo stack anyway.
-			if (StoredToolSelection && NewStoredToolSelection && *StoredToolSelection == *NewStoredToolSelection)
-			{
-				GetToolManager()->RequestToolSelectionStore(StoredToolSelection);
-			}
-			else
-			{
-				// If NewStoredToolSelection is empty, this will clear the stored selection
-				GetToolManager()->RequestToolSelectionStore(NewStoredToolSelection);
-			}
-			
+			UE::Geometry::SetToolOutputSelection(this, OutputSelection);
+		
 			GetToolManager()->EndUndoTransaction();
 		}
 
@@ -553,7 +522,6 @@ void UEditMeshPolygonsTool::Shutdown(EToolShutdownType ShutdownType)
 
 	SelectionMechanic = nullptr;
 	DragAlignmentMechanic = nullptr;
-	StoredToolSelection = nullptr;
 
 	TransformGizmo = nullptr;
 	TransformProxy = nullptr;
diff --git a/Engine/Plugins/Experimental/MeshModelingToolsetExp/Source/MeshModelingToolsExp/Private/UVProjectionTool.cpp b/Engine/Plugins/Experimental/MeshModelingToolsetExp/Source/MeshModelingToolsExp/Private/UVProjectionTool.cpp
--- a/Engine/Plugins/Experimental/MeshModelingToolsetExp/Source/MeshModelingToolsExp/Private/UVProjectionTool.cpp
+++ b/Engine/Plugins/Experimental/MeshModelingToolsetExp/Source/MeshModelingToolsExp/Private/UVProjectionTool.cpp
@@ -291,9 +291,9 @@ void UUVProjectionTool::InitializeMesh()
 
 	// initialize triangle ROI if one exists
 	TriangleROI = MakeShared<TArray<int32>, ESPMode::ThreadSafe>();
-	if (HasInputSelection())
+	if (USingleSelectionMeshEditingTool::HasInputSelection())
 	{
-		GetStoredSelectionAsTriangles(GetInputSelection(), *InputMesh, *TriangleROI);
+		UE::Geometry::GetStoredSelectionAsTriangles(GetInputSelection(), *InputMesh, *TriangleROI);
 		TriangleROISet.Append(*TriangleROI);
 	}
 	VertexROI = MakeShared<TArray<int32>, ESPMode::ThreadSafe>();
diff --git a/Engine/Plugins/Experimental/MeshModelingToolsetExp/Source/MeshModelingToolsExp/Public/EditMeshPolygonsTool.h b/Engine/Plugins/Experimental/MeshModelingToolsetExp/Source/MeshModelingToolsExp/Public/EditMeshPolygonsTool.h
--- a/Engine/Plugins/Experimental/MeshModelingToolsetExp/Source/MeshModelingToolsExp/Public/EditMeshPolygonsTool.h
+++ b/Engine/Plugins/Experimental/MeshModelingToolsetExp/Source/MeshModelingToolsExp/Public/EditMeshPolygonsTool.h
@@ -11,7 +11,7 @@
 #include "InteractiveToolBuilder.h"
 #include "InteractiveToolQueryInterfaces.h" // IInteractiveToolNestedAcceptCancelAPI
 #include "Operations/GroupTopologyDeformer.h"
-#include "SingleSelectionTool.h"
+#include "BaseTools/SingleSelectionMeshEditingTool.h"
 
 #include "GeometryBase.h"
 
@@ -21,7 +21,7 @@ PREDECLARE_GEOMETRY(class FGroupTopology);
 PREDECLARE_GEOMETRY(struct FGroupTopologySelection);
 
 class UDragAlignmentMechanic;
-class UGroupTopologyStorableSelection;
+class UPersistentMeshSelection;
 class UMeshOpPreviewWithBackgroundCompute; 
 class FMeshVertexChangeBuilder;
 class UPolyEditInsertEdgeActivity;
@@ -39,17 +39,16 @@ class UTransformProxy;
  * ToolBuilder
  */
 UCLASS()
-class MESHMODELINGTOOLSEXP_API UEditMeshPolygonsToolBuilder : public UInteractiveToolBuilder
+class MESHMODELINGTOOLSEXP_API UEditMeshPolygonsToolBuilder : public USingleSelectionMeshEditingToolBuilder
 {
 	GENERATED_BODY()
 public:
 	bool bTriangleMode = false;
 
-	virtual bool CanBuildTool(const FToolBuilderState& SceneState) const override;
-	virtual UInteractiveTool* BuildTool(const FToolBuilderState& SceneState) const override;
+	virtual USingleSelectionMeshEditingTool* CreateNewTool(const FToolBuilderState& SceneState) const override;
+	virtual void InitializeNewTool(USingleSelectionMeshEditingTool* Tool, const FToolBuilderState& SceneState) const override;
 
-protected:
-	virtual const FToolTargetTypeRequirements& GetTargetRequirements() const override;
+	virtual bool WantsInputSelectionIfAvailable() const { return true; }
 };
 
 
@@ -134,7 +133,7 @@ class MESHMODELINGTOOLSEXP_API UEditMeshPolygonsActionModeToolBuilder : public U
 public:
 	EEditMeshPolygonsToolActions StartupAction = EEditMeshPolygonsToolActions::Extrude;
 
-	virtual UInteractiveTool* BuildTool(const FToolBuilderState& SceneState) const override;
+	virtual void InitializeNewTool(USingleSelectionMeshEditingTool* Tool, const FToolBuilderState& SceneState) const override;
 };
 
 UENUM()
@@ -155,7 +154,7 @@ class MESHMODELINGTOOLSEXP_API UEditMeshPolygonsSelectionModeToolBuilder : publi
 public:
 	EEditMeshPolygonsToolSelectionMode SelectionMode = EEditMeshPolygonsToolSelectionMode::Faces;
 
-	virtual UInteractiveTool* BuildTool(const FToolBuilderState& SceneState) const override;
+	virtual void InitializeNewTool(USingleSelectionMeshEditingTool* Tool, const FToolBuilderState& SceneState) const override;
 };
 
 
@@ -376,7 +375,7 @@ public:
  *
  */
 UCLASS()
-class MESHMODELINGTOOLSEXP_API UEditMeshPolygonsTool : public USingleSelectionTool, 
+class MESHMODELINGTOOLSEXP_API UEditMeshPolygonsTool : public USingleSelectionMeshEditingTool,
 	public IToolActivityHost, 
 	public IMeshVertexCommandChangeTarget,
 	public IInteractiveToolNestedAcceptCancelAPI
@@ -390,15 +389,6 @@ public:
 	virtual void RegisterActions(FInteractiveToolActionSet& ActionSet) override;
 	void EnableTriangleMode();
 
-	/**
-	 * This should be set before tool Setup() is called to allow the tool to load
-	 * the passed-in selection.
-	 */
-	void SetStoredToolSelection(const UGroupTopologyStorableSelection *StoredToolSelectionIn) 
-	{ 
-		StoredToolSelection = StoredToolSelectionIn; 
-	}
-
 	virtual void SetWorld(UWorld* World) { this->TargetWorld = World; }
 
 	virtual void Setup() override;
@@ -496,9 +486,6 @@ protected:
 	UPROPERTY()
 	TObjectPtr<UDragAlignmentMechanic> DragAlignmentMechanic = nullptr;
 
-	UPROPERTY()
-	TObjectPtr<const UGroupTopologyStorableSelection> StoredToolSelection = nullptr;
-
 	UPROPERTY()
 	TObjectPtr<UTransformGizmo> TransformGizmo = nullptr;
 
@@ -507,7 +494,7 @@ protected:
 
 	FText DefaultMessage;
 
-	bool IsStoredToolSelectionUsable(const UGroupTopologyStorableSelection* StoredSelection);
+	bool IsToolInputSelectionUsable(const UPersistentMeshSelection* InputSelection);
 	bool bSelectionStateDirty = false;
 	void OnSelectionModifiedEvent();
 
diff --git a/Engine/Plugins/Experimental/MeshModelingToolsetExp/Source/ModelingComponents/Private/BaseTools/SingleSelectionMeshEditingTool.cpp b/Engine/Plugins/Experimental/MeshModelingToolsetExp/Source/ModelingComponents/Private/BaseTools/SingleSelectionMeshEditingTool.cpp
--- a/Engine/Plugins/Experimental/MeshModelingToolsetExp/Source/ModelingComponents/Private/BaseTools/SingleSelectionMeshEditingTool.cpp
+++ b/Engine/Plugins/Experimental/MeshModelingToolsetExp/Source/ModelingComponents/Private/BaseTools/SingleSelectionMeshEditingTool.cpp
@@ -7,6 +7,7 @@
 #include "TargetInterfaces/MeshDescriptionProvider.h"
 #include "TargetInterfaces/PrimitiveComponentBackedTarget.h"
 #include "ToolTargetManager.h"
+#include "Selection/StoredMeshSelectionUtil.h"
 
 /*
  * ToolBuilder
@@ -41,9 +42,13 @@ void USingleSelectionMeshEditingToolBuilder::InitializeNewTool(USingleSelectionM
 	NewTool->SetTarget(Target);
 	NewTool->SetWorld(SceneState.World);
 
-	if (WantsInputSelectionIfAvailable() && SceneState.StoredToolSelection != nullptr)
+	if (WantsInputSelectionIfAvailable())
 	{
-		NewTool->SetInputSelection(SceneState.StoredToolSelection);
+		const UPersistentMeshSelection* InputSelection = UE::Geometry::GetCurrentToolInputSelection(SceneState, Target);
+		if (InputSelection != nullptr)
+		{
+			NewTool->SetInputSelection(InputSelection);
+		}
 	}
 }
 
diff --git a/Engine/Plugins/Experimental/MeshModelingToolsetExp/Source/ModelingComponents/Private/Selection/GroupTopologyStorableSelection.cpp b/Engine/Plugins/Experimental/MeshModelingToolsetExp/Source/ModelingComponents/Private/Selection/GroupTopologyStorableSelection.cpp
deleted file mode 100644
diff --git a/Engine/Plugins/Experimental/MeshModelingToolsetExp/Source/ModelingComponents/Private/Selection/PersistentMeshSelection.cpp b/Engine/Plugins/Experimental/MeshModelingToolsetExp/Source/ModelingComponents/Private/Selection/PersistentMeshSelection.cpp
new file mode 100644
--- /dev/null
+++ b/Engine/Plugins/Experimental/MeshModelingToolsetExp/Source/ModelingComponents/Private/Selection/PersistentMeshSelection.cpp
@@ -0,0 +1,140 @@
+// Copyright Epic Games, Inc. All Rights Reserved.
+
+#include "Selection/PersistentMeshSelection.h"
+#include "DynamicMesh/DynamicMesh3.h"
+#include "Selection/ToolSelectionUtil.h"
+#include "Util/CompactMaps.h"
+#include "GroupTopology.h"
+#include "Selection/StoredMeshSelectionUtil.h"
+#include "MeshRegionBoundaryLoops.h"
+
+#include "ExplicitUseGeometryMathTypes.h"		// using UE::Geometry::(math types)
+using namespace UE::Geometry;
+
+
+void UPersistentMeshSelection::SetSelection(const FGroupTopology& TopologyIn, const FGroupTopologySelection& SelectionIn,
+											const FCompactMaps* CompactMaps)
+{
+	FGenericMeshSelection& Data = this->Selection;
+
+	Data.FaceIDs = SelectionIn.SelectedGroupIDs.Array();
+	Data.VertexIDs.Reset();
+	Data.EdgeIDs.Reset();
+
+	const FDynamicMesh3* GroupMesh = TopologyIn.GetMesh();
+
+	for (int32 CornerID : SelectionIn.SelectedCornerIDs)
+	{
+		int32 CornerVID = TopologyIn.GetCornerVertexID(CornerID);
+		if (CompactMaps != nullptr)
+		{
+			CornerVID = CompactMaps->GetVertexMapping(CornerVID);
+		}
+		Data.VertexIDs.Add(CornerVID);
+
+		Data.RenderVertices.Add(GroupMesh->GetVertex(CornerVID));
+	}
+
+	for (int32 GroupEdgeID : SelectionIn.SelectedEdgeIDs)
+	{
+		FIndex2i EdgeVerts = (CompactMaps != nullptr) ?
+			GetGroupEdgeRepresentativeVerts(TopologyIn, GroupEdgeID, *CompactMaps) : GetGroupEdgeRepresentativeVerts(TopologyIn, GroupEdgeID);
+		Data.EdgeIDs.Add(EdgeVerts);
+
+		const FGroupTopology::FGroupEdge& GroupEdge = TopologyIn.Edges[GroupEdgeID];
+		const TArray<int32>& GroupEdgeVertices = GroupEdge.Span.Vertices;
+		int32 N = GroupEdgeVertices.Num();
+		int32 EdgeVertA = (CompactMaps) ? CompactMaps->GetVertexMapping(GroupEdgeVertices[0]) : GroupEdgeVertices[0];
+		for (int32 k = 1; k < N; ++k)
+		{
+			int32 EdgeVertB = (CompactMaps) ? CompactMaps->GetVertexMapping(GroupEdgeVertices[k]) : GroupEdgeVertices[k];
+			Data.RenderEdges.Add( FSegment3d(GroupMesh->GetVertex(EdgeVertA), GroupMesh->GetVertex(EdgeVertB)) );
+			EdgeVertA = EdgeVertB;
+		}
+	}
+
+	if (Data.FaceIDs.Num() > 0)
+	{
+		TArray<int32> TriangleIDs;
+		TopologyIn.GetSelectedTriangles(SelectionIn, TriangleIDs);
+
+		if (CompactMaps)
+		{
+			for (int32& tid : TriangleIDs)
+			{
+				tid = CompactMaps->GetTriangleMapping(tid);
+			}
+		}
+		
+		FMeshRegionBoundaryLoops RegionLoops(GroupMesh, TriangleIDs, true);
+		for (const FEdgeLoop& Loop : RegionLoops.Loops)
+		{
+			int32 N = Loop.GetVertexCount();
+			for (int32 k = 0; k < N; ++k )
+			{
+				Data.RenderEdges.Add( FSegment3d(GroupMesh->GetVertex(Loop.Vertices[k]), GroupMesh->GetVertex(Loop.Vertices[(k+1)%N])) );
+			}
+		}
+	}
+
+
+}
+
+
+void UPersistentMeshSelection::ExtractIntoSelectionObject(const FGroupTopology& TopologyIn, FGroupTopologySelection& SelectionOut) const
+{
+	const FGenericMeshSelection& Data = this->Selection;
+
+	SelectionOut.Clear();
+
+	const FDynamicMesh3* Mesh = TopologyIn.GetMesh();
+	if (!Mesh)
+	{
+		ensureMsgf(false, TEXT("FStoredGroupTopologySelection::ExtractIntoSelectionObject: target topology must have valid underlying mesh. "));
+		return;
+	}
+
+	SelectionOut.SelectedGroupIDs = TSet<int32>(Data.FaceIDs);
+
+	for (int32 Vid : Data.VertexIDs)
+	{
+		if (!Mesh->IsVertex(Vid))
+		{
+			ensureMsgf(false, TEXT("FStoredGroupTopologySelection::ExtractIntoSelectionObject: target topology's mesh was missing a vertex ID. "
+				"Perhaps the mesh was compacted without updating the stored selection?"));
+			continue;
+		}
+		int32 CornerID = TopologyIn.GetCornerIDFromVertexID(Vid);
+		if (CornerID == IndexConstants::InvalidID)
+		{
+			ensureMsgf(false, TEXT("FStoredGroupTopologySelection::ExtractIntoSelectionObject: target topology did not have an expected vert as a corner. "
+				"Is the topology initialized, and based on the same mesh?"));
+			continue;
+		}
+		SelectionOut.SelectedCornerIDs.Add(CornerID);
+	}
+	for (const FIndex2i& EdgeVerts : Data.EdgeIDs)
+	{
+		if (!Mesh->IsVertex(EdgeVerts.A) || !Mesh->IsVertex(EdgeVerts.B))
+		{
+			ensureMsgf(false, TEXT("FStoredGroupTopologySelection::ExtractIntoSelectionObject: target topology's mesh was missing a vertex ID. "
+				"Perhaps the mesh was compacted without updating the stored selection?"));
+			continue;
+		}
+		int32 Eid = Mesh->FindEdge(EdgeVerts.A, EdgeVerts.B);
+		if (Eid == IndexConstants::InvalidID)
+		{
+			ensureMsgf(false, TEXT("FStoredGroupTopologySelection::ExtractIntoSelectionObject: target topology's mesh was missing an expected edge."));
+			continue;
+		}
+		int32 GroupEdgeID = TopologyIn.FindGroupEdgeID(Eid);
+		if (Eid == IndexConstants::InvalidID)
+		{
+			ensureMsgf(false, TEXT("FStoredGroupTopologySelection::ExtractIntoSelectionObject: target topology did not have an expected group edge."
+				"Is the topology initialized, and based on the same mesh?"));
+			continue;
+		}
+
+		SelectionOut.SelectedEdgeIDs.Add(GroupEdgeID);
+	}
+}
\ No newline at end of file
diff --git a/Engine/Plugins/Experimental/MeshModelingToolsetExp/Source/ModelingComponents/Private/Selection/PersistentMeshSelectionManager.cpp b/Engine/Plugins/Experimental/MeshModelingToolsetExp/Source/ModelingComponents/Private/Selection/PersistentMeshSelectionManager.cpp
new file mode 100644
--- /dev/null
+++ b/Engine/Plugins/Experimental/MeshModelingToolsetExp/Source/ModelingComponents/Private/Selection/PersistentMeshSelectionManager.cpp
@@ -0,0 +1,216 @@
+// Copyright Epic Games, Inc. All Rights Reserved.
+
+#include "Selection/PersistentMeshSelectionManager.h"
+#include "Selection/PersistentMeshSelection.h"
+
+#include "InteractiveToolsContext.h"
+#include "InteractiveToolManager.h"
+#include "ContextObjectStore.h"
+#include "Drawing/PreviewGeometryActor.h"
+
+#include "ExplicitUseGeometryMathTypes.h"		// using UE::Geometry::(math types)
+using namespace UE::Geometry;
+
+#define LOCTEXT_NAMESPACE "UPersistentMeshSelectionManager"
+
+void UPersistentMeshSelectionManager::Initialize(TObjectPtr<UInteractiveToolsContext> ToolsContext)
+{
+	ParentContext = ToolsContext;
+}
+
+void UPersistentMeshSelectionManager::Shutdown()
+{
+	if (SelectionDisplay != nullptr)
+	{
+		SelectionDisplay->Disconnect();
+		SelectionDisplay = nullptr;
+	}
+}
+
+bool UPersistentMeshSelectionManager::HasActiveSelection()
+{
+	return (ActiveSelection != nullptr);
+}
+
+UPersistentMeshSelection* UPersistentMeshSelectionManager::GetActiveSelection()
+{
+	return ActiveSelection;
+}
+
+void UPersistentMeshSelectionManager::SetNewActiveSelection(UPersistentMeshSelection* Selection)
+{
+	TUniquePtr<FPersistentMeshSelectionChange> SelectionChange = MakeUnique<FPersistentMeshSelectionChange>();
+	if (ActiveSelection != nullptr)
+	{
+		SelectionChange->From = ActiveSelection->GetSelection();
+	}
+	if (Selection != nullptr)
+	{
+		SelectionChange->To = Selection->GetSelection();
+	}
+	ParentContext->ToolManager->GetContextTransactionsAPI()->AppendChange(this, MoveTemp(SelectionChange),
+		LOCTEXT("SelectionChange", "Selection Change"));
+
+	ActiveSelection = Selection;
+	OnSelectionModified();
+}
+
+
+void UPersistentMeshSelectionManager::SetNewActiveSelectionInternal(UPersistentMeshSelection* Selection)
+{
+	ActiveSelection = Selection;
+	OnSelectionModified();
+}
+
+
+void UPersistentMeshSelectionManager::ClearActiveSelection()
+{
+	if (ActiveSelection == nullptr)
+	{
+		return;
+	}
+
+	SetNewActiveSelection(nullptr);
+}
+
+
+
+
+void UPersistentMeshSelectionManager::OnSelectionModified()
+{
+	if (ActiveSelection != nullptr && SelectionDisplay == nullptr)
+	{
+		SelectionDisplay = NewObject<UPreviewGeometry>(ParentContext);
+		SelectionDisplay->CreateInWorld(ActiveSelection->GetTargetComponent()->GetWorld(), FTransform::Identity);
+	}
+	if (SelectionDisplay == nullptr)
+	{
+		return;
+	}
+
+	const FGenericMeshSelection* SelectionData = (ActiveSelection != nullptr) ? &ActiveSelection->GetSelection() : nullptr;
+
+	bool bShowLines = (SelectionData != nullptr) && (SelectionData->HasRenderableLines());
+
+	if (bShowLines)
+	{
+		const FColor ROIBorderColor(240, 15, 240);
+		const float ROIBorderThickness = 8.0f;
+		//const float ROIBorderDepthBias = 0.1f * (float)(WorldBounds.DiagonalLength() * 0.01);
+		const float ROIBorderDepthBias = 0.01f;
+
+		UE::Geometry::FTransform3d Transform(ActiveSelection->GetTargetComponent()->GetComponentToWorld());
+
+		const TArray<UE::Geometry::FSegment3d>& Lines = SelectionData->RenderEdges;
+		SelectionDisplay->CreateOrUpdateLineSet(TEXT("SelectionEdges"),  Lines.Num(),
+			[&](int32 Index, TArray<FRenderableLine>& LinesOut) 
+			{
+				const UE::Geometry::FSegment3d& Segment = Lines[Index];
+				FVector3d A = Transform.TransformPosition(Segment.StartPoint());
+				FVector3d B = Transform.TransformPosition(Segment.EndPoint());
+				LinesOut.Add(FRenderableLine((FVector)A, (FVector)B, ROIBorderColor, ROIBorderThickness, ROIBorderDepthBias));
+			}, 1);
+
+		SelectionDisplay->SetLineSetVisibility(TEXT("SelectionEdges"), true);
+	}
+	else
+	{
+		SelectionDisplay->SetLineSetVisibility(TEXT("SelectionEdges"), false);
+	}
+
+}
+
+
+
+
+bool UE::Geometry::RegisterPersistentMeshSelectionManager(UInteractiveToolsContext* ToolsContext)
+{
+	if (ensure(ToolsContext))
+	{
+		UPersistentMeshSelectionManager* Found = ToolsContext->ContextObjectStore->FindContext<UPersistentMeshSelectionManager>();
+		if (Found == nullptr)
+		{
+			UPersistentMeshSelectionManager* SelectionManager = NewObject<UPersistentMeshSelectionManager>(ToolsContext->ToolManager);
+			if (ensure(SelectionManager))
+			{
+				SelectionManager->Initialize(ToolsContext);
+				ToolsContext->ContextObjectStore->AddContextObject(SelectionManager);
+				return true;
+			}
+			else
+			{
+				return false;
+			}
+		}
+		return true;
+	}
+	return false;
+}
+
+
+
+bool UE::Geometry::DeregisterPersistentMeshSelectionManager(UInteractiveToolsContext* ToolsContext)
+{
+	if (ensure(ToolsContext))
+	{
+		UPersistentMeshSelectionManager* Found = ToolsContext->ContextObjectStore->FindContext<UPersistentMeshSelectionManager>();
+		if (Found != nullptr)
+		{
+			Found->Shutdown();
+			ToolsContext->ContextObjectStore->RemoveContextObject(Found);
+		}
+		return true;
+	}
+	return false;
+}
+
+
+UPersistentMeshSelectionManager* UE::Geometry::FindPersistentMeshSelectionManager(UInteractiveToolManager* ToolManager)
+{
+	if (ensure(ToolManager))
+	{
+		UPersistentMeshSelectionManager* Found = ToolManager->GetContextObjectStore()->FindContext<UPersistentMeshSelectionManager>();
+		if (Found != nullptr)
+		{
+			return Found;
+		}
+	}
+	return nullptr;
+}
+
+
+
+
+void FPersistentMeshSelectionChange::Apply(UObject* Object)
+{
+	UPersistentMeshSelectionManager* SelectionManager = Cast<UPersistentMeshSelectionManager>(Object);
+	if (SelectionManager)
+	{
+		UPersistentMeshSelection* NewSelection = NewObject<UPersistentMeshSelection>(SelectionManager);
+		NewSelection->SetSelection(To);
+		SelectionManager->SetNewActiveSelectionInternal(NewSelection);
+	}
+}
+
+void FPersistentMeshSelectionChange::Revert(UObject* Object)
+{
+	UPersistentMeshSelectionManager* SelectionManager = Cast<UPersistentMeshSelectionManager>(Object);
+	if (SelectionManager)
+	{
+		UPersistentMeshSelection* NewSelection = NewObject<UPersistentMeshSelection>(SelectionManager);
+		NewSelection->SetSelection(From);
+		SelectionManager->SetNewActiveSelectionInternal(NewSelection);
+	}
+}
+
+bool FPersistentMeshSelectionChange::HasExpired(UObject* Object) const
+{
+	return false;
+}
+
+FString FPersistentMeshSelectionChange::ToString() const
+{
+	return TEXT("PersistentMeshSelectionChange");
+}
+
+#undef LOCTEXT_NAMESPACE
\ No newline at end of file
diff --git a/Engine/Plugins/Experimental/MeshModelingToolsetExp/Source/ModelingComponents/Private/Selection/PolygonSelectionMechanic.cpp b/Engine/Plugins/Experimental/MeshModelingToolsetExp/Source/ModelingComponents/Private/Selection/PolygonSelectionMechanic.cpp
--- a/Engine/Plugins/Experimental/MeshModelingToolsetExp/Source/ModelingComponents/Private/Selection/PolygonSelectionMechanic.cpp
+++ b/Engine/Plugins/Experimental/MeshModelingToolsetExp/Source/ModelingComponents/Private/Selection/PolygonSelectionMechanic.cpp
@@ -7,7 +7,7 @@
 #include "BaseBehaviors/SingleClickBehavior.h"
 #include "InteractiveToolManager.h"
 #include "Util/ColorConstants.h"
-#include "Selection/GroupTopologyStorableSelection.h"
+#include "Selection/PersistentMeshSelection.h"
 #include "ToolSceneQueriesUtil.h"
 #include "ToolSetupUtil.h"
 
@@ -83,21 +83,14 @@ void UPolygonSelectionMechanic::Shutdown()
 	}
 }
 
-void UPolygonSelectionMechanic::GetStorableSelection(UGroupTopologyStorableSelection& StorableSelectionOut, const FCompactMaps* CompactMapsToApply) const
+void UPolygonSelectionMechanic::GetSelection(UPersistentMeshSelection& SelectionOut, const FCompactMaps* CompactMapsToApply) const
 {
-	if (CompactMapsToApply)
-	{
-		StorableSelectionOut.SetSelection(*Topology, PersistentSelection, *CompactMapsToApply);
-	}
-	else
-	{
-		StorableSelectionOut.SetSelection(*Topology, PersistentSelection);
-	}
+	SelectionOut.SetSelection(*Topology, PersistentSelection, CompactMapsToApply);
 }
 
-void UPolygonSelectionMechanic::LoadStorableSelection(const UGroupTopologyStorableSelection& StorableSelectionIn)
+void UPolygonSelectionMechanic::LoadSelection(const UPersistentMeshSelection& SelectionIn)
 {
-	StorableSelectionIn.ExtractIntoSelectionObject(*Topology, PersistentSelection);
+	SelectionIn.ExtractIntoSelectionObject(*Topology, PersistentSelection);
 }
 
 void UPolygonSelectionMechanic::Initialize(
diff --git a/Engine/Plugins/Experimental/MeshModelingToolsetExp/Source/ModelingComponents/Private/Selection/StoredMeshSelectionUtil.cpp b/Engine/Plugins/Experimental/MeshModelingToolsetExp/Source/ModelingComponents/Private/Selection/StoredMeshSelectionUtil.cpp
--- a/Engine/Plugins/Experimental/MeshModelingToolsetExp/Source/ModelingComponents/Private/Selection/StoredMeshSelectionUtil.cpp
+++ b/Engine/Plugins/Experimental/MeshModelingToolsetExp/Source/ModelingComponents/Private/Selection/StoredMeshSelectionUtil.cpp
@@ -1,62 +1,68 @@
 // Copyright Epic Games, Inc. All Rights Reserved.
 
 #include "Selection/StoredMeshSelectionUtil.h"
-#include "Selection/GroupTopologyStorableSelection.h"
+#include "Selection/PersistentMeshSelection.h"
+#include "Selection/PersistentMeshSelectionManager.h"
 #include "TargetInterfaces/PrimitiveComponentBackedTarget.h"
 #include "ToolTargetManager.h"
+#include "InteractiveTool.h"
+#include "InteractiveToolManager.h"
+#include "GroupTopology.h"
 
+#include "ExplicitUseGeometryMathTypes.h"		// using UE::Geometry::(math types)
+using namespace UE::Geometry;
 
-const UInteractiveToolStorableSelection* UE::Geometry::GetCurrentToolInputSelection(const FToolBuilderState& SceneState, UToolTarget* Target)
+
+const UPersistentMeshSelection* UE::Geometry::GetCurrentToolInputSelection(const FToolBuilderState& SceneState, UToolTarget* Target)
 {
-	const UInteractiveToolStorableSelection* Selection = SceneState.StoredToolSelection;
-	if (Selection == nullptr)
+	UPersistentMeshSelectionManager* SelectionManager = FindPersistentMeshSelectionManager(SceneState.ToolManager);
+	if (SelectionManager == nullptr)
+	{
+		return nullptr;
+	}
+
+	UPersistentMeshSelection* ActiveSelection = SelectionManager->GetActiveSelection();
+	if (ActiveSelection == nullptr)
 	{
 		return nullptr;
 	}
 
+	// todo: more validation here?
 	IPrimitiveComponentBackedTarget* TargetInterface = Cast<IPrimitiveComponentBackedTarget>(Target);
 	if (TargetInterface == nullptr)
 	{
 		return nullptr;
 	}
-
-	const UGroupTopologyStorableSelection* GroupTopoSelection = Cast<UGroupTopologyStorableSelection>(Selection);
-	if (GroupTopoSelection == nullptr || GroupTopoSelection->IdentifyingInfo.ComponentTarget != TargetInterface->GetOwnerComponent())
+	if (TargetInterface->GetOwnerComponent() != ActiveSelection->GetTargetComponent())
 	{
 		return nullptr;
 	}
 
-	return Selection;
+	return ActiveSelection;
 }
 
 
 
 bool UE::Geometry::GetStoredSelectionAsTriangles(
-	const UInteractiveToolStorableSelection* Selection,
+	const UPersistentMeshSelection* Selection,
 	const FDynamicMesh3& Mesh,
 	TArray<int32>& TrianglesOut)
 {
-	const UGroupTopologyStorableSelection* GroupTopoSelection = Cast<UGroupTopologyStorableSelection>(Selection);
-	if (GroupTopoSelection == nullptr)
-	{
-		return false;
-	}
-
 	// don't support UV selection currently - unclear how reproducible this is?
-	if (GroupTopoSelection->IdentifyingInfo.TopologyType == UGroupTopologyStorableSelection::ETopologyType::FUVGroupTopology)
+	if (Selection->GetSelectionType() == FGenericMeshSelection::ETopologyType::FUVGroupTopology)
 	{
 		return false;
 	}
 
-	if (GroupTopoSelection->GetGroupIDs().Num() == 0)
+	if (Selection->GetSelection().FaceIDs.Num() == 0)
 	{
 		return false;
 	}
 
-	if (GroupTopoSelection->IdentifyingInfo.TopologyType == UGroupTopologyStorableSelection::ETopologyType::FGroupTopology)
+	if (Selection->GetSelectionType() == FGenericMeshSelection::ETopologyType::FGroupTopology)
 	{
 		TSet<int32> SelectedGroups;
-		for (int32 gid : GroupTopoSelection->GetGroupIDs())
+		for (int32 gid : Selection->GetSelection().FaceIDs)
 		{
 			SelectedGroups.Add(gid);
 		}
@@ -70,9 +76,9 @@ bool UE::Geometry::GetStoredSelectionAsTriangles(
 			}
 		}
 	}
-	else if (GroupTopoSelection->IdentifyingInfo.TopologyType == UGroupTopologyStorableSelection::ETopologyType::FTriangleGroupTopology)
+	else if (Selection->GetSelectionType() == FGenericMeshSelection::ETopologyType::FTriangleGroupTopology)
 	{
-		for (int gid : GroupTopoSelection->GetGroupIDs())
+		for (int gid : Selection->GetSelection().FaceIDs)
 		{
 			if (Mesh.IsTriangle(gid))
 			{
@@ -82,4 +88,125 @@ bool UE::Geometry::GetStoredSelectionAsTriangles(
 	}
 
 	return true;
-}
\ No newline at end of file
+}
+
+
+
+
+void UE::Geometry::ClearActiveToolSelection(UInteractiveToolManager* ToolManager)
+{
+	UPersistentMeshSelectionManager* SelectionManager = FindPersistentMeshSelectionManager(ToolManager);
+	if (SelectionManager != nullptr)
+	{
+		SelectionManager->ClearActiveSelection();
+	}
+}
+
+
+bool UE::Geometry::SetToolOutputSelection(UInteractiveTool* Tool, UPersistentMeshSelection* Selection)
+{
+	UPersistentMeshSelectionManager* SelectionManager = FindPersistentMeshSelectionManager(Tool->GetToolManager());
+	if (SelectionManager == nullptr)
+	{
+		return false;
+	}
+
+	if (Selection != nullptr)
+	{
+		SelectionManager->SetNewActiveSelection(Selection);
+	}
+	else
+	{
+		SelectionManager->ClearActiveSelection();
+	}
+
+	return true;
+}
+
+
+
+
+
+namespace UELocal
+{
+	// Helper function for GetGroupEdgeRepresentativeVerts(). Given a loop as a list of vids, with
+	// the first and last vid the same, returns the lowest vid and its lower-vid neighbor.
+	FIndex2i GetLoopRepresentativeVerts(const TArray<int32>& Verts)
+	{
+		int32 NumVerts = Verts.Num();
+
+		int32 MinVid = Verts[0];
+		// Last vert is a repeat of first vert, so neighbor is second to last.
+		int32 MinNeighbor = FMath::Min(Verts[1], Verts[NumVerts - 2]);
+
+		for (int32 i = 1; i < NumVerts - 1; ++i)
+		{
+			if (Verts[i] < MinVid)
+			{
+				MinVid = Verts[i];
+				MinNeighbor = FMath::Min(Verts[i - 1], Verts[i + 1]);
+			}
+		}
+		return FIndex2i(MinVid, MinNeighbor); // we know that MinVid is smaller
+	}
+}//end namespace UELocal
+using namespace UELocal;
+
+
+
+
+
+UE::Geometry::FIndex2i UE::Geometry::GetGroupEdgeRepresentativeVerts(const FGroupTopology& TopologyIn, int GroupEdgeID, const FCompactMaps& CompactMaps)
+{
+	check(GroupEdgeID >= 0 && GroupEdgeID < TopologyIn.Edges.Num());
+
+	const FGroupTopology::FGroupEdge& GroupEdge = TopologyIn.Edges[GroupEdgeID];
+	const TArray<int32>& Verts = GroupEdge.Span.Vertices;
+
+	if (GroupEdge.EndpointCorners.A != IndexConstants::InvalidID)
+	{
+		// Use remapped vids
+		int32 FirstVid = CompactMaps.GetVertexMapping(Verts[0]);
+		int32 FirstNeighbor = CompactMaps.GetVertexMapping(Verts[1]);
+		int32 LastVid = CompactMaps.GetVertexMapping(Verts.Last());
+		int32 LastNeighbor = CompactMaps.GetVertexMapping(Verts[Verts.Num() - 2]);
+
+		return FirstVid < LastVid ?
+			FIndex2i(FMath::Min(FirstVid, FirstNeighbor), FMath::Max(FirstVid, FirstNeighbor))
+			: FIndex2i(FMath::Min(LastVid, LastNeighbor), FMath::Max(LastVid, LastNeighbor));
+	}
+	else
+	{
+		TArray<int32> RemappedVerts;
+		RemappedVerts.SetNum(Verts.Num());
+		for (int32 i = 0; i < Verts.Num(); ++i)
+		{
+			RemappedVerts[i] = CompactMaps.GetVertexMapping(Verts[i]);
+		}
+		return GetLoopRepresentativeVerts(RemappedVerts);
+	}
+}
+
+UE::Geometry::FIndex2i UE::Geometry::GetGroupEdgeRepresentativeVerts(const FGroupTopology& TopologyIn, int GroupEdgeID)
+{
+	check(GroupEdgeID >= 0 && GroupEdgeID < TopologyIn.Edges.Num());
+
+	const FGroupTopology::FGroupEdge& GroupEdge = TopologyIn.Edges[GroupEdgeID];
+	const TArray<int32>& Verts = GroupEdge.Span.Vertices;
+
+	if (GroupEdge.EndpointCorners.A != IndexConstants::InvalidID)
+	{
+		int32 FirstVid = Verts[0];
+		int32 FirstNeighbor = Verts[1];
+		int32 LastVid = Verts.Last();
+		int32 LastNeighbor = Verts[Verts.Num() - 2];
+
+		return FirstVid < LastVid ?
+			FIndex2i(FMath::Min(FirstVid, FirstNeighbor), FMath::Max(FirstVid, FirstNeighbor))
+			: FIndex2i(FMath::Min(LastVid, LastNeighbor), FMath::Max(LastVid, LastNeighbor));
+	}
+	else
+	{
+		return GetLoopRepresentativeVerts(Verts);
+	}
+}
diff --git a/Engine/Plugins/Experimental/MeshModelingToolsetExp/Source/ModelingComponents/Public/BaseTools/SingleSelectionMeshEditingTool.h b/Engine/Plugins/Experimental/MeshModelingToolsetExp/Source/ModelingComponents/Public/BaseTools/SingleSelectionMeshEditingTool.h
--- a/Engine/Plugins/Experimental/MeshModelingToolsetExp/Source/ModelingComponents/Public/BaseTools/SingleSelectionMeshEditingTool.h
+++ b/Engine/Plugins/Experimental/MeshModelingToolsetExp/Source/ModelingComponents/Public/BaseTools/SingleSelectionMeshEditingTool.h
@@ -4,10 +4,10 @@
 
 #include "SingleSelectionTool.h"
 #include "InteractiveToolBuilder.h"
-#include "InteractiveToolStorableSelection.h"
 #include "SingleSelectionMeshEditingTool.generated.h"
 
 class USingleSelectionMeshEditingTool;
+class UPersistentMeshSelection;
 
 /**
  * USingleSelectionMeshEditingToolBuilder is a base tool builder for single
@@ -51,6 +51,8 @@ class MODELINGCOMPONENTS_API USingleSelectionMeshEditingTool : public USingleSel
 public:
 	virtual void SetWorld(UWorld* World) { TargetWorld = World; }
 
+	virtual UWorld* GetTargetWorld() { return TargetWorld; }
+
 protected:
 	UWorld* TargetWorld = nullptr;
 
@@ -64,9 +66,9 @@ public:
 	 * to allow the Tool to configure it's behavior based on the Selection (which may or may-not exist).
 	 * If the Tool requires a Selection, this needs to be handled at the Builder level.
 	 */
-	virtual void SetInputSelection(const UInteractiveToolStorableSelection* StoredToolSelectionIn)
+	virtual void SetInputSelection(const UPersistentMeshSelection* SelectionIn)
 	{
-		InputSelection = StoredToolSelectionIn;
+		InputSelection = SelectionIn;
 	}
 
 	/** @return true if an InputSelection is available */
@@ -76,7 +78,7 @@ public:
 	}
 
 	/** @return the input Selection, or nullptr if one was not configured */
-	virtual const UInteractiveToolStorableSelection* GetInputSelection() const
+	virtual const UPersistentMeshSelection* GetInputSelection() const
 	{
 		return InputSelection;
 	}
@@ -84,9 +86,9 @@ public:
 protected:
 
 	/**
-	 * (optional) Stored Selection provided on Tool Input. This should never be modified after the Tool's Setup() is called
+	 * (optional) Mesh Selection provided on Tool Input. This should never be modified after the Tool's Setup() is called
 	 */
 	UPROPERTY()
-	TObjectPtr<const UInteractiveToolStorableSelection> InputSelection = nullptr;
+	TObjectPtr<const UPersistentMeshSelection> InputSelection = nullptr;
 
 };
diff --git a/Engine/Plugins/Experimental/MeshModelingToolsetExp/Source/ModelingComponents/Public/Selection/GroupTopologyStorableSelection.h b/Engine/Plugins/Experimental/MeshModelingToolsetExp/Source/ModelingComponents/Public/Selection/GroupTopologyStorableSelection.h
deleted file mode 100644
diff --git a/Engine/Plugins/Experimental/MeshModelingToolsetExp/Source/ModelingComponents/Public/Selection/PersistentMeshSelection.h b/Engine/Plugins/Experimental/MeshModelingToolsetExp/Source/ModelingComponents/Public/Selection/PersistentMeshSelection.h
new file mode 100644
--- /dev/null
+++ b/Engine/Plugins/Experimental/MeshModelingToolsetExp/Source/ModelingComponents/Public/Selection/PersistentMeshSelection.h
@@ -0,0 +1,145 @@
+// Copyright Epic Games, Inc. All Rights Reserved.
+
+#pragma once
+
+#include "CoreMinimal.h"
+#include "IndexTypes.h"
+#include "SegmentTypes.h"
+#include "InteractiveToolStorableSelection.h"
+
+#include "PersistentMeshSelection.generated.h"
+
+class UPrimitiveComponent;
+class UInteractiveTool;
+PREDECLARE_USE_GEOMETRY_STRUCT(FGroupTopologySelection);
+PREDECLARE_USE_GEOMETRY_CLASS(FGroupTopology);
+PREDECLARE_USE_GEOMETRY_CLASS(FCompactMaps);
+
+
+
+
+/**
+ * FGenericMeshSelection represents various types of selection on a Mesh.
+ * This includes various types of indices that could be interpreted in different ways.
+ * 
+ * In addition, "Render" geometry is stored, which can be used by higher-level
+ * code to draw the selection in some way (eg a selection highlight)
+ * 
+ * @warning this class is likely to change in the future
+ */
+struct MODELINGCOMPONENTS_API FGenericMeshSelection
+{
+	// selection type
+	enum class ETopologyType
+	{
+		FGroupTopology,
+		FTriangleGroupTopology,
+		FUVGroupTopology,
+	};
+
+
+
+	//
+	// Selection representation (may be interpreted differently depending on TopologyType)
+	//
+
+	// selection type
+	ETopologyType TopologyType = ETopologyType::FGroupTopology;
+
+	// selected vertices or "corners" (eg of polygroup topology)
+	TArray<int32> VertexIDs;
+	// selected edges, represented as index pairs because for many selections, 
+	// using a pair of vertices defining/on the edge is more reliable (due to unstable edge IDs)
+	TArray<UE::Geometry::FIndex2i> EdgeIDs;
+	// selected triangles/faces/regions
+	TArray<int32> FaceIDs;
+
+
+	//
+	// Selection Target Information
+	//
+
+	// Component this selection applies to (eg that owns mesh, etc)
+	UPrimitiveComponent* SourceComponent = nullptr;
+
+	//
+	// Renderable Selection Representation
+	//
+
+	// set of 3D points representing selection (in world space)
+	TArray<FVector3d> RenderVertices;
+	// set of 3D lines representing selection (in world space)
+	TArray<UE::Geometry::FSegment3d> RenderEdges;
+
+
+	/** @return selected GroupIDs */
+	const TArray<int32>& GetGroupIDs() const { return FaceIDs; }
+
+	/** @return true if selection is empty */
+	bool IsEmpty() const
+	{
+		return VertexIDs.IsEmpty() && EdgeIDs.IsEmpty() && FaceIDs.IsEmpty();
+	}
+
+	/** @return true if selection has 3D lines that can be rendered */
+	bool HasRenderableLines() const { return RenderEdges.Num() > 0; }
+
+
+	bool operator==(const FGenericMeshSelection& Other) const
+	{
+		return SourceComponent == Other.SourceComponent
+			&& TopologyType == Other.TopologyType
+			&& VertexIDs == Other.VertexIDs
+			&& EdgeIDs == Other.EdgeIDs
+			&& FaceIDs == Other.FaceIDs;
+	}
+
+};
+
+
+
+
+
+/**
+ * UPersistentMeshSelection is a UObject wrapper for a FGenericMeshSelection
+ */
+UCLASS()
+class MODELINGCOMPONENTS_API UPersistentMeshSelection : public UObject
+{
+	GENERATED_BODY()
+
+public:
+
+	/**
+	 * Resets the contents of the object using the given selection.
+	 *
+	 * @param CompactMaps If the mesh was compacted without updating the passed in topology object, these
+	 *  maps will be used to give an object that will work with the new mesh vids.
+	 */
+	void SetSelection(const FGroupTopology& TopologyIn, const FGroupTopologySelection& SelectionIn, const FCompactMaps* CompactMaps = nullptr);
+
+	/**
+	 * Initializes a FGroupTopologySelection using the current contents of the object. The topology
+	 * must already be initialized.
+	 */
+	void ExtractIntoSelectionObject(const FGroupTopology& TopologyIn, FGroupTopologySelection& SelectionOut) const;
+
+	/** @return true if the selection is empty */
+	bool IsEmpty() const
+	{
+		return Selection.IsEmpty();
+	}
+
+	UPrimitiveComponent* GetTargetComponent() const { return Selection.SourceComponent; }
+	FGenericMeshSelection::ETopologyType GetSelectionType() const { return Selection.TopologyType; }
+
+	/** replace the internal Selection data */
+	void SetSelection(const FGenericMeshSelection& SelectionIn) { Selection = SelectionIn; }
+
+	/** @return the internal Selection data */
+	const FGenericMeshSelection& GetSelection() const { return Selection; }
+
+protected:
+	FGenericMeshSelection Selection;
+};
+
diff --git a/Engine/Plugins/Experimental/MeshModelingToolsetExp/Source/ModelingComponents/Public/Selection/PersistentMeshSelectionManager.h b/Engine/Plugins/Experimental/MeshModelingToolsetExp/Source/ModelingComponents/Public/Selection/PersistentMeshSelectionManager.h
new file mode 100644
--- /dev/null
+++ b/Engine/Plugins/Experimental/MeshModelingToolsetExp/Source/ModelingComponents/Public/Selection/PersistentMeshSelectionManager.h
@@ -0,0 +1,127 @@
+// Copyright Epic Games, Inc. All Rights Reserved.
+
+#pragma once
+
+#include "CoreMinimal.h"
+#include "IndexTypes.h"
+#include "Selection/PersistentMeshSelection.h"
+#include "InteractiveToolChange.h"
+
+#include "PersistentMeshSelectionManager.generated.h"
+
+class UInteractiveToolsContext;
+class UInteractiveToolManager;
+class UPrimitiveComponent;
+class UInteractiveTool;
+class UPreviewGeometry;
+
+
+/**
+ * UPersistentMeshSelectionManager manages an active Mesh Selection.
+ * The assumption is that this class will be registered with the ContextStore of an InteractiveToolsContext,
+ * and the various Tool implementations will get/set the selection.
+ * 
+ * 
+ * 
+ * StoredMeshSelectionUtil.h contains various utility functions to simplify usage of this class,
+ * those functions are preferable to directly using the Manager.
+ * 
+ */
+UCLASS()
+class MODELINGCOMPONENTS_API UPersistentMeshSelectionManager : public UObject
+{
+	GENERATED_BODY()
+
+public:
+	
+	virtual void Initialize(TObjectPtr<UInteractiveToolsContext> ToolsContext);
+	virtual void Shutdown();
+
+	/** @return true if there is an active selection */
+	virtual bool HasActiveSelection();
+
+	/** @return the active selection */
+	virtual UPersistentMeshSelection* GetActiveSelection();
+
+	/** Set a new active selection. This is undoable, ie a FPersistentMeshSelectionChange will be emitted */
+	virtual void SetNewActiveSelection(UPersistentMeshSelection* Selection);
+
+	/** Clear the active selection. This is undoable, ie a FPersistentMeshSelectionChange will be emitted */
+	virtual void ClearActiveSelection();
+
+
+protected:
+	UPROPERTY()
+	TObjectPtr<UInteractiveToolsContext> ParentContext;
+
+	UPROPERTY()
+	TObjectPtr<UPersistentMeshSelection> ActiveSelection;
+
+	UPROPERTY()
+	TObjectPtr<UPreviewGeometry> SelectionDisplay;
+
+	virtual void OnSelectionModified();
+
+	virtual void SetNewActiveSelectionInternal(UPersistentMeshSelection* Selection);
+	friend class FPersistentMeshSelectionChange;
+};
+
+
+
+
+
+/**
+ * FPersistentMeshSelectionChange stores a selection change, ie before-and-after states.
+ * UPersistentMeshSelectionManager emits these changes on selection modifications.
+ */
+class MODELINGCOMPONENTS_API FPersistentMeshSelectionChange : public FToolCommandChange
+{
+public:
+	FGenericMeshSelection From;
+	FGenericMeshSelection To;
+
+	virtual void Apply(UObject* Object) override;
+	virtual void Revert(UObject* Object) override;
+	virtual bool HasExpired(UObject* Object) const override;
+	virtual FString ToString() const override;
+};
+
+
+
+
+
+
+
+
+namespace UE
+{
+namespace Geometry
+{
+	//
+	// The functions below are helper functions that simplify usage of a UPersistentMeshSelectionManager 
+	// that is registered as a ContextStoreObject in an InteractiveToolsContext
+	//
+
+	/**
+	 * If one does not already exist, create a new instance of UPersistentMeshSelectionManager and add it to the
+	 * ToolsContext's ContextObjectStore
+	 * @return true if the ContextObjectStore now has a UPersistentMeshSelectionManager (whether it already existed, or was created)
+	 */
+	MODELINGCOMPONENTS_API bool RegisterPersistentMeshSelectionManager(UInteractiveToolsContext* ToolsContext);
+
+	/**
+	 * Remove any existing UPersistentMeshSelectionManager from the ToolsContext's ContextObjectStore
+	 * @return true if the ContextObjectStore no longer has a UPersistentMeshSelectionManager (whether it was removed, or did not exist)
+	 */
+	MODELINGCOMPONENTS_API bool DeregisterPersistentMeshSelectionManager(UInteractiveToolsContext* ToolsContext);
+
+
+	/**
+	 * Find an existing UPersistentMeshSelectionManager in the ToolsContext's ContextObjectStore
+	 * @return SelectionManager pointer or nullptr if not found
+	 */
+	MODELINGCOMPONENTS_API UPersistentMeshSelectionManager* FindPersistentMeshSelectionManager(UInteractiveToolManager* ToolManager);
+
+
+}
+}
\ No newline at end of file
diff --git a/Engine/Plugins/Experimental/MeshModelingToolsetExp/Source/ModelingComponents/Public/Selection/PolygonSelectionMechanic.h b/Engine/Plugins/Experimental/MeshModelingToolsetExp/Source/ModelingComponents/Public/Selection/PolygonSelectionMechanic.h
--- a/Engine/Plugins/Experimental/MeshModelingToolsetExp/Source/ModelingComponents/Public/Selection/PolygonSelectionMechanic.h
+++ b/Engine/Plugins/Experimental/MeshModelingToolsetExp/Source/ModelingComponents/Public/Selection/PolygonSelectionMechanic.h
@@ -16,7 +16,7 @@
 #include "PolygonSelectionMechanic.generated.h"
 
 class FPolygonSelectionMechanicSelectionChange;
-class UGroupTopologyStorableSelection;
+class UPersistentMeshSelection;
 class UMouseHoverBehavior;
 class URectangleMarqueeMechanic;
 class USingleClickInputBehavior;
@@ -251,14 +251,13 @@ public:
 	 * Gives the current selection as a storable selection object. Can optionally apply the passed-in
 	 * compact maps to the object if the topology in the mechanic was not updated after compacting.
 	 */
-	void GetStorableSelection(UGroupTopologyStorableSelection& StorableSelectionOut, 
-		const FCompactMaps* CompactMapsToApply = nullptr) const;
+	void GetSelection(UPersistentMeshSelection& SelectionOut, const FCompactMaps* CompactMapsToApply = nullptr) const;
 
 	/**
 	 * Sets the current selection using the given storable selection object. The topology in the
 	 * mechanic must already be initialized for this to work.
 	 */
-	void LoadStorableSelection(const UGroupTopologyStorableSelection& StorableSelectionIn);
+	void LoadSelection(const UPersistentMeshSelection& Selection);
 
 	/**
 	 * Can be used by in an OnSelectionChanged event to inspect the clicked location (i.e., the
diff --git a/Engine/Plugins/Experimental/MeshModelingToolsetExp/Source/ModelingComponents/Public/Selection/StoredMeshSelectionUtil.h b/Engine/Plugins/Experimental/MeshModelingToolsetExp/Source/ModelingComponents/Public/Selection/StoredMeshSelectionUtil.h
--- a/Engine/Plugins/Experimental/MeshModelingToolsetExp/Source/ModelingComponents/Public/Selection/StoredMeshSelectionUtil.h
+++ b/Engine/Plugins/Experimental/MeshModelingToolsetExp/Source/ModelingComponents/Public/Selection/StoredMeshSelectionUtil.h
@@ -6,9 +6,12 @@
 #include "DynamicMesh/DynamicMesh3.h"
 
 struct FToolBuilderState;
-class UInteractiveToolStorableSelection;
+class UPersistentMeshSelection;
 class UToolTarget;
-
+class UInteractiveTool;
+class UInteractiveToolManager;
+PREDECLARE_USE_GEOMETRY_CLASS(FGroupTopology);
+PREDECLARE_USE_GEOMETRY_CLASS(FCompactMaps);
 
 //
 // Utility functions for Tool implementations to use to work with Stored Selections
@@ -19,22 +22,69 @@ namespace Geometry
 {
 
 	/**
-	 * Given a FToolBuilderState, determine if a Stored Selection exists that is applicable to the ToolTarget
-	 * @return Stored Selection if found, or nullptr if not
+	 * Given a FToolBuilderState, determine if an Input Selection exists that is applicable to the ToolTarget
+	 * @return Selection if found, or nullptr if not
 	 */
-	MODELINGCOMPONENTS_API const UInteractiveToolStorableSelection* GetCurrentToolInputSelection(const FToolBuilderState& SceneState, UToolTarget* Target);
+	MODELINGCOMPONENTS_API const UPersistentMeshSelection* GetCurrentToolInputSelection(const FToolBuilderState& SceneState, UToolTarget* Target);
 
 	/**
 	 * Convert the given Seletion to a list of Triangles of the specified Mesh. 
-	 * This will check the known variants of UInteractiveToolStorableSelection and convert their various selection types
-	 * (eg group topology, UV island, etc) to triangle indices
+	 * This will attempt to convert the different selection types represented by UPersistentMeshSelection
+	 * (eg group topology, UV island, etc) to triangle indices in TrianglesOut
 	 * @return true if a list of triangles could be created.
 	 */
 	MODELINGCOMPONENTS_API bool GetStoredSelectionAsTriangles(
-		const UInteractiveToolStorableSelection* Selection, 
+		const UPersistentMeshSelection* Selection,
 		const FDynamicMesh3& Mesh,
 		TArray<int32>& TrianglesOut );
 
+	/**
+	 * Set the active output selection of a Tool, which (if supported) will be persist between Tools
+	 * via a UPersistentMeshSelectionManager context object. Passing a null Selection will clear any active selection.
+	 * @return true if the selection was saved
+	 */
+	MODELINGCOMPONENTS_API bool SetToolOutputSelection(
+		UInteractiveTool* Tool,
+		UPersistentMeshSelection* Selection);
+
+	/**
+	 * Clear any currently-active stored Tool selection
+	 */
+	MODELINGCOMPONENTS_API void ClearActiveToolSelection(UInteractiveToolManager* ToolManager);
+
+
+	//
+	// Utility functions for group topology manipulation.
+	// These should probably move to another location
+	//
+
+
+	/**
+	 * Returns a pair of vertex ID's that are representative of a group edge, to be able to identify
+	 * a selected group edge independently of a group topology object. 
+	 *
+	 * For non-loop group edges, this will be the vids of the lower-vid endpoint and its neighbor 
+	 * in the group edge, arranged in increasing vid order. For loop group edges, this will be the 
+	 * lowest vid in the group edge and its lower-vid neighbor in the group edge.
+	 *
+	 * This is basically just identifying the group edge with a specific component edge, but using
+	 * vids instead of an edge id makes it a bit easier to apply a compact mapping if the mesh was
+	 * compacted, and helps the representative survive translation to/from a mesh description in
+	 * cases of compaction.
+	 *
+	 * @param CompactMaps This gets used to remap the vids given in the topology first (i.e., it assumes 
+	 *  that the CompactMaps have not yet been applied to the contents of the topology object).
+	 */
+	MODELINGCOMPONENTS_API FIndex2i GetGroupEdgeRepresentativeVerts(const FGroupTopology& TopologyIn, int GroupEdgeID, const FCompactMaps& CompactMaps);
+
+	/**
+	 * Returns a pair of vertex ID's that are representative of a group edge, to be able to identify
+	 * a selected group edge independently of a group topology object. 
+	 *
+	 * See other overload for more details.
+	 */
+	MODELINGCOMPONENTS_API FIndex2i GetGroupEdgeRepresentativeVerts(const FGroupTopology& TopologyIn, int GroupEdgeID);
+
 
 }
 }
\ No newline at end of file
diff --git a/Engine/Source/Editor/Experimental/EditorInteractiveToolsFramework/Private/EdModeInteractiveToolsContext.cpp b/Engine/Source/Editor/Experimental/EditorInteractiveToolsFramework/Private/EdModeInteractiveToolsContext.cpp
--- a/Engine/Source/Editor/Experimental/EditorInteractiveToolsFramework/Private/EdModeInteractiveToolsContext.cpp
+++ b/Engine/Source/Editor/Experimental/EditorInteractiveToolsFramework/Private/EdModeInteractiveToolsContext.cpp
@@ -27,7 +27,6 @@
 
 #include "ToolContextInterfaces.h"
 #include "InteractiveToolObjects.h"
-#include "InteractiveToolsSelectionStoreSubsystem.h"
 #include "EditorInteractiveGizmoManager.h"
 #include "BaseBehaviors/ClickDragBehavior.h"
 #include "EditorModeManager.h"
@@ -191,15 +190,6 @@ public:
 		EditorModeManager->GetSelectedActors()->GetSelectedObjects(StateOut.SelectedActors);
 		EditorModeManager->GetSelectedComponents()->GetSelectedObjects(StateOut.SelectedComponents);
 		StateOut.TypedElementSelectionSet = EditorModeManager->GetEditorSelectionSet();
-
-		UInteractiveToolsSelectionStoreSubsystem* ToolSelectionStore = GEngine->GetEngineSubsystem<UInteractiveToolsSelectionStoreSubsystem>();
-		ensureMsgf(ToolSelectionStore, TEXT("UInteractiveToolsSelectionStoreSubsystem was null, so tool selections will not be loaded. "
-			"Was the owning module loaded?"));
-
-		// If we someday have tool selections be tool manager specific, we can pass in the tool manager
-		// as a parameter here.
-		StateOut.StoredToolSelection = 
-			(ToolSelectionStore ? ToolSelectionStore->GetStoredSelection() : nullptr);
 	}
 
 	virtual void GetCurrentViewState(FViewCameraState& StateOut) const override
@@ -516,21 +506,6 @@ public:
 		return true;
 	}
 
-	virtual bool RequestToolSelectionStore(const UInteractiveToolStorableSelection* StorableSelection,
-		const FToolSelectionStoreParams & = FToolSelectionStoreParams()) override
-	{
-		if (UInteractiveToolsSelectionStoreSubsystem* ToolSelectionStore = GEngine->GetEngineSubsystem<UInteractiveToolsSelectionStoreSubsystem>())
-		{
-			ToolSelectionStore->Modify(false); // Allows this to be part of an undo transaction if one is open.
-			ToolSelectionStore->SetStoredSelection(StorableSelection);
-			return true;
-		}
-
-		ensureMsgf(false, TEXT("UInteractiveToolsSelectionStoreSubsystem was null, so tool selections will not be stored. "
-			"Was the owning module loaded?"));
-		return false;
-	}
-
 protected:
 	UEdModeInteractiveToolsContext* ToolsContext;
 	FEditorModeTools* EditorModeManager;
diff --git a/Engine/Source/Runtime/InteractiveToolsFramework/Private/InteractiveToolManager.cpp b/Engine/Source/Runtime/InteractiveToolsFramework/Private/InteractiveToolManager.cpp
--- a/Engine/Source/Runtime/InteractiveToolsFramework/Private/InteractiveToolManager.cpp
+++ b/Engine/Source/Runtime/InteractiveToolsFramework/Private/InteractiveToolManager.cpp
@@ -256,21 +256,6 @@ void UInteractiveToolManager::DeactivateToolInternal(EToolSide Side, EToolShutdo
 
 		ActiveLeftTool->Shutdown(ShutdownType);
 
-		// Every tool invocation that doesn't do anything with the tool selection should result in it being
-		// cleared, but cancelled tools don't count.
-		if (ShutdownType != EToolShutdownType::Cancel && !bActiveToolMadeSelectionStoreRequest)
-		{
-			if (UInteractiveToolsSelectionStoreSubsystem* ToolSelectionStore = GEngine->GetEngineSubsystem<UInteractiveToolsSelectionStoreSubsystem>())
-			{
-				// Note that it would be better if the tool cleared the selection store in its Shutdown()
-				// and bundled it with any undo transaction it issued. We can't do that here, so this ends
-				// up being a non-undoable clear. Not ideal, but acceptible.
-				IToolsContextTransactionsAPI::FToolSelectionStoreParams Params;
-				Params.ToolManager = this;
-				TransactionsAPI->RequestToolSelectionStore(nullptr, Params);
-			}
-		}
-
 		InputRouter->DeregisterSource(ActiveLeftTool);
 
 		UInteractiveTool* DoneTool = ActiveLeftTool;
@@ -449,17 +434,6 @@ bool UInteractiveToolManager::RequestSelectionChange(const FSelectedOjectsChange
 	return TransactionsAPI->RequestSelectionChange(SelectionChange);
 }
 
-bool UInteractiveToolManager::RequestToolSelectionStore(const UInteractiveToolStorableSelection* StorableSelection)
-{
-	bActiveToolMadeSelectionStoreRequest = true;
-
-	// We don't actually use this information right now, because the tool selection object is global. But
-	// if we change selections to be tool manager specific later, this will make it easier.
-	IToolsContextTransactionsAPI::FToolSelectionStoreParams Params;
-	Params.ToolManager = this;
-
-	return TransactionsAPI->RequestToolSelectionStore(StorableSelection, Params);
-}
 
 
 
diff --git a/Engine/Source/Runtime/InteractiveToolsFramework/Public/InteractiveToolManager.h b/Engine/Source/Runtime/InteractiveToolsFramework/Public/InteractiveToolManager.h
--- a/Engine/Source/Runtime/InteractiveToolsFramework/Public/InteractiveToolManager.h
+++ b/Engine/Source/Runtime/InteractiveToolsFramework/Public/InteractiveToolManager.h
@@ -225,19 +225,6 @@ public:
 	 */
 	virtual bool RequestSelectionChange(const FSelectedOjectsChangeList& SelectionChange);
 
-	/**
-	 * Store an arbitrary tool selection object that will be loaded for the next tool to be invoked.
-	 * If that tool doesn't support the selection object and does not get cancelled, the stored
-	 * selection will be discarded. Passing in a nullptr will clear the stored selection.
-	 *
-	 * Ideally, a tool would make this call in any Shutdown() non-cancel invocation, even if to
-	 * keep the currently stored tool selection. If a tool doesn't, the tool manager will discard
-	 * the last stored tool selection to avoid keeping stale ones across tool invocations, but this
-	 * will be a non-undoable clear because the tool manager can't put the action into whatever
-	 * transaction the tool issues in its Shutdown().
-	 */
-	virtual bool RequestToolSelectionStore(const UInteractiveToolStorableSelection* StorableSelection);
-
 	//
 	// State control  (@todo: have the Context call these? not safe for anyone to call)
 	//
diff --git a/Engine/Source/Runtime/InteractiveToolsFramework/Public/ToolContextInterfaces.h b/Engine/Source/Runtime/InteractiveToolsFramework/Public/ToolContextInterfaces.h
--- a/Engine/Source/Runtime/InteractiveToolsFramework/Public/ToolContextInterfaces.h
+++ b/Engine/Source/Runtime/InteractiveToolsFramework/Public/ToolContextInterfaces.h
@@ -20,7 +20,6 @@ class UInteractiveToolManager;
 class UInteractiveGizmoManager;
 class UToolTargetManager;
 struct FMeshDescription;
-class UInteractiveToolStorableSelection;
 struct FTypedElementHandle;
 class UTypedElementSelectionSet;
 
@@ -45,12 +44,6 @@ struct INTERACTIVETOOLSFRAMEWORK_API FToolBuilderState
 	/** Current selected Components. May be empty or nullptr. */
 	TArray<UActorComponent*> SelectedComponents;
 
-	/** 
-	 * A tool selection object stored by a previous tool, containing, for instance,
-	 * selected groups within a mesh, etc.
-	 */
-	const UInteractiveToolStorableSelection* StoredToolSelection = nullptr;
-
 	TWeakObjectPtr<UTypedElementSelectionSet> TypedElementSelectionSet;
 };
 
@@ -358,27 +351,6 @@ public:
 	 */
 	virtual bool RequestSelectionChange(const FSelectedOjectsChangeList& SelectionChange) = 0;
 
-
-	// Currently, no additional parameters are needed for tool selection storage since there
-	// is one global stored object for the whole editor at a time. That may be what we want.
-	// However, in case we someday decide that we want the selection object to be specific to
-	// a tool manager, etc, we keep this struct as a way to add parameters. For now the only
-	// parameter (unused) is the requesting tool manager.
-	struct FToolSelectionStoreParams
-	{
-		UInteractiveToolManager* ToolManager = nullptr;
-	};
-
-	/**
-	 * Request an arbitrary tool selection object to be stored so that it can be retrieved by other tools,
-	 * modes, or asset editors. Only one of these is active at a time, so this will clear the previous
-	 * stored tool selection when used.
-	 *
-	 * @param StorableSelection Selection object to hold on to. Can be nullptr to clear the stored selection.
-	 * @param Params Not currently used, here to make the selection store system easier to tweak later.
-	 */
-	virtual bool RequestToolSelectionStore(const UInteractiveToolStorableSelection* StorableSelection, 
-		const FToolSelectionStoreParams& Params) = 0;
 };
 
 UENUM()
