UV Editor: Support remove and toggle from selection in UV Select tool

#rnx

diff --git a/Engine/Plugins/Experimental/MeshModelingToolsetExp/Source/ModelingComponents/Private/Selection/MeshSelectionMechanic.cpp b/Engine/Plugins/Experimental/MeshModelingToolsetExp/Source/ModelingComponents/Private/Selection/MeshSelectionMechanic.cpp
--- a/Engine/Plugins/Experimental/MeshModelingToolsetExp/Source/ModelingComponents/Private/Selection/MeshSelectionMechanic.cpp
+++ b/Engine/Plugins/Experimental/MeshModelingToolsetExp/Source/ModelingComponents/Private/Selection/MeshSelectionMechanic.cpp
@@ -64,14 +64,23 @@ namespace MeshSelectionMechanicLocals
 		FDynamicMeshSelection After;
 		bool bBroadcastOnSelectionChanged;
 	};
+
+	template <typename InElementType>
+	void ToggleItem(TSet<InElementType>& Set, InElementType Item)
+	{
+		if (Set.Remove(Item) == 0)
+		{
+			Set.Add(Item);
+		}
+	}
 }
 
 void UMeshSelectionMechanic::Setup(UInteractiveTool* ParentToolIn)
 {
 	UInteractionMechanic::Setup(ParentToolIn);
 
-	// TODO: Add shift/ctrl modifiers
 	USingleClickInputBehavior* ClickBehavior = NewObject<USingleClickInputBehavior>();
+	ClickBehavior->Modifiers.RegisterModifier(CtrlModifierID, FInputDeviceState::IsCtrlKeyDown);
 	ClickBehavior->Modifiers.RegisterModifier(ShiftModifierID, FInputDeviceState::IsShiftKeyDown);
 	ClickBehavior->Initialize(this);
 	ParentTool->AddInputBehavior(ClickBehavior);
@@ -340,13 +349,14 @@ void UMeshSelectionMechanic::OnClicked(const FInputDeviceRay& ClickPos)
 {
 	FDynamicMeshSelection OriginalSelection = CurrentSelection;
 
-	if (!InMultiSelectMode() ||
+	const bool bIncompatibleSelectionMode =
 		(SelectionMode == EMeshSelectionMechanicMode::Component && CurrentSelection.Type != FDynamicMeshSelection::EType::Triangle) ||
 		(SelectionMode == EMeshSelectionMechanicMode::Edge && CurrentSelection.Type != FDynamicMeshSelection::EType::Edge) ||
 		(SelectionMode == EMeshSelectionMechanicMode::Vertex && CurrentSelection.Type != FDynamicMeshSelection::EType::Vertex) || 
 		(SelectionMode == EMeshSelectionMechanicMode::Triangle && CurrentSelection.Type != FDynamicMeshSelection::EType::Triangle) || 
-		(SelectionMode == EMeshSelectionMechanicMode::Mesh && CurrentSelection.Type != FDynamicMeshSelection::EType::Triangle) )
-	{ // If we're not enabling multiselect, or our desired mode doesn't match our current selection type, clear the existing selection
+		(SelectionMode == EMeshSelectionMechanicMode::Mesh && CurrentSelection.Type != FDynamicMeshSelection::EType::Triangle);
+	if (bIncompatibleSelectionMode || ShouldRestartSelection())
+	{
 		CurrentSelection.SelectedIDs.Empty();
 		CurrentSelection.Mesh = nullptr;
 	}
@@ -355,10 +365,10 @@ void UMeshSelectionMechanic::OnClicked(const FInputDeviceRay& ClickPos)
 	double RayT = 0; 
 	for (int32 i = 0; i < MeshSpatials.Num(); ++i)
 	{
-		//Short circuit the loop if we're in multiselect mode, to save us from testing things we aren't allowed to select
-		if (InMultiSelectMode() && !CurrentSelection.IsEmpty() 
-			&& CurrentSelectionIndex != IndexConstants::InvalidID 
-			&& i != CurrentSelectionIndex)
+		// We only add/toggle/remove from selections in the currently selected mesh. If this isn't the case or if we're
+		// removing from a selection but there is nothing to remove, then we can break early.
+		if ((!ShouldRestartSelection() && i != CurrentSelectionIndex) ||
+			(ShouldRemoveFromSelection() && CurrentSelection.IsEmpty()))
 		{
 			continue;
 		}
@@ -387,10 +397,21 @@ void UMeshSelectionMechanic::OnClicked(const FInputDeviceRay& ClickPos)
 			TArray<int32> SeedTriangles;
 			SeedTriangles.Add(HitTid);
 			MeshSelectedComponent.FindTrianglesConnectedToSeeds(SeedTriangles);
-			// Since we've already handled the situation of not holding triangles in a multiselect above, this should be safe.
-			// But just to be sure, we'll toss in a check to validate everything.
-			ensure((InMultiSelectMode() && CurrentSelection.Type == FDynamicMeshSelection::EType::Triangle) || CurrentSelection.SelectedIDs.IsEmpty());
-			CurrentSelection.SelectedIDs.Append( TSet(MeshSelectedComponent.Components[0].Indices) );
+			
+			if (ShouldAddToSelection() || ShouldRestartSelection())
+			{
+				CurrentSelection.SelectedIDs.Append(TSet<int>(MeshSelectedComponent.Components[0].Indices));
+			}
+			else if (ShouldRemoveFromSelection())
+			{
+				CurrentSelection.SelectedIDs = CurrentSelection.SelectedIDs.Difference(TSet<int>(MeshSelectedComponent.Components[0].Indices));
+			}
+			else if (ShouldToggleFromSelection())
+			{
+				for (int I : MeshSelectedComponent.Components[0].Indices) {
+					MeshSelectionMechanicLocals::ToggleItem(CurrentSelection.SelectedIDs, I);
+				}
+			}
 			CurrentSelection.Type = FDynamicMeshSelection::EType::Triangle;
 		}
 		// TODO: We'll need the ability to hit occluded triangles to see if there is a better edge
@@ -414,10 +435,18 @@ void UMeshSelectionMechanic::OnClicked(const FInputDeviceRay& ClickPos)
 						Position1, Position2,
 						ToolSceneQueriesUtil::GetDefaultVisualAngleSnapThreshD()); }))
 			{
-				// Since we've already handled multiselect and not holding edges above, this should be safe.	
-				// But just to be sure, we'll toss in a check to validate everything.
-				ensure((InMultiSelectMode() && CurrentSelection.Type == FDynamicMeshSelection::EType::Edge) || CurrentSelection.SelectedIDs.IsEmpty());
-				CurrentSelection.SelectedIDs.Add(Result.ID);
+				if (ShouldAddToSelection() || ShouldRestartSelection())
+				{
+					CurrentSelection.SelectedIDs.Add(Result.ID);
+				}
+				else if (ShouldRemoveFromSelection())
+				{
+					CurrentSelection.SelectedIDs.Remove(Result.ID);
+				}
+				else if (ShouldToggleFromSelection())
+				{
+					MeshSelectionMechanicLocals::ToggleItem(CurrentSelection.SelectedIDs, Result.ID);
+				}
 				CurrentSelection.Type = FDynamicMeshSelection::EType::Edge;
 			}
 		}
@@ -440,25 +469,53 @@ void UMeshSelectionMechanic::OnClicked(const FInputDeviceRay& ClickPos)
 						Position1, Position2,
 						ToolSceneQueriesUtil::GetDefaultVisualAngleSnapThreshD()); }))
 			{
-				// Since we've already handled multiselect and not holding vertices above, this should be safe.	
-				// But just to be sure, we'll toss in a check to validate everything.
-				ensure((InMultiSelectMode() && CurrentSelection.Type == FDynamicMeshSelection::EType::Vertex) || CurrentSelection.SelectedIDs.IsEmpty());
-				CurrentSelection.SelectedIDs.Add(Result.ID);
+				if (ShouldAddToSelection() || ShouldRestartSelection())
+				{
+					CurrentSelection.SelectedIDs.Add(Result.ID);
+				}
+				else if (ShouldRemoveFromSelection())
+				{
+					CurrentSelection.SelectedIDs.Remove(Result.ID);
+				}
+				else if (ShouldToggleFromSelection())
+				{
+					MeshSelectionMechanicLocals::ToggleItem(CurrentSelection.SelectedIDs, Result.ID);
+				}
 				CurrentSelection.Type = FDynamicMeshSelection::EType::Vertex;
 			}
 		}
 		else if (SelectionMode == EMeshSelectionMechanicMode::Triangle)
 		{
-			ensure((InMultiSelectMode() && CurrentSelection.Type == FDynamicMeshSelection::EType::Triangle) || CurrentSelection.SelectedIDs.IsEmpty());
-			CurrentSelection.SelectedIDs.Add(HitTid);
+			if (ShouldAddToSelection() || ShouldRestartSelection())
+			{
+				CurrentSelection.SelectedIDs.Add(HitTid);
+			}
+			else if (ShouldRemoveFromSelection())
+			{
+				CurrentSelection.SelectedIDs.Remove(HitTid);
+			}
+			else if (ShouldToggleFromSelection())
+			{
+				MeshSelectionMechanicLocals::ToggleItem(CurrentSelection.SelectedIDs, HitTid);
+			}
 			CurrentSelection.Type = FDynamicMeshSelection::EType::Triangle;
 		}
 		else if (SelectionMode == EMeshSelectionMechanicMode::Mesh)
 		{
-			ensure((InMultiSelectMode() && CurrentSelection.Type == FDynamicMeshSelection::EType::Triangle) || CurrentSelection.SelectedIDs.IsEmpty());
 			for (int32 Tid : CurrentSelection.Mesh->TriangleIndicesItr())
 			{
-				CurrentSelection.SelectedIDs.Add( Tid );
+				if (ShouldAddToSelection() || ShouldRestartSelection())
+				{
+					CurrentSelection.SelectedIDs.Add(Tid);
+				}
+				else if (ShouldRemoveFromSelection())
+				{
+					CurrentSelection.SelectedIDs.Remove(Tid);
+				}
+				else if (ShouldToggleFromSelection())
+				{
+					MeshSelectionMechanicLocals::ToggleItem(CurrentSelection.SelectedIDs, Tid);
+				}
 			}
 			CurrentSelection.Type = FDynamicMeshSelection::EType::Triangle;
 		}
@@ -480,7 +537,9 @@ void UMeshSelectionMechanic::OnUpdateModifierState(int ModifierID, bool bIsOn)
 	case ShiftModifierID:
 		bShiftToggle = bIsOn;
 		break;
-	// Add more modifiers here, if needed.
+	case CtrlModifierID:
+		bCtrlToggle = bIsOn;
+		break;
 	}
 }
 
diff --git a/Engine/Plugins/Experimental/MeshModelingToolsetExp/Source/ModelingComponents/Public/Selection/MeshSelectionMechanic.h b/Engine/Plugins/Experimental/MeshModelingToolsetExp/Source/ModelingComponents/Public/Selection/MeshSelectionMechanic.h
--- a/Engine/Plugins/Experimental/MeshModelingToolsetExp/Source/ModelingComponents/Public/Selection/MeshSelectionMechanic.h
+++ b/Engine/Plugins/Experimental/MeshModelingToolsetExp/Source/ModelingComponents/Public/Selection/MeshSelectionMechanic.h
@@ -103,7 +103,11 @@ public:
 
 protected:
 
-	bool InMultiSelectMode() const { return bShiftToggle; }
+	// All four combinations of shift/ctrl down are assigned a behaviour
+	bool ShouldAddToSelection() const { return !bCtrlToggle && bShiftToggle; }
+	bool ShouldRemoveFromSelection() const { return bCtrlToggle && !bShiftToggle; }
+	bool ShouldToggleFromSelection() const { return bCtrlToggle && bShiftToggle; }
+	bool ShouldRestartSelection() const { return !bCtrlToggle && !bShiftToggle; }
 
 	UPROPERTY()
 	TObjectPtr<APreviewGeometryActor> PreviewGeometryActor = nullptr;
@@ -125,7 +129,9 @@ protected:
 	FColor PointColor = FColor::Yellow;
 
 	bool bShiftToggle = false;
+	bool bCtrlToggle = false;
 	static const int32 ShiftModifierID = 1;
+	static const int32 CtrlModifierID = 2;
 
 	float LineThickness = 3;
 	float PointThickness = 6;
