ModelingTools: fix incorrect string passed to SaveProperties in ParameterizeMeshTool, and handle case in ParameterizeMeshOp where target UV layer does not exist.
#rnx

diff --git a/Engine/Plugins/Experimental/MeshModelingToolsetExp/Source/MeshModelingToolsEditorOnly/Private/ParameterizeMeshTool.cpp b/Engine/Plugins/Experimental/MeshModelingToolsetExp/Source/MeshModelingToolsEditorOnly/Private/ParameterizeMeshTool.cpp
--- a/Engine/Plugins/Experimental/MeshModelingToolsetExp/Source/MeshModelingToolsEditorOnly/Private/ParameterizeMeshTool.cpp
+++ b/Engine/Plugins/Experimental/MeshModelingToolsetExp/Source/MeshModelingToolsEditorOnly/Private/ParameterizeMeshTool.cpp
@@ -144,7 +144,7 @@ void UParameterizeMeshTool::Shutdown(EToolShutdownType ShutdownType)
 {
 	if (UVLayoutView)
 	{
-		UVLayoutView->Settings->SaveProperties(this, TEXT("RecomputeUVsTool"));
+		UVLayoutView->Settings->SaveProperties(this, TEXT("ParameterizeMeshTool"));
 		UVLayoutView->Disconnect();
 	}
 
diff --git a/Engine/Plugins/Experimental/MeshModelingToolsetExp/Source/ModelingOperatorsEditorOnly/Private/ParameterizationOps/ParameterizeMeshOp.cpp b/Engine/Plugins/Experimental/MeshModelingToolsetExp/Source/ModelingOperatorsEditorOnly/Private/ParameterizationOps/ParameterizeMeshOp.cpp
--- a/Engine/Plugins/Experimental/MeshModelingToolsetExp/Source/ModelingOperatorsEditorOnly/Private/ParameterizationOps/ParameterizeMeshOp.cpp
+++ b/Engine/Plugins/Experimental/MeshModelingToolsetExp/Source/ModelingOperatorsEditorOnly/Private/ParameterizationOps/ParameterizeMeshOp.cpp
@@ -165,24 +165,19 @@ void FParameterizeMeshOp::CopyNewUVsToMesh(
 	const TArray<int32>& VertexRemapArray,
 	bool bReverseOrientation)
 {
-	// Add the UVs to the FDynamicMesh
-	const bool bHasAttributes = Mesh.HasAttributes();
-
-	if (bHasAttributes)
-	{
-		FDynamicMeshAttributeSet* Attributes = Mesh.Attributes();
-		Attributes->GetUVLayer(UVLayer)->ClearElements(); // delete existing UVs
-	}
-	else
+	if (!Mesh.HasAttributes())
 	{
-		// Add attrs for UVS
 		Mesh.EnableAttributes();
 	}
-
 	FDynamicMeshUVOverlay* UVOverlay = Mesh.Attributes()->GetUVLayer(UVLayer);
+	if (UVOverlay == nullptr)
+	{
+		Mesh.Attributes()->SetNumUVLayers(UVLayer+1);
+		UVOverlay = Mesh.Attributes()->GetUVLayer(UVLayer);
+	}
 
-	// This mesh shouldn't already have UVs.
-
+	// delete any existing UVs
+	UVOverlay->ClearElements();
 	checkSlow(UVOverlay->ElementCount() == 0);
 
 	// Add the UVs to the overlay
