ModelingTools: expose WindingThreshold option in BooleanMeshesOp and CutMeshWithMeshTool
#rnx

diff --git a/Engine/Plugins/Experimental/MeshModelingToolsetExp/Source/MeshModelingToolsExp/Private/CutMeshWithMeshTool.cpp b/Engine/Plugins/Experimental/MeshModelingToolsetExp/Source/MeshModelingToolsExp/Private/CutMeshWithMeshTool.cpp
--- a/Engine/Plugins/Experimental/MeshModelingToolsetExp/Source/MeshModelingToolsExp/Private/CutMeshWithMeshTool.cpp
+++ b/Engine/Plugins/Experimental/MeshModelingToolsetExp/Source/MeshModelingToolsExp/Private/CutMeshWithMeshTool.cpp
@@ -165,6 +165,7 @@ public:
 
 	bool bAttemptToFixHoles = true;
 	bool bCollapseExtraEdges = true;
+	double WindingThreshold = 0.5;
 
 	virtual void CalculateResult(FProgressCancel* Progress)
 	{
@@ -172,6 +173,7 @@ public:
 		SubtractOp->CSGOperation = ECSGOperation::DifferenceAB;
 		SubtractOp->bAttemptFixHoles = bAttemptToFixHoles;
 		SubtractOp->bTryCollapseExtraEdges = bCollapseExtraEdges;
+		SubtractOp->WindingThreshold = WindingThreshold;
 		SubtractOp->Meshes.Add(TargetMesh);
 		SubtractOp->Transforms.Add(TargetMeshTransform);
 		SubtractOp->Meshes.Add(CuttingMesh);
@@ -181,6 +183,7 @@ public:
 		IntersectOp->CSGOperation = ECSGOperation::Intersect;
 		IntersectOp->bAttemptFixHoles = bAttemptToFixHoles;
 		IntersectOp->bTryCollapseExtraEdges = bCollapseExtraEdges;
+		IntersectOp->WindingThreshold = WindingThreshold;
 		IntersectOp->Meshes.Add(TargetMesh);
 		IntersectOp->Transforms.Add(TargetMeshTransform);
 		IntersectOp->Meshes.Add(CuttingMesh);
@@ -283,6 +286,7 @@ TUniquePtr<FDynamicMeshOperator> UCutMeshWithMeshTool::MakeNewOperator()
 
 	CuttingOp->bAttemptToFixHoles = CutProperties->bAttemptFixHoles;
 	CuttingOp->bCollapseExtraEdges = CutProperties->bCollapseExtraEdges;
+	CuttingOp->WindingThreshold = CutProperties->WindingThreshold;
 
 	return CuttingOp;
 }
diff --git a/Engine/Plugins/Experimental/MeshModelingToolsetExp/Source/MeshModelingToolsExp/Public/CutMeshWithMeshTool.h b/Engine/Plugins/Experimental/MeshModelingToolsetExp/Source/MeshModelingToolsExp/Public/CutMeshWithMeshTool.h
--- a/Engine/Plugins/Experimental/MeshModelingToolsetExp/Source/MeshModelingToolsExp/Public/CutMeshWithMeshTool.h
+++ b/Engine/Plugins/Experimental/MeshModelingToolsetExp/Source/MeshModelingToolsExp/Public/CutMeshWithMeshTool.h
@@ -23,6 +23,10 @@ class MESHMODELINGTOOLSEXP_API UCutMeshWithMeshToolProperties : public UInteract
 {
 	GENERATED_BODY()
 public:
+	/** Automatically attempt to fill any holes left by Booleans (e.g. due to numerical errors) */
+	UPROPERTY(EditAnywhere, Category = Options, meta = (UIMin = "0", UIMax = "1"))
+	float WindingThreshold = 0.5;
+
 	/** Show boundary edges created by the Boolean operations (often due to numerical error) */
 	UPROPERTY(EditAnywhere, Category = Options)
 	bool bShowNewBoundaryEdges = true;
diff --git a/Engine/Plugins/Experimental/MeshModelingToolsetExp/Source/ModelingOperators/Private/CompositionOps/BooleanMeshesOp.cpp b/Engine/Plugins/Experimental/MeshModelingToolsetExp/Source/ModelingOperators/Private/CompositionOps/BooleanMeshesOp.cpp
--- a/Engine/Plugins/Experimental/MeshModelingToolsetExp/Source/ModelingOperators/Private/CompositionOps/BooleanMeshesOp.cpp
+++ b/Engine/Plugins/Experimental/MeshModelingToolsetExp/Source/ModelingOperators/Private/CompositionOps/BooleanMeshesOp.cpp
@@ -74,6 +74,7 @@ void FBooleanMeshesOp::CalculateResult(FProgressCancel* Progress)
 
 	MeshBoolean.bPutResultInInputSpace = false;
 	MeshBoolean.bSimplifyAlongNewEdges = bTryCollapseExtraEdges;
+	MeshBoolean.WindingThreshold = WindingThreshold;
 	MeshBoolean.Progress = Progress;
 	bool bSuccess = MeshBoolean.Compute();
 	ResultTransform = MeshBoolean.ResultTransform;
diff --git a/Engine/Plugins/Experimental/MeshModelingToolsetExp/Source/ModelingOperators/Public/CompositionOps/BooleanMeshesOp.h b/Engine/Plugins/Experimental/MeshModelingToolsetExp/Source/ModelingOperators/Public/CompositionOps/BooleanMeshesOp.h
--- a/Engine/Plugins/Experimental/MeshModelingToolsetExp/Source/ModelingOperators/Public/CompositionOps/BooleanMeshesOp.h
+++ b/Engine/Plugins/Experimental/MeshModelingToolsetExp/Source/ModelingOperators/Public/CompositionOps/BooleanMeshesOp.h
@@ -64,6 +64,7 @@ public:
 	TArray<TSharedPtr<const FDynamicMesh3, ESPMode::ThreadSafe>> Meshes;
 	TArray<FTransform> Transforms; // 1:1 with Meshes
 	bool bAttemptFixHoles = false;
+	double WindingThreshold = 0.5;
 
 	/** If true, try to do edge-collapses along cut edges to remove unnecessary edges inserted by cut */
 	bool bTryCollapseExtraEdges = false;
