MeshModelingTools: Fix bug causing PolygonSelectionMechanic to incorrectly allow selections containing both edges and corners simultaneously

#rnx

diff --git a/Engine/Plugins/Experimental/MeshModelingToolsetExp/Source/ModelingComponents/Private/Selection/PolygonSelectionMechanic.cpp b/Engine/Plugins/Experimental/MeshModelingToolsetExp/Source/ModelingComponents/Private/Selection/PolygonSelectionMechanic.cpp
--- a/Engine/Plugins/Experimental/MeshModelingToolsetExp/Source/ModelingComponents/Private/Selection/PolygonSelectionMechanic.cpp
+++ b/Engine/Plugins/Experimental/MeshModelingToolsetExp/Source/ModelingComponents/Private/Selection/PolygonSelectionMechanic.cpp
@@ -598,6 +598,9 @@ void UPolygonSelectionMechanic::OnDragRectangleStarted()
 	BeginChange();
 
 	PreDragPersistentSelection = PersistentSelection;
+	PreDragTopoSelectorSettings = GetTopoSelectorSettings(false);
+	PreDragTopoSelectorSettings.bIgnoreOcclusion = Properties->bMarqueeIgnoreOcclusion; // uses a separate setting for marquee
+	
 	PersistentSelection.Clear();
 }
 
@@ -605,10 +608,7 @@ void UPolygonSelectionMechanic::OnDragRectangleChanged(const FCameraRectangle& C
 {
 	FGroupTopologySelection RectangleSelection;
 
-	FGroupTopologySelector::FSelectionSettings TopoSelectorSettings = GetTopoSelectorSettings(false);
-	TopoSelectorSettings.bIgnoreOcclusion = Properties->bMarqueeIgnoreOcclusion; // uses a separate setting for marquee
-	
-	TopoSelector->FindSelectedElement(TopoSelectorSettings, CurrentRectangle, TargetTransform,
+	TopoSelector->FindSelectedElement(PreDragTopoSelectorSettings, CurrentRectangle, TargetTransform,
 		RectangleSelection);
 
 	if (ShouldAddToSelectionFunc())
diff --git a/Engine/Plugins/Experimental/MeshModelingToolsetExp/Source/ModelingComponents/Public/Selection/PolygonSelectionMechanic.h b/Engine/Plugins/Experimental/MeshModelingToolsetExp/Source/ModelingComponents/Public/Selection/PolygonSelectionMechanic.h
--- a/Engine/Plugins/Experimental/MeshModelingToolsetExp/Source/ModelingComponents/Public/Selection/PolygonSelectionMechanic.h
+++ b/Engine/Plugins/Experimental/MeshModelingToolsetExp/Source/ModelingComponents/Public/Selection/PolygonSelectionMechanic.h
@@ -376,7 +376,10 @@ protected:
 	int32 SelectionTimestamp = 0;
 	TUniquePtr<FPolygonSelectionMechanicSelectionChange> ActiveChange;
 
-	FGroupTopologySelection PreDragPersistentSelection; // Used for box selection
+	// Used for box selection
+	FGroupTopologySelection PreDragPersistentSelection;
+	FGroupTopologySelector::FSelectionSettings PreDragTopoSelectorSettings;
+
 	FVector3d LastClickedHitPosition;
 	FVector3d LastClickedHitNormal;
 
