ModelingTools: Add N-to-1 detail mesh bake tool
- Added new FindNearestHitTriangle overload with bary coords.
- Refactored BakeMeshAttributeMapsTool into common base class with new MultiMesh (N-to-1) variant.
- Extended MeshDescriptionAdapter to support UVs.
- Added new queries on IMeshSpatialWrapper for mesh attributes.

#rnx

diff --git a/Engine/Plugins/Editor/ModelingToolsEditorMode/Source/ModelingToolsEditorMode/Private/ModelingToolsEditorMode.cpp b/Engine/Plugins/Editor/ModelingToolsEditorMode/Source/ModelingToolsEditorMode/Private/ModelingToolsEditorMode.cpp
--- a/Engine/Plugins/Editor/ModelingToolsEditorMode/Source/ModelingToolsEditorMode/Private/ModelingToolsEditorMode.cpp
+++ b/Engine/Plugins/Editor/ModelingToolsEditorMode/Source/ModelingToolsEditorMode/Private/ModelingToolsEditorMode.cpp
@@ -76,6 +76,7 @@
 #include "AlignObjectsTool.h"
 #include "EditUVIslandsTool.h"
 #include "BakeMeshAttributeMapsTool.h"
+#include "BakeMultiMeshAttributeMapsTool.h"
 #include "BakeMeshAttributeVertexTool.h"
 #include "MeshAttributePaintTool.h"
 #include "ParameterizeMeshTool.h"
@@ -551,6 +552,9 @@ void UModelingToolsEditorMode::Enter()
 	auto BakeMeshAttributeMapsToolBuilder = NewObject<UBakeMeshAttributeMapsToolBuilder>();
 	RegisterTool(ToolManagerCommands.BeginBakeMeshAttributeMapsTool, TEXT("BeginBakeMeshAttributeMapsTool"), BakeMeshAttributeMapsToolBuilder);
 
+	auto BakeMultiMeshAttributeMapsToolBuilder = NewObject<UBakeMultiMeshAttributeMapsToolBuilder>();
+	RegisterTool(ToolManagerCommands.BeginBakeMultiMeshAttributeMapsTool, TEXT("BeginBakeMultiMeshAttributeMapsTool"), BakeMultiMeshAttributeMapsToolBuilder);
+
 	auto BakeMeshAttributeVertexToolBuilder = NewObject<UBakeMeshAttributeVertexToolBuilder>();
 	RegisterTool(ToolManagerCommands.BeginBakeMeshAttributeVertexTool, TEXT("BeginBakeMeshAttributeVertexTool"), BakeMeshAttributeVertexToolBuilder);
 
diff --git a/Engine/Plugins/Editor/ModelingToolsEditorMode/Source/ModelingToolsEditorMode/Private/ModelingToolsEditorModeStyle.cpp b/Engine/Plugins/Editor/ModelingToolsEditorMode/Source/ModelingToolsEditorMode/Private/ModelingToolsEditorModeStyle.cpp
--- a/Engine/Plugins/Editor/ModelingToolsEditorMode/Source/ModelingToolsEditorMode/Private/ModelingToolsEditorModeStyle.cpp
+++ b/Engine/Plugins/Editor/ModelingToolsEditorMode/Source/ModelingToolsEditorMode/Private/ModelingToolsEditorModeStyle.cpp
@@ -166,6 +166,8 @@ void FModelingToolsEditorModeStyle::Initialize()
 		StyleSet->Set("ModelingToolsManagerCommands.BeginUVSeamEditTool.Small", 		new IMAGE_PLUGIN_BRUSH("Icons/ModelingUVSeamEdit_x40",			Icon20x20));
 		StyleSet->Set("ModelingToolsManagerCommands.BeginBakeMeshAttributeMapsTool", 			new IMAGE_PLUGIN_BRUSH("Icons/ModelingBakeMaps_x40",			Icon20x20));
 		StyleSet->Set("ModelingToolsManagerCommands.BeginBakeMeshAttributeMapsTool.Small", 		new IMAGE_PLUGIN_BRUSH("Icons/ModelingBakeMaps_x40",			Icon20x20));
+		StyleSet->Set("ModelingToolsManagerCommands.BeginBakeMultiMeshAttributeMapsTool", 			new IMAGE_PLUGIN_BRUSH("Icons/ModelingBakeMaps_x40",			Icon20x20));
+		StyleSet->Set("ModelingToolsManagerCommands.BeginBakeMultiMeshAttributeMapsTool.Small", 		new IMAGE_PLUGIN_BRUSH("Icons/ModelingBakeMaps_x40",			Icon20x20));
 		StyleSet->Set("ModelingToolsManagerCommands.BeginBakeMeshAttributeVertexTool", new IMAGE_PLUGIN_BRUSH("Icons/ModelingBakeMaps_x40", Icon20x20));
 		StyleSet->Set("ModelingToolsManagerCommands.BeginBakeMeshAttributeVertexTool.Small", new IMAGE_PLUGIN_BRUSH("Icons/ModelingBakeMaps_x40", Icon20x20));
 		//StyleSet->Set("ModelingToolsManagerCommands.BeginRemoveOccludedTrianglesTool", 				new IMAGE_PLUGIN_BRUSH("Icons/Jacket_40x",			Icon20x20));
diff --git a/Engine/Plugins/Editor/ModelingToolsEditorMode/Source/ModelingToolsEditorMode/Private/ModelingToolsEditorModeToolkit.cpp b/Engine/Plugins/Editor/ModelingToolsEditorMode/Source/ModelingToolsEditorMode/Private/ModelingToolsEditorModeToolkit.cpp
--- a/Engine/Plugins/Editor/ModelingToolsEditorMode/Source/ModelingToolsEditorMode/Private/ModelingToolsEditorModeToolkit.cpp
+++ b/Engine/Plugins/Editor/ModelingToolsEditorMode/Source/ModelingToolsEditorMode/Private/ModelingToolsEditorModeToolkit.cpp
@@ -630,6 +630,7 @@ void FModelingToolsEditorModeToolkit::BuildToolPalette_Experimental(FName Palett
 		ToolbarBuilder.AddToolBarButton(Commands.BeginMeshAttributePaintTool);
 		ToolbarBuilder.AddToolBarButton(Commands.BeginEditMeshMaterialsTool);
 		ToolbarBuilder.AddToolBarButton(Commands.BeginBakeMeshAttributeMapsTool);
+		ToolbarBuilder.AddToolBarButton(Commands.BeginBakeMultiMeshAttributeMapsTool);
 		ToolbarBuilder.AddToolBarButton(Commands.BeginBakeMeshAttributeVertexTool);
 	}
 	else if (PaletteIndex == UVTabName)
diff --git a/Engine/Plugins/Editor/ModelingToolsEditorMode/Source/ModelingToolsEditorMode/Private/ModelingToolsManagerActions.cpp b/Engine/Plugins/Editor/ModelingToolsEditorMode/Source/ModelingToolsEditorMode/Private/ModelingToolsManagerActions.cpp
--- a/Engine/Plugins/Editor/ModelingToolsEditorMode/Source/ModelingToolsEditorMode/Private/ModelingToolsManagerActions.cpp
+++ b/Engine/Plugins/Editor/ModelingToolsEditorMode/Source/ModelingToolsEditorMode/Private/ModelingToolsManagerActions.cpp
@@ -101,6 +101,7 @@ void FModelingToolsManagerCommands::RegisterCommands()
 	UI_COMMAND(BeginMeshAttributePaintTool, "MapPnt", "Paint Attribute Maps", EUserInterfaceActionType::ToggleButton, FInputChord());
 	UI_COMMAND(BeginAttributeEditorTool, "AttrEd", "Edit/Configure Mesh Attributes", EUserInterfaceActionType::ToggleButton, FInputChord());
 	UI_COMMAND(BeginBakeMeshAttributeMapsTool, "BakeTx", "Bake Image Maps for Target Mesh (optionally from second Source Mesh)", EUserInterfaceActionType::ToggleButton, FInputChord());
+	UI_COMMAND(BeginBakeMultiMeshAttributeMapsTool, "BakeAll", "Bake Image Maps for Target Mesh from multiple Source/Detail meshes", EUserInterfaceActionType::ToggleButton, FInputChord());
 	UI_COMMAND(BeginBakeMeshAttributeVertexTool, "BakeVtx", "Bake Vertex Colors for Target Mesh (optionally from a second Source Mesh)", EUserInterfaceActionType::ToggleButton, FInputChord());
 	UI_COMMAND(BeginUVSeamEditTool, "SeamEd", "Add UV Seams to Mesh", EUserInterfaceActionType::ToggleButton, FInputChord());
 
diff --git a/Engine/Plugins/Editor/ModelingToolsEditorMode/Source/ModelingToolsEditorMode/Public/ModelingToolsManagerActions.h b/Engine/Plugins/Editor/ModelingToolsEditorMode/Source/ModelingToolsEditorMode/Public/ModelingToolsManagerActions.h
--- a/Engine/Plugins/Editor/ModelingToolsEditorMode/Source/ModelingToolsEditorMode/Public/ModelingToolsManagerActions.h
+++ b/Engine/Plugins/Editor/ModelingToolsEditorMode/Source/ModelingToolsEditorMode/Public/ModelingToolsManagerActions.h
@@ -96,6 +96,7 @@ public:
 	TSharedPtr<FUICommandInfo> BeginMeshAttributePaintTool;
 	TSharedPtr<FUICommandInfo> BeginAttributeEditorTool;
 	TSharedPtr<FUICommandInfo> BeginBakeMeshAttributeMapsTool;
+	TSharedPtr<FUICommandInfo> BeginBakeMultiMeshAttributeMapsTool;
 	TSharedPtr<FUICommandInfo> BeginBakeMeshAttributeVertexTool;
 	TSharedPtr<FUICommandInfo> BeginUVSeamEditTool;
 
diff --git a/Engine/Plugins/Experimental/MeshModelingToolsetExp/Source/MeshModelingToolsExp/Private/BakeMeshAttributeMapsTool.cpp b/Engine/Plugins/Experimental/MeshModelingToolsetExp/Source/MeshModelingToolsExp/Private/BakeMeshAttributeMapsTool.cpp
--- a/Engine/Plugins/Experimental/MeshModelingToolsetExp/Source/MeshModelingToolsExp/Private/BakeMeshAttributeMapsTool.cpp
+++ b/Engine/Plugins/Experimental/MeshModelingToolsetExp/Source/MeshModelingToolsExp/Private/BakeMeshAttributeMapsTool.cpp
@@ -38,30 +38,6 @@ using namespace UE::Geometry;
 
 #define LOCTEXT_NAMESPACE "UBakeMeshAttributeMapsTool"
 
-/*
- * Static init
- */
-
-// Only include the Occlusion bitmask rather than its components
-// (AmbientOcclusion | BentNormal). Since the Occlusion baker can
-// bake both types in a single pass, only iterating over the Occlusion
-// bitmask gives direct access to both types without the need to
-// externally track if we've handled the Occlusion evaluator in a prior
-// iteration loop.
-static constexpr EBakeMapType ALL_BAKE_MAP_TYPES[] =
-{
-	EBakeMapType::TangentSpaceNormalMap,
-	EBakeMapType::Occlusion, // (AmbientOcclusion | BentNormal)
-	EBakeMapType::Curvature,
-	EBakeMapType::Texture2DImage,
-	EBakeMapType::NormalImage,
-	EBakeMapType::FaceNormalImage,
-	EBakeMapType::PositionImage,
-	EBakeMapType::MaterialID,
-	EBakeMapType::MultiTexture,
-	EBakeMapType::VertexColorImage
-};
-
 
 /*
  * ToolBuilder
@@ -284,48 +260,11 @@ public:
 
 void UBakeMeshAttributeMapsTool::Setup()
 {
-	UInteractiveTool::Setup();
-
-	// Setup preview materials
-	InitializeEmptyMaps();
-
-	UMaterial* Material = LoadObject<UMaterial>(nullptr, TEXT("/MeshModelingToolsetExp/Materials/BakePreviewMaterial"));
-	check(Material);
-	if (Material != nullptr)
-	{
-		PreviewMaterial = UMaterialInstanceDynamic::Create(Material, GetToolManager());
-		PreviewMaterial->SetTextureParameterValue(TEXT("NormalMap"), EmptyNormalMap);
-		PreviewMaterial->SetTextureParameterValue(TEXT("OcclusionMap"), EmptyColorMapWhite);
-		PreviewMaterial->SetTextureParameterValue(TEXT("ColorMap"), EmptyColorMapWhite);
-	}
-	UMaterial* BentNormalMaterial = LoadObject<UMaterial>(nullptr, TEXT("/MeshModelingToolsetExp/Materials/BakeBentNormalPreviewMaterial"));
-	check(BentNormalMaterial);
-	if (BentNormalMaterial != nullptr)
-	{
-		BentNormalPreviewMaterial = UMaterialInstanceDynamic::Create(BentNormalMaterial, GetToolManager());
-	}
-	UMaterial* WorkingMaterial = LoadObject<UMaterial>(nullptr, TEXT("/MeshModelingToolsetExp/Materials/InProgressMaterial"));
-	check(WorkingMaterial);
-	if (WorkingMaterial != nullptr)
-	{
-		WorkingPreviewMaterial = UMaterialInstanceDynamic::Create(WorkingMaterial, GetToolManager());
-	}
+	Super::Setup();
 
 	// Initialize preview mesh
 	bIsBakeToSelf = (Targets.Num() == 1);
 
-	UE::ToolTarget::HideSourceObject(Targets[0]);
-
-	const FDynamicMesh3 InputMeshWithTangents = UE::ToolTarget::GetDynamicMeshCopy(Targets[0], true);
-	PreviewMesh = NewObject<UPreviewMesh>(this);
-	PreviewMesh->CreateInWorld(TargetWorld, FTransform::Identity);
-	PreviewMesh->SetTransform(static_cast<FTransform>(UE::ToolTarget::GetLocalToWorldTransform(Targets[0])));
-	PreviewMesh->SetTangentsMode(EDynamicMeshComponentTangentsMode::ExternallyProvided);
-	PreviewMesh->ReplaceMesh(InputMeshWithTangents);
-	PreviewMesh->SetMaterials(UE::ToolTarget::GetMaterialSet(Targets[0]).Materials);
-	PreviewMesh->SetOverrideRenderMaterial(PreviewMaterial);
-	PreviewMesh->SetVisible(true);
-
 	PreviewMesh->ProcessMesh([this](const FDynamicMesh3& Mesh)
 	{
 		BaseMesh.Copy(Mesh);
@@ -433,10 +372,6 @@ void UBakeMeshAttributeMapsTool::Setup()
 	MultiTextureProps->WatchProperty(MultiTextureProps->MaterialIDSourceTextureMap, SetDirtyCallback, NotEqualsCallback);
 	MultiTextureProps->WatchProperty(MultiTextureProps->UVLayer, [this](float) { bInputsDirty = true; });
 
-	VisualizationProps = NewObject<UBakedOcclusionMapVisualizationProperties>(this);
-	VisualizationProps->RestoreProperties(this);
-	AddToolPropertySource(VisualizationProps);
-
 	UpdateOnModeChange();
 
 	bInputsDirty = true;
@@ -446,6 +381,8 @@ void UBakeMeshAttributeMapsTool::Setup()
 	GetToolManager()->DisplayMessage(
 		LOCTEXT("OnStartTool", "Bake Maps. Select Bake Mesh (LowPoly) first, then (optionally) Detail Mesh second. Texture Assets will be created on Accept. "),
 		EToolMessageLevel::UserNotification);
+
+	SetupBaseToolProperties();
 }
 
 
@@ -457,7 +394,7 @@ bool UBakeMeshAttributeMapsTool::CanAccept() const
 	if (bCanAccept)
 	{
 		// Allow Accept if all non-None types have valid results.
-		int NumResults = Settings->Result.Num();
+		const int NumResults = Settings->Result.Num();
 		for (int ResultIdx = 0; ResultIdx < NumResults; ++ResultIdx)
 		{
 			bCanAccept = bCanAccept && Settings->Result[ResultIdx];
@@ -529,196 +466,75 @@ TUniquePtr<UE::Geometry::TGenericDataOperator<FMeshMapBaker>> UBakeMeshAttribute
 }
 
 
-void UBakeMeshAttributeMapsTool::SetWorld(UWorld* World)
-{
-	TargetWorld = World;
-}
-
-
 void UBakeMeshAttributeMapsTool::Shutdown(EToolShutdownType ShutdownType)
 {
+	Super::Shutdown(ShutdownType);
+	
 	Settings->SaveProperties(this);
 	OcclusionMapProps->SaveProperties(this);
 	NormalMapProps->SaveProperties(this);
 	CurvatureMapProps->SaveProperties(this);
 	Texture2DProps->SaveProperties(this);
 	MultiTextureProps->SaveProperties(this);
-	VisualizationProps->SaveProperties(this);
 
 	if (Compute)
 	{
 		Compute->Shutdown();
 	}
-	if (PreviewMesh != nullptr)
+	if (ShutdownType == EToolShutdownType::Accept)
 	{
-		UE::ToolTarget::ShowSourceObject(Targets[0]);
+		UStaticMeshComponent* StaticMeshComponent = CastChecked<UStaticMeshComponent>(UE::ToolTarget::GetTargetComponent(Targets[0]));
+		UStaticMesh* StaticMeshAsset = StaticMeshComponent->GetStaticMesh();
+		check(StaticMeshAsset);
+		FString BaseName = UE::ToolTarget::GetTargetActor(Targets[0])->GetName();
 
-		if (ShutdownType == EToolShutdownType::Accept)
-		{
-			UStaticMeshComponent* StaticMeshComponent = CastChecked<UStaticMeshComponent>(UE::ToolTarget::GetTargetComponent(Targets[0]));
-			UStaticMesh* StaticMeshAsset = StaticMeshComponent->GetStaticMesh();
-			check(StaticMeshAsset);
-			FString BaseName = UE::ToolTarget::GetTargetActor(Targets[0])->GetName();
-
-			bool bCreatedAssetOK = true;
-			int NumResults = Settings->Result.Num();
-			for (int ResultIdx = 0; ResultIdx < NumResults; ++ResultIdx)
-			{
-				FTexture2DBuilder::ETextureType TexType = FTexture2DBuilder::ETextureType::ColorLinear;
-				FString TexName;
-				switch (ResultTypes[ResultIdx])
-				{
-				default:
-					// Should never reach this case.
-					check(false);
-					continue;
-				case EBakeMapType::TangentSpaceNormalMap:
-					TexName = FString::Printf(TEXT("%s_Normals"), *BaseName);
-					TexType = FTexture2DBuilder::ETextureType::NormalMap;
-					break;
-				case EBakeMapType::AmbientOcclusion:
-					TexName = FString::Printf(TEXT("%s_Occlusion"), *BaseName);
-					TexType = FTexture2DBuilder::ETextureType::AmbientOcclusion;
-					break;
-				case EBakeMapType::BentNormal:
-					TexName = FString::Printf(TEXT("%s_BentNormal"), *BaseName);
-					TexType = FTexture2DBuilder::ETextureType::NormalMap;
-					break;
-				case EBakeMapType::Curvature:
-					TexName = FString::Printf(TEXT("%s_Curvature"), *BaseName);
-					break;
-				case EBakeMapType::NormalImage:
-					TexName = FString::Printf(TEXT("%s_NormalImg"), *BaseName);
-					break;
-				case EBakeMapType::FaceNormalImage:
-					TexName = FString::Printf(TEXT("%s_FaceNormalImg"), *BaseName);
-					break;
-				case EBakeMapType::MaterialID:
-					TexName = FString::Printf(TEXT("%s_MaterialIDImg"), *BaseName);
-					break;
-				case EBakeMapType::VertexColorImage:
-					TexName = FString::Printf(TEXT("%s_VertexColorIDImg"), *BaseName);
-					break;
-				case EBakeMapType::PositionImage:
-					TexName = FString::Printf(TEXT("%s_PositionImg"), *BaseName);
-					break;
-				case EBakeMapType::Texture2DImage:
-					TexName = FString::Printf(TEXT("%s_TextureImg"), *BaseName);
-					TexType = CachedTexture2DImageSettings.bSRGB ? FTexture2DBuilder::ETextureType::Color : FTexture2DBuilder::ETextureType::ColorLinear;
-					break;
-				case EBakeMapType::MultiTexture:
-					TexName = FString::Printf(TEXT("%s_MultiTextureImg"), *BaseName);
-					TexType = CachedTexture2DImageSettings.bSRGB ? FTexture2DBuilder::ETextureType::Color : FTexture2DBuilder::ETextureType::ColorLinear;
-					break;
-				}
-				FTexture2DBuilder::CopyPlatformDataToSourceData(Settings->Result[ResultIdx], TexType);
-				bCreatedAssetOK = bCreatedAssetOK && UE::Modeling::CreateTextureObject(GetToolManager(), FCreateTextureObjectParams{ 0, StaticMeshAsset->GetWorld(), StaticMeshAsset, TexName, Settings->Result[ResultIdx] }).IsOK();
-			}
-			ensure(bCreatedAssetOK);
-		}
-
-		PreviewMesh->SetVisible(false);
-		PreviewMesh->Disconnect();
-		PreviewMesh = nullptr;
-	}
-}
-
-
-void UBakeMeshAttributeMapsTool::OnTick(float DeltaTime)
-{
-	if (Compute)
-	{
-		Compute->Tick(DeltaTime);
-
-		float ElapsedComputeTime = Compute->GetElapsedComputeTime();
-		if (!CanAccept() && ElapsedComputeTime > SecondsBeforeWorkingMaterial)
+		bool bCreatedAssetOK = true;
+		const int NumResults = Settings->Result.Num();
+		for (int ResultIdx = 0; ResultIdx < NumResults; ++ResultIdx)
 		{
-			PreviewMesh->SetOverrideRenderMaterial(WorkingPreviewMaterial);
+			const FTexture2DBuilder::ETextureType TexType = GetTextureType(ResultTypes[ResultIdx]);
+			FString TexName;
+			GetTextureName(ResultTypes[ResultIdx], BaseName, TexName);
+			FTexture2DBuilder::CopyPlatformDataToSourceData(Settings->Result[ResultIdx], TexType);
+			bCreatedAssetOK = bCreatedAssetOK && UE::Modeling::CreateTextureObject(GetToolManager(), FCreateTextureObjectParams{ 0, StaticMeshAsset->GetWorld(), StaticMeshAsset, TexName, Settings->Result[ResultIdx] }).IsOK();
 		}
+		ensure(bCreatedAssetOK);
 	}
 }
 
 
-void UBakeMeshAttributeMapsTool::Render(IToolsContextRenderAPI* RenderAPI)
+void UBakeMeshAttributeMapsTool::UpdateDetailMesh()
 {
-	UpdateResult();
-
-	float GrayLevel = VisualizationProps->BaseGrayLevel;
-	PreviewMaterial->SetVectorParameterValue(TEXT("BaseColor"), FVector(GrayLevel, GrayLevel, GrayLevel) );
-	float AOWeight = VisualizationProps->OcclusionMultiplier;
-	PreviewMaterial->SetScalarParameterValue(TEXT("AOWeight"), AOWeight );
-}
-
-
-
+	UToolTarget* DetailTarget = Targets[bIsBakeToSelf ? 0 : 1];
 
-int SelectTextureToBake(const TArray<UTexture*>& Textures)
-{
-	TArray<int> TextureVotes;
-	TextureVotes.Init(0, Textures.Num());
+	const bool bWantMeshTangents = (DetailMeshProps->DetailMeshNormalMap != nullptr);
+	DetailMesh = MakeShared<FDynamicMesh3, ESPMode::ThreadSafe>(UE::ToolTarget::GetDynamicMeshCopy(DetailTarget, bWantMeshTangents));
 
-	for (int TextureIndex = 0; TextureIndex < Textures.Num(); ++TextureIndex)
+	if (Settings->bUseWorldSpace && bIsBakeToSelf == false)
 	{
-		UTexture* Tex = Textures[TextureIndex];
-		UTexture2D* Tex2D = Cast<UTexture2D>(Tex);
-
-		if (Tex2D)
-		{
-			// Texture uses SRGB
-			if (Tex->SRGB != 0)
-			{
-				++TextureVotes[TextureIndex];
-			}
-
-#if WITH_EDITORONLY_DATA
-			// Texture has multiple channels
-			ETextureSourceFormat Format = Tex->Source.GetFormat();
-			if (Format == TSF_BGRA8 || Format == TSF_BGRE8 || Format == TSF_RGBA16 || Format == TSF_RGBA16F)
-			{
-				++TextureVotes[TextureIndex];
-			}
-#endif
-
-			// What else? Largest texture? Most layers? Most mipmaps?
-		}
+		using FTransform3d = UE::Geometry::FTransform3d;
+		const FTransform3d DetailToWorld = UE::ToolTarget::GetLocalToWorldTransform(DetailTarget);
+		MeshTransforms::ApplyTransform(*DetailMesh, DetailToWorld);
+		const FTransform3d WorldToBase = UE::ToolTarget::GetLocalToWorldTransform(Targets[0]);
+		MeshTransforms::ApplyTransform(*DetailMesh, WorldToBase.Inverse());
 	}
+	
+	DetailSpatial = MakeShared<FDynamicMeshAABBTree3, ESPMode::ThreadSafe>();
+	DetailSpatial->SetMesh(DetailMesh.Get(), true);
 
-	int MaxIndex = -1;
-	int MaxVotes = -1;
-	for (int TextureIndex = 0; TextureIndex < Textures.Num(); ++TextureIndex)
+	// Extract tangents if a DetailMesh normal map was provided.
+	if (bWantMeshTangents)
 	{
-		if (TextureVotes[TextureIndex] > MaxVotes)
-		{
-			MaxIndex = TextureIndex;
-			MaxVotes = TextureVotes[TextureIndex];
-		}
+		DetailMeshTangents = MakeShared<FMeshTangentsd, ESPMode::ThreadSafe>(DetailMesh.Get());
+		DetailMeshTangents->CopyTriVertexTangents(*DetailMesh);
 	}
-
-	return MaxIndex;
-}
-
-
-void UBakeMeshAttributeMapsTool::GetTexturesFromDetailMesh(const UPrimitiveComponent* DetailComponent)
-{
-	constexpr bool bGuessAtTextures = true;
-
-	MultiTextureProps->AllSourceTextures.Reset();
-	MultiTextureProps->MaterialIDSourceTextureMap.Reset();
-
-	TArray<UMaterialInterface*> Materials;
-	DetailComponent->GetUsedMaterials(Materials);
-	
-	for (int32 MaterialID = 0; MaterialID < Materials.Num(); ++MaterialID)	// TODO: This won't match MaterialIDs on the FDynamicMesh3 in general, will it?
+	else
 	{
-		UMaterialInterface* MaterialInterface = Materials[MaterialID];
-		if (MaterialInterface == nullptr)
-		{
-			continue;
-		}
+		DetailMeshTangents = nullptr;
+	}
 
-		TArray<UTexture*> Textures;
-		MaterialInterface->GetUsedTextures(Textures, EMaterialQualityLevel::High, true, ERHIFeatureLevel::SM5, true);
-		
+	ProcessComponentTextures(UE::ToolTarget::GetTargetComponent(DetailTarget), [this](const int MaterialID, const TArray<UTexture*>& Textures)
+	{
 		for (UTexture* Tex : Textures)
 		{
 			UTexture2D* Tex2D = Cast<UTexture2D>(Tex);
@@ -728,9 +544,10 @@ void UBakeMeshAttributeMapsTool::GetTexturesFromDetailMesh(const UPrimitiveCompo
 			}
 		}
 
+		constexpr bool bGuessAtTextures = true;
 		if (bGuessAtTextures)
 		{
-			int SelectedTextureIndex = SelectTextureToBake(Textures);
+			const int SelectedTextureIndex = SelectColorTextureToBake(Textures);
 			if (SelectedTextureIndex >= 0)
 			{
 				UTexture2D* Tex2D = Cast<UTexture2D>(Textures[SelectedTextureIndex]);
@@ -743,60 +560,15 @@ void UBakeMeshAttributeMapsTool::GetTexturesFromDetailMesh(const UPrimitiveCompo
 		{
 			MultiTextureProps->MaterialIDSourceTextureMap.Add(MaterialID, nullptr);
 		}
-	}
-
-}
-
-
-void UBakeMeshAttributeMapsTool::UpdateDetailMesh()
-{
-	UToolTarget* DetailTarget = Targets[bIsBakeToSelf ? 0 : 1];
-
-	const bool bWantMeshTangents = (DetailMeshProps->DetailMeshNormalMap != nullptr);
-	DetailMesh = MakeShared<FDynamicMesh3, ESPMode::ThreadSafe>(UE::ToolTarget::GetDynamicMeshCopy(DetailTarget, bWantMeshTangents));
-
-	if (Settings->bUseWorldSpace && bIsBakeToSelf == false)
-	{
-		using FTransform3d = UE::Geometry::FTransform3d;
-		const FTransform3d DetailToWorld = UE::ToolTarget::GetLocalToWorldTransform(DetailTarget);
-		MeshTransforms::ApplyTransform(*DetailMesh, DetailToWorld);
-		const FTransform3d WorldToBase = UE::ToolTarget::GetLocalToWorldTransform(Targets[0]);
-		MeshTransforms::ApplyTransform(*DetailMesh, WorldToBase.Inverse());
-	}
-	
-	DetailSpatial = MakeShared<FDynamicMeshAABBTree3, ESPMode::ThreadSafe>();
-	DetailSpatial->SetMesh(DetailMesh.Get(), true);
-
-	// Extract tangents if a DetailMesh normal map was provided.
-	if (bWantMeshTangents)
-	{
-		DetailMeshTangents = MakeShared<FMeshTangentsd, ESPMode::ThreadSafe>(DetailMesh.Get());
-		DetailMeshTangents->CopyTriVertexTangents(*DetailMesh);
-	}
-	else
-	{
-		DetailMeshTangents = nullptr;
-	}
-
-	GetTexturesFromDetailMesh(UE::ToolTarget::GetTargetComponent(DetailTarget));
+	});
 
 	bInputsDirty = true;
 	DetailMeshTimestamp++;
-
 }
 
 
-
-
-
-
-
-
-
-
 void UBakeMeshAttributeMapsTool::UpdateResult()
 {
-
 	if (bDetailMeshValid == false)
 	{
 		UpdateDetailMesh();
@@ -941,8 +713,8 @@ EBakeOpState UBakeMeshAttributeMapsTool::UpdateResult_Normal()
 {
 	EBakeOpState ResultState = EBakeOpState::Complete;
 
-	int32 ImageSize = (int32)Settings->Resolution;
-	FImageDimensions Dimensions(ImageSize, ImageSize);
+	const int32 ImageSize = (int32)Settings->Resolution;
+	const FImageDimensions Dimensions(ImageSize, ImageSize);
 
 	FNormalMapSettings NormalMapSettings;
 	NormalMapSettings.Dimensions = Dimensions;
@@ -960,8 +732,8 @@ EBakeOpState UBakeMeshAttributeMapsTool::UpdateResult_Occlusion()
 {
 	EBakeOpState ResultState = EBakeOpState::Complete;
 
-	int32 ImageSize = (int32)Settings->Resolution;
-	FImageDimensions Dimensions(ImageSize, ImageSize);
+	const int32 ImageSize = (int32)Settings->Resolution;
+	const FImageDimensions Dimensions(ImageSize, ImageSize);
 
 	FOcclusionMapSettings OcclusionMapSettings;
 	OcclusionMapSettings.Dimensions = Dimensions;
@@ -982,13 +754,12 @@ EBakeOpState UBakeMeshAttributeMapsTool::UpdateResult_Occlusion()
 }
 
 
-
 EBakeOpState UBakeMeshAttributeMapsTool::UpdateResult_Curvature()
 {
 	EBakeOpState ResultState = EBakeOpState::Complete;
 
-	int32 ImageSize = (int32)Settings->Resolution;
-	FImageDimensions Dimensions(ImageSize, ImageSize);
+	const int32 ImageSize = (int32)Settings->Resolution;
+	const FImageDimensions Dimensions(ImageSize, ImageSize);
 
 	FCurvatureMapSettings CurvatureMapSettings;
 	CurvatureMapSettings.Dimensions = Dimensions;
@@ -1047,13 +818,12 @@ EBakeOpState UBakeMeshAttributeMapsTool::UpdateResult_Curvature()
 }
 
 
-
 EBakeOpState UBakeMeshAttributeMapsTool::UpdateResult_MeshProperty()
 {
 	EBakeOpState ResultState = EBakeOpState::Complete;
 
-	int32 ImageSize = (int32)Settings->Resolution;
-	FImageDimensions Dimensions(ImageSize, ImageSize);
+	const int32 ImageSize = (int32)Settings->Resolution;
+	const FImageDimensions Dimensions(ImageSize, ImageSize);
 
 	FMeshPropertyMapSettings MeshPropertyMapSettings;
 	MeshPropertyMapSettings.Dimensions = Dimensions;
@@ -1071,8 +841,8 @@ EBakeOpState UBakeMeshAttributeMapsTool::UpdateResult_Texture2DImage()
 {
 	EBakeOpState ResultState = EBakeOpState::Complete;
 
-	int32 ImageSize = (int32)Settings->Resolution;
-	FImageDimensions Dimensions(ImageSize, ImageSize);
+	const int32 ImageSize = (int32)Settings->Resolution;
+	const FImageDimensions Dimensions(ImageSize, ImageSize);
 
 	FTexture2DImageSettings NewSettings;
 	NewSettings.Dimensions = Dimensions;
@@ -1116,8 +886,8 @@ EBakeOpState UBakeMeshAttributeMapsTool::UpdateResult_MultiTexture()
 {
 	EBakeOpState ResultState = EBakeOpState::Complete;
 
-	int32 ImageSize = (int32)Settings->Resolution;
-	FImageDimensions Dimensions(ImageSize, ImageSize);
+	const int32 ImageSize = (int32)Settings->Resolution;
+	const FImageDimensions Dimensions(ImageSize, ImageSize);
 
 	FTexture2DImageSettings NewSettings;
 	NewSettings.Dimensions = Dimensions;
@@ -1176,9 +946,6 @@ EBakeOpState UBakeMeshAttributeMapsTool::UpdateResult_MultiTexture()
 }
 
 
-
-
-
 void UBakeMeshAttributeMapsTool::UpdateVisualization()
 {
 	PreviewMesh->SetOverrideRenderMaterial(PreviewMaterial);
@@ -1193,62 +960,15 @@ void UBakeMeshAttributeMapsTool::UpdateVisualization()
 	// Set the preview material according to the preview index.
 	if (Settings->MapPreview >= 0 && Settings->MapPreview < Settings->Result.Num())
 	{
-		const EBakeMapType& PreviewMapType = ResultTypes[Settings->MapPreview];
-		if (PreviewMapType != EBakeMapType::None)
-		{
-			UTexture2D* PreviewMap = CachedMaps[CachedMapIndices[PreviewMapType]];
-			switch (PreviewMapType)
-			{
-			default:
-				PreviewMaterial->SetTextureParameterValue(TEXT("NormalMap"), EmptyNormalMap);
-				PreviewMaterial->SetTextureParameterValue(TEXT("OcclusionMap"), EmptyColorMapWhite);
-				PreviewMaterial->SetTextureParameterValue(TEXT("ColorMap"), EmptyColorMapWhite);
-				break;
-			case EBakeMapType::TangentSpaceNormalMap:
-				PreviewMaterial->SetTextureParameterValue(TEXT("NormalMap"), PreviewMap);
-				PreviewMaterial->SetTextureParameterValue(TEXT("OcclusionMap"), EmptyColorMapWhite);
-				PreviewMaterial->SetTextureParameterValue(TEXT("ColorMap"), EmptyColorMapWhite);
-				break;
-			case EBakeMapType::AmbientOcclusion:
-				PreviewMaterial->SetTextureParameterValue(TEXT("NormalMap"), EmptyNormalMap);
-				PreviewMaterial->SetTextureParameterValue(TEXT("OcclusionMap"), PreviewMap);
-				PreviewMaterial->SetTextureParameterValue(TEXT("ColorMap"), EmptyColorMapWhite);
-				break;
-			case EBakeMapType::BentNormal:
-				BentNormalPreviewMaterial->SetTextureParameterValue(TEXT("NormalMap"), EmptyNormalMap);
-				if (CachedMapIndices.Contains(EBakeMapType::AmbientOcclusion))
-				{
-					BentNormalPreviewMaterial->SetTextureParameterValue(TEXT("OcclusionMap"), CachedMaps[CachedMapIndices[EBakeMapType::AmbientOcclusion]]);
-				}
-				else
-				{
-					BentNormalPreviewMaterial->SetTextureParameterValue(TEXT("OcclusionMap"), EmptyColorMapWhite);
-				}
-				BentNormalPreviewMaterial->SetTextureParameterValue(TEXT("ColorMap"), EmptyColorMapWhite);
-				BentNormalPreviewMaterial->SetTextureParameterValue(TEXT("BentNormalMap"), PreviewMap);
-				PreviewMesh->SetOverrideRenderMaterial(BentNormalPreviewMaterial);
-				break;
-			case EBakeMapType::Curvature:
-			case EBakeMapType::NormalImage:
-			case EBakeMapType::FaceNormalImage:
-			case EBakeMapType::PositionImage:
-			case EBakeMapType::MaterialID:
-			case EBakeMapType::Texture2DImage:
-			case EBakeMapType::MultiTexture:
-			case EBakeMapType::VertexColorImage:
-				PreviewMaterial->SetTextureParameterValue(TEXT("NormalMap"), EmptyNormalMap);
-				PreviewMaterial->SetTextureParameterValue(TEXT("OcclusionMap"), EmptyColorMapWhite);
-				PreviewMaterial->SetTextureParameterValue(TEXT("ColorMap"), PreviewMap);
-				break;
-			}
-		}
+		UpdatePreview(Settings->MapPreview);
 	}
 }
 
 
-
 void UBakeMeshAttributeMapsTool::UpdateOnModeChange()
 {
+	OnMapTypesUpdated(Settings->MapTypes);
+	
 	SetToolPropertySourceEnabled(NormalMapProps, false);
 	SetToolPropertySourceEnabled(OcclusionMapProps, false);
 	SetToolPropertySourceEnabled(CurvatureMapProps, false);
@@ -1287,185 +1007,9 @@ void UBakeMeshAttributeMapsTool::UpdateOnModeChange()
 		}
 	}
 
-	ResultTypes = GetMapTypesArray(Settings->MapTypes);
 	Settings->Result.Empty();
 	Settings->Result.SetNum(ResultTypes.Num());
-
-	// Generate a map between EBakeMapType and CachedMaps
-	CachedMapIndices.Empty();
-	int32 CachedMapIdx = 0;
-
-	// Use the processed bitfield which may contain additional targets
-	// (ex. AO if BentNormal was requested).
-	const EBakeMapType BakeMapTypes = GetMapTypes(Settings->MapTypes);
-	for (EBakeMapType MapType : ALL_BAKE_MAP_TYPES)
-	{
-		if (MapType == EBakeMapType::Occlusion)
-		{
-			if ((bool)(BakeMapTypes & EBakeMapType::AmbientOcclusion))
-			{
-				CachedMapIndices.Add(EBakeMapType::AmbientOcclusion, CachedMapIdx++);
-			}
-			if ((bool)(BakeMapTypes & EBakeMapType::BentNormal))
-			{
-				CachedMapIndices.Add(EBakeMapType::BentNormal, CachedMapIdx++);
-			}
-		}
-		else if( (bool)(BakeMapTypes & MapType) )
-		{
-			CachedMapIndices.Add(MapType, CachedMapIdx++);
-		}
-	}
-	CachedMaps.Empty();
-	CachedMaps.SetNum(CachedMapIndices.Num());
-}
-
-
-void UBakeMeshAttributeMapsTool::OnMapsUpdated(const TUniquePtr<FMeshMapBaker>& NewResult)
-{
-	FImageDimensions BakeDimensions = NewResult->GetDimensions();
-	const int32 NumBakers = NewResult->NumBakers();
-	for (int32 BakerIdx = 0; BakerIdx < NumBakers; ++BakerIdx)
-	{
-		FMeshMapEvaluator* Baker = NewResult->GetBaker(BakerIdx);
-
-		auto UpdateCachedMap = [this, &NewResult, &BakerIdx, &BakeDimensions](const EBakeMapType BakeMapType, const FTexture2DBuilder::ETextureType TexType, const int32 ResultIdx) -> void
-		{
-			FTexture2DBuilder TextureBuilder;
-			TextureBuilder.Initialize(TexType, BakeDimensions);
-			TextureBuilder.Copy(*NewResult->GetBakeResults(BakerIdx)[ResultIdx]);
-			TextureBuilder.Commit(false);
-
-			// The CachedMap & CachedMapIndices can be thrown out of sync if updated during
-			// a background compute. Validate the computed type against our cached maps.
-			if (CachedMapIndices.Contains(BakeMapType))
-			{
-				CachedMaps[CachedMapIndices[BakeMapType]] = TextureBuilder.GetTexture2D();
-			}
-		};
-
-		switch (Baker->Type())
-		{
-		case EMeshMapEvaluatorType::Normal:
-		{
-			UpdateCachedMap(EBakeMapType::TangentSpaceNormalMap, FTexture2DBuilder::ETextureType::NormalMap, 0);
-			break;
-		}
-		case EMeshMapEvaluatorType::Occlusion:
-		{
-			// Occlusion Evaluator always outputs AmbientOcclusion then BentNormal.
-			FMeshOcclusionMapEvaluator* OcclusionBaker = static_cast<FMeshOcclusionMapEvaluator*>(Baker);
-			int32 OcclusionIdx = 0;
-			if ((bool)(OcclusionBaker->OcclusionType & EMeshOcclusionMapType::AmbientOcclusion))
-			{
-				UpdateCachedMap(EBakeMapType::AmbientOcclusion, FTexture2DBuilder::ETextureType::AmbientOcclusion, OcclusionIdx++);
-			}
-			if ((bool)(OcclusionBaker->OcclusionType & EMeshOcclusionMapType::BentNormal))
-			{
-				UpdateCachedMap(EBakeMapType::BentNormal, FTexture2DBuilder::ETextureType::NormalMap, OcclusionIdx++);
-			}
-			break;
-		}
-		case EMeshMapEvaluatorType::Curvature:
-		{
-			UpdateCachedMap(EBakeMapType::Curvature, FTexture2DBuilder::ETextureType::ColorLinear, 0);
-			break;
-		}
-		case EMeshMapEvaluatorType::Property:
-		{
-			FMeshPropertyMapEvaluator* PropertyBaker = static_cast<FMeshPropertyMapEvaluator*>(Baker);
-			EBakeMapType MapType = EBakeMapType::None;
-			switch (PropertyBaker->Property)
-			{
-			case EMeshPropertyMapType::Normal:
-				MapType = EBakeMapType::NormalImage;
-				break;
-			case EMeshPropertyMapType::FacetNormal:
-				MapType = EBakeMapType::FaceNormalImage;
-				break;
-			case EMeshPropertyMapType::Position:
-				MapType = EBakeMapType::PositionImage;
-				break;
-			case EMeshPropertyMapType::MaterialID:
-				MapType = EBakeMapType::MaterialID;
-				break;
-			case EMeshPropertyMapType::VertexColor:
-				MapType = EBakeMapType::VertexColorImage;
-				break;
-			case EMeshPropertyMapType::UVPosition:
-			default:
-				break;
-			}
-
-			UpdateCachedMap(MapType, FTexture2DBuilder::ETextureType::ColorLinear, 0);
-			break;
-		}
-		case EMeshMapEvaluatorType::ResampleImage:
-		{
-			FTexture2DBuilder::ETextureType TexType = CachedTexture2DImageSettings.bSRGB ? FTexture2DBuilder::ETextureType::Color : FTexture2DBuilder::ETextureType::ColorLinear;
-			UpdateCachedMap(EBakeMapType::Texture2DImage, TexType, 0);
-			break;
-		}
-		case EMeshMapEvaluatorType::MultiResampleImage:
-		{
-			FTexture2DBuilder::ETextureType TexType = CachedTexture2DImageSettings.bSRGB ? FTexture2DBuilder::ETextureType::Color : FTexture2DBuilder::ETextureType::ColorLinear;
-			UpdateCachedMap(EBakeMapType::MultiTexture, TexType, 0);
-			break;
-		}
-		}
-	}
-
-	UpdateVisualization();
-	GetToolManager()->PostInvalidation();
 }
 
 
-EBakeMapType UBakeMeshAttributeMapsTool::GetMapTypes(const int32& MapTypes) const
-{
-	EBakeMapType OutMapTypes = (EBakeMapType)MapTypes & EBakeMapType::All;
-	// Force AO bake for BentNormal preview
-	if ((bool)(OutMapTypes & EBakeMapType::BentNormal))
-	{
-		OutMapTypes |= EBakeMapType::AmbientOcclusion;
-	}
-	return OutMapTypes;
-}
-
-TArray<EBakeMapType> UBakeMeshAttributeMapsTool::GetMapTypesArray(const int32& MapTypes) const
-{
-	TArray<EBakeMapType> OutMapTypes;
-	int32 Bitfield = MapTypes & (int32)EBakeMapType::All;
-	for (int32 BitIdx = 0; Bitfield; Bitfield >>= 1, ++BitIdx)
-	{
-		if (Bitfield & 1)
-		{
-			OutMapTypes.Add((EBakeMapType)(1 << BitIdx));
-		}
-	}
-	return OutMapTypes;
-}
-
-
-void UBakeMeshAttributeMapsTool::InitializeEmptyMaps()
-{
-	FTexture2DBuilder NormalsBuilder;
-	NormalsBuilder.Initialize(FTexture2DBuilder::ETextureType::NormalMap, FImageDimensions(16, 16));
-	NormalsBuilder.Commit(false);
-	EmptyNormalMap = NormalsBuilder.GetTexture2D();
-
-	FTexture2DBuilder ColorBuilderBlack;
-	ColorBuilderBlack.Initialize(FTexture2DBuilder::ETextureType::Color, FImageDimensions(16, 16));
-	ColorBuilderBlack.Clear(FColor(0,0,0));
-	ColorBuilderBlack.Commit(false);
-	EmptyColorMapBlack = ColorBuilderBlack.GetTexture2D();
-
-	FTexture2DBuilder ColorBuilderWhite;
-	ColorBuilderWhite.Initialize(FTexture2DBuilder::ETextureType::Color, FImageDimensions(16, 16));
-	ColorBuilderWhite.Clear(FColor::White);
-	ColorBuilderWhite.Commit(false);
-	EmptyColorMapWhite = ColorBuilderWhite.GetTexture2D();
-}
-
-
-
 #undef LOCTEXT_NAMESPACE
diff --git a/Engine/Plugins/Experimental/MeshModelingToolsetExp/Source/MeshModelingToolsExp/Private/BakeMeshAttributeMapsToolBase.cpp b/Engine/Plugins/Experimental/MeshModelingToolsetExp/Source/MeshModelingToolsExp/Private/BakeMeshAttributeMapsToolBase.cpp
new file mode 100644
--- /dev/null
+++ b/Engine/Plugins/Experimental/MeshModelingToolsetExp/Source/MeshModelingToolsExp/Private/BakeMeshAttributeMapsToolBase.cpp
@@ -0,0 +1,493 @@
+// Copyright Epic Games, Inc. All Rights Reserved.
+
+#include "BakeMeshAttributeMapsToolBase.h"
+#include "InteractiveToolManager.h"
+#include "ToolBuilderUtil.h"
+
+#include "DynamicMesh/DynamicMesh3.h"
+#include "DynamicMesh/DynamicMeshAttributeSet.h"
+#include "Materials/Material.h"
+#include "Materials/MaterialInstanceDynamic.h"
+#include "ImageUtils.h"
+
+#include "Sampling/MeshOcclusionMapEvaluator.h"
+#include "Sampling/MeshPropertyMapEvaluator.h"
+
+#include "AssetUtils/Texture2DBuilder.h"
+#include "AssetUtils/Texture2DUtil.h"
+
+#include "TargetInterfaces/MaterialProvider.h"
+#include "ModelingToolTargetUtil.h"
+
+#include "ExplicitUseGeometryMathTypes.h"		// using UE::Geometry::(math types)
+using namespace UE::Geometry;
+
+#define LOCTEXT_NAMESPACE "UBakeMeshAttributeMapsToolBase"
+
+
+void UBakeMeshAttributeMapsToolBase::Setup()
+{
+	Super::Setup();
+
+	InitializeEmptyMaps();
+
+	// Setup preview materials
+	UMaterial* Material = LoadObject<UMaterial>(nullptr, TEXT("/MeshModelingToolsetExp/Materials/BakePreviewMaterial"));
+	check(Material);
+	if (Material != nullptr)
+	{
+		PreviewMaterial = UMaterialInstanceDynamic::Create(Material, GetToolManager());
+		PreviewMaterial->SetTextureParameterValue(TEXT("NormalMap"), EmptyNormalMap);
+		PreviewMaterial->SetTextureParameterValue(TEXT("OcclusionMap"), EmptyColorMapWhite);
+		PreviewMaterial->SetTextureParameterValue(TEXT("ColorMap"), EmptyColorMapWhite);
+	}
+	UMaterial* BentNormalMaterial = LoadObject<UMaterial>(nullptr, TEXT("/MeshModelingToolsetExp/Materials/BakeBentNormalPreviewMaterial"));
+	check(BentNormalMaterial);
+	if (BentNormalMaterial != nullptr)
+	{
+		BentNormalPreviewMaterial = UMaterialInstanceDynamic::Create(BentNormalMaterial, GetToolManager());
+	}
+	UMaterial* WorkingMaterial = LoadObject<UMaterial>(nullptr, TEXT("/MeshModelingToolsetExp/Materials/InProgressMaterial"));
+	check(WorkingMaterial);
+	if (WorkingMaterial != nullptr)
+	{
+		WorkingPreviewMaterial = UMaterialInstanceDynamic::Create(WorkingMaterial, GetToolManager());
+	}
+
+	// Initialize preview mesh
+	UE::ToolTarget::HideSourceObject(Targets[0]);
+
+	const FDynamicMesh3 InputMesh = UE::ToolTarget::GetDynamicMeshCopy(Targets[0], true);
+	const UE::Geometry::FTransform3d BaseToWorld = UE::ToolTarget::GetLocalToWorldTransform(Targets[0]);
+	PreviewMesh = NewObject<UPreviewMesh>(this);
+	PreviewMesh->CreateInWorld(TargetWorld, FTransform::Identity);
+	PreviewMesh->SetTransform(static_cast<FTransform>(BaseToWorld));
+	PreviewMesh->SetTangentsMode(EDynamicMeshComponentTangentsMode::ExternallyProvided);
+	PreviewMesh->ReplaceMesh(InputMesh);
+	PreviewMesh->SetMaterials(UE::ToolTarget::GetMaterialSet(Targets[0]).Materials);
+	PreviewMesh->SetOverrideRenderMaterial(PreviewMaterial);
+	PreviewMesh->SetVisible(true);
+}
+
+void UBakeMeshAttributeMapsToolBase::SetupBaseToolProperties()
+{
+	VisualizationProps = NewObject<UBakedOcclusionMapVisualizationProperties>(this);
+	VisualizationProps->RestoreProperties(this);
+	AddToolPropertySource(VisualizationProps);
+}
+
+void UBakeMeshAttributeMapsToolBase::Shutdown(EToolShutdownType ShutdownType)
+{
+	VisualizationProps->SaveProperties(this);
+
+	if (PreviewMesh != nullptr)
+	{
+		PreviewMesh->SetVisible(false);
+		PreviewMesh->Disconnect();
+		PreviewMesh = nullptr;
+	}
+
+	UE::ToolTarget::ShowSourceObject(Targets[0]);
+}
+
+
+void UBakeMeshAttributeMapsToolBase::OnTick(float DeltaTime)
+{
+	if (Compute)
+	{
+		Compute->Tick(DeltaTime);
+
+		float ElapsedComputeTime = Compute->GetElapsedComputeTime();
+		if (!CanAccept() && ElapsedComputeTime > SecondsBeforeWorkingMaterial)
+		{
+			PreviewMesh->SetOverrideRenderMaterial(WorkingPreviewMaterial);
+		}
+	}
+}
+
+
+void UBakeMeshAttributeMapsToolBase::Render(IToolsContextRenderAPI* RenderAPI)
+{
+	UpdateResult();
+	
+	float GrayLevel = VisualizationProps->BaseGrayLevel;
+	PreviewMaterial->SetVectorParameterValue(TEXT("BaseColor"), FVector(GrayLevel, GrayLevel, GrayLevel) );
+	float AOWeight = VisualizationProps->OcclusionMultiplier;
+	PreviewMaterial->SetScalarParameterValue(TEXT("AOWeight"), AOWeight );
+}
+
+
+TUniquePtr<UE::Geometry::TGenericDataOperator<FMeshMapBaker>> UBakeMeshAttributeMapsToolBase::MakeNewOperator()
+{
+	return nullptr;
+}
+
+
+void UBakeMeshAttributeMapsToolBase::SetWorld(UWorld* World)
+{
+	TargetWorld = World;
+}
+
+void UBakeMeshAttributeMapsToolBase::UpdateResult()
+{
+}
+
+void UBakeMeshAttributeMapsToolBase::UpdateVisualization()
+{
+}
+
+void UBakeMeshAttributeMapsToolBase::OnMapTypesUpdated(const int32 MapTypes)
+{
+	const EBakeMapType BakeMapTypes = GetMapTypes(MapTypes);
+	ResultTypes = GetMapTypesArray(MapTypes);
+	
+	// Generate a map between EBakeMapType and CachedMaps
+	CachedMapIndices.Empty();
+	int32 CachedMapIdx = 0;
+
+	// Use the processed bitfield which may contain additional targets
+	// (ex. AO if BentNormal was requested).
+	for (EBakeMapType MapType : ALL_BAKE_MAP_TYPES)
+	{
+		if (MapType == EBakeMapType::Occlusion)
+		{
+			if ((bool)(BakeMapTypes & EBakeMapType::AmbientOcclusion))
+			{
+				CachedMapIndices.Add(EBakeMapType::AmbientOcclusion, CachedMapIdx++);
+			}
+			if ((bool)(BakeMapTypes & EBakeMapType::BentNormal))
+			{
+				CachedMapIndices.Add(EBakeMapType::BentNormal, CachedMapIdx++);
+			}
+		}
+		else if( (bool)(BakeMapTypes & MapType) )
+		{
+			CachedMapIndices.Add(MapType, CachedMapIdx++);
+		}
+	}
+	CachedMaps.Empty();
+	CachedMaps.SetNum(CachedMapIndices.Num());
+}
+
+void UBakeMeshAttributeMapsToolBase::UpdatePreview(const int PreviewIdx)
+{
+	const EBakeMapType& PreviewMapType = ResultTypes[PreviewIdx];
+	if (PreviewMapType == EBakeMapType::None)
+	{
+		return;
+	}
+	
+	UTexture2D* PreviewMap = CachedMaps[CachedMapIndices[PreviewMapType]];
+	switch (PreviewMapType)
+	{
+	default:
+		PreviewMaterial->SetTextureParameterValue(TEXT("NormalMap"), EmptyNormalMap);
+		PreviewMaterial->SetTextureParameterValue(TEXT("OcclusionMap"), EmptyColorMapWhite);
+		PreviewMaterial->SetTextureParameterValue(TEXT("ColorMap"), EmptyColorMapWhite);
+		break;
+	case EBakeMapType::TangentSpaceNormalMap:
+		PreviewMaterial->SetTextureParameterValue(TEXT("NormalMap"), PreviewMap);
+		PreviewMaterial->SetTextureParameterValue(TEXT("OcclusionMap"), EmptyColorMapWhite);
+		PreviewMaterial->SetTextureParameterValue(TEXT("ColorMap"), EmptyColorMapWhite);
+		break;
+	case EBakeMapType::AmbientOcclusion:
+		PreviewMaterial->SetTextureParameterValue(TEXT("NormalMap"), EmptyNormalMap);
+		PreviewMaterial->SetTextureParameterValue(TEXT("OcclusionMap"), PreviewMap);
+		PreviewMaterial->SetTextureParameterValue(TEXT("ColorMap"), EmptyColorMapWhite);
+		break;
+	case EBakeMapType::BentNormal:
+	{
+		UTexture2D* AOMap = CachedMapIndices.Contains(EBakeMapType::AmbientOcclusion) ? CachedMaps[CachedMapIndices[EBakeMapType::AmbientOcclusion]] : EmptyColorMapWhite;
+		BentNormalPreviewMaterial->SetTextureParameterValue(TEXT("NormalMap"), EmptyNormalMap);
+		BentNormalPreviewMaterial->SetTextureParameterValue(TEXT("OcclusionMap"), AOMap);
+		BentNormalPreviewMaterial->SetTextureParameterValue(TEXT("ColorMap"), EmptyColorMapWhite);
+		BentNormalPreviewMaterial->SetTextureParameterValue(TEXT("BentNormalMap"), PreviewMap);
+		PreviewMesh->SetOverrideRenderMaterial(BentNormalPreviewMaterial);
+		break;
+	}	
+	case EBakeMapType::Curvature:
+	case EBakeMapType::NormalImage:
+	case EBakeMapType::FaceNormalImage:
+	case EBakeMapType::PositionImage:
+	case EBakeMapType::MaterialID:
+	case EBakeMapType::Texture2DImage:
+	case EBakeMapType::MultiTexture:
+	case EBakeMapType::VertexColorImage:
+		PreviewMaterial->SetTextureParameterValue(TEXT("NormalMap"), EmptyNormalMap);
+		PreviewMaterial->SetTextureParameterValue(TEXT("OcclusionMap"), EmptyColorMapWhite);
+		PreviewMaterial->SetTextureParameterValue(TEXT("ColorMap"), PreviewMap);
+		break;
+	}
+}
+
+
+void UBakeMeshAttributeMapsToolBase::OnMapsUpdated(const TUniquePtr<UE::Geometry::FMeshMapBaker>& NewResult)
+{
+	FImageDimensions BakeDimensions = NewResult->GetDimensions();
+	const int32 NumBakers = NewResult->NumBakers();
+	for (int32 BakerIdx = 0; BakerIdx < NumBakers; ++BakerIdx)
+	{
+		FMeshMapEvaluator* Baker = NewResult->GetBaker(BakerIdx);
+
+		auto UpdateCachedMap = [this, &NewResult, &BakerIdx, &BakeDimensions](const EBakeMapType BakeMapType, const FTexture2DBuilder::ETextureType TexType, const int32 ResultIdx) -> void
+		{
+			FTexture2DBuilder TextureBuilder;
+			TextureBuilder.Initialize(TexType, BakeDimensions);
+			TextureBuilder.Copy(*NewResult->GetBakeResults(BakerIdx)[ResultIdx]);
+			TextureBuilder.Commit(false);
+
+			// The CachedMap & CachedMapIndices can be thrown out of sync if updated during
+			// a background compute. Validate the computed type against our cached maps.
+			if (CachedMapIndices.Contains(BakeMapType))
+			{
+				CachedMaps[CachedMapIndices[BakeMapType]] = TextureBuilder.GetTexture2D();
+			}
+		};
+
+		switch (Baker->Type())
+		{
+		case EMeshMapEvaluatorType::Normal:
+		{
+			UpdateCachedMap(EBakeMapType::TangentSpaceNormalMap, FTexture2DBuilder::ETextureType::NormalMap, 0);
+			break;
+		}
+		case EMeshMapEvaluatorType::Occlusion:
+		{
+			// Occlusion Evaluator always outputs AmbientOcclusion then BentNormal.
+			FMeshOcclusionMapEvaluator* OcclusionBaker = static_cast<FMeshOcclusionMapEvaluator*>(Baker);
+			int32 OcclusionIdx = 0;
+			if ((bool)(OcclusionBaker->OcclusionType & EMeshOcclusionMapType::AmbientOcclusion))
+			{
+				UpdateCachedMap(EBakeMapType::AmbientOcclusion, FTexture2DBuilder::ETextureType::AmbientOcclusion, OcclusionIdx++);
+			}
+			if ((bool)(OcclusionBaker->OcclusionType & EMeshOcclusionMapType::BentNormal))
+			{
+				UpdateCachedMap(EBakeMapType::BentNormal, FTexture2DBuilder::ETextureType::NormalMap, OcclusionIdx++);
+			}
+			break;
+		}
+		case EMeshMapEvaluatorType::Curvature:
+		{
+			UpdateCachedMap(EBakeMapType::Curvature, FTexture2DBuilder::ETextureType::ColorLinear, 0);
+			break;
+		}
+		case EMeshMapEvaluatorType::Property:
+		{
+			FMeshPropertyMapEvaluator* PropertyBaker = static_cast<FMeshPropertyMapEvaluator*>(Baker);
+			EBakeMapType MapType = EBakeMapType::None;
+			switch (PropertyBaker->Property)
+			{
+			case EMeshPropertyMapType::Normal:
+				MapType = EBakeMapType::NormalImage;
+				break;
+			case EMeshPropertyMapType::FacetNormal:
+				MapType = EBakeMapType::FaceNormalImage;
+				break;
+			case EMeshPropertyMapType::Position:
+				MapType = EBakeMapType::PositionImage;
+				break;
+			case EMeshPropertyMapType::MaterialID:
+				MapType = EBakeMapType::MaterialID;
+				break;
+			case EMeshPropertyMapType::VertexColor:
+				MapType = EBakeMapType::VertexColorImage;
+				break;
+			case EMeshPropertyMapType::UVPosition:
+			default:
+				break;
+			}
+
+			UpdateCachedMap(MapType, FTexture2DBuilder::ETextureType::ColorLinear, 0);
+			break;
+		}
+		case EMeshMapEvaluatorType::ResampleImage:
+		{
+			constexpr FTexture2DBuilder::ETextureType TexType = FTexture2DBuilder::ETextureType::ColorLinear;
+			UpdateCachedMap(EBakeMapType::Texture2DImage, TexType, 0);
+			break;
+		}
+		case EMeshMapEvaluatorType::MultiResampleImage:
+		{
+			constexpr FTexture2DBuilder::ETextureType TexType = FTexture2DBuilder::ETextureType::ColorLinear;
+			UpdateCachedMap(EBakeMapType::MultiTexture, TexType, 0);
+			break;
+		}
+		default:
+			break;
+		}
+	}
+
+	UpdateVisualization();
+	GetToolManager()->PostInvalidation();
+}
+
+
+EBakeMapType UBakeMeshAttributeMapsToolBase::GetMapTypes(const int32& MapTypes)
+{
+	EBakeMapType OutMapTypes = (EBakeMapType)MapTypes & EBakeMapType::All;
+	// Force AO bake for BentNormal preview
+	if ((bool)(OutMapTypes & EBakeMapType::BentNormal))
+	{
+		OutMapTypes |= EBakeMapType::AmbientOcclusion;
+	}
+	return OutMapTypes;
+}
+
+
+TArray<EBakeMapType> UBakeMeshAttributeMapsToolBase::GetMapTypesArray(const int32& MapTypes)
+{
+	TArray<EBakeMapType> OutMapTypes;
+	int32 Bitfield = MapTypes & (int32)EBakeMapType::All;
+	for (int32 BitIdx = 0; Bitfield; Bitfield >>= 1, ++BitIdx)
+	{
+		if (Bitfield & 1)
+		{
+			OutMapTypes.Add((EBakeMapType)(1 << BitIdx));
+		}
+	}
+	return OutMapTypes;
+}
+
+
+UE::Geometry::FTexture2DBuilder::ETextureType UBakeMeshAttributeMapsToolBase::GetTextureType(EBakeMapType MapType)
+{
+	FTexture2DBuilder::ETextureType TexType = FTexture2DBuilder::ETextureType::ColorLinear;
+	switch (MapType)
+	{
+	default:
+		checkNoEntry();
+		break;
+	case EBakeMapType::TangentSpaceNormalMap:
+		TexType = FTexture2DBuilder::ETextureType::NormalMap;
+		break;
+	case EBakeMapType::AmbientOcclusion:
+		TexType = FTexture2DBuilder::ETextureType::AmbientOcclusion;
+		break;
+	case EBakeMapType::BentNormal:
+		TexType = FTexture2DBuilder::ETextureType::NormalMap;
+		break;
+	case EBakeMapType::Curvature:
+	case EBakeMapType::NormalImage:
+	case EBakeMapType::FaceNormalImage:
+	case EBakeMapType::MaterialID:
+	case EBakeMapType::VertexColorImage:
+	case EBakeMapType::PositionImage:
+	case EBakeMapType::Texture2DImage:
+	case EBakeMapType::MultiTexture:
+		break;
+	}
+	return TexType;
+}
+
+
+void UBakeMeshAttributeMapsToolBase::GetTextureName(EBakeMapType MapType, const FString& BaseName, FString& TexName)
+{
+	switch (MapType)
+	{
+	default:
+		checkNoEntry();
+		break;
+	case EBakeMapType::TangentSpaceNormalMap:
+		TexName = FString::Printf(TEXT("%s_Normals"), *BaseName);
+		break;
+	case EBakeMapType::AmbientOcclusion:
+		TexName = FString::Printf(TEXT("%s_Occlusion"), *BaseName);
+		break;
+	case EBakeMapType::BentNormal:
+		TexName = FString::Printf(TEXT("%s_BentNormal"), *BaseName);
+		break;
+	case EBakeMapType::Curvature:
+		TexName = FString::Printf(TEXT("%s_Curvature"), *BaseName);
+		break;
+	case EBakeMapType::NormalImage:
+		TexName = FString::Printf(TEXT("%s_NormalImg"), *BaseName);
+		break;
+	case EBakeMapType::FaceNormalImage:
+		TexName = FString::Printf(TEXT("%s_FaceNormalImg"), *BaseName);
+		break;
+	case EBakeMapType::MaterialID:
+		TexName = FString::Printf(TEXT("%s_MaterialIDImg"), *BaseName);
+		break;
+	case EBakeMapType::VertexColorImage:
+		TexName = FString::Printf(TEXT("%s_VertexColorIDImg"), *BaseName);
+		break;
+	case EBakeMapType::PositionImage:
+		TexName = FString::Printf(TEXT("%s_PositionImg"), *BaseName);
+		break;
+	case EBakeMapType::Texture2DImage:
+		TexName = FString::Printf(TEXT("%s_TextureImg"), *BaseName);
+		break;
+	case EBakeMapType::MultiTexture:
+		TexName = FString::Printf(TEXT("%s_MultiTextureImg"), *BaseName);
+		break;
+	}
+}
+
+
+int UBakeMeshAttributeMapsToolBase::SelectColorTextureToBake(const TArray<UTexture*>& Textures)
+{
+	TArray<int> TextureVotes;
+	TextureVotes.Init(0, Textures.Num());
+
+	for (int TextureIndex = 0; TextureIndex < Textures.Num(); ++TextureIndex)
+	{
+		UTexture* Tex = Textures[TextureIndex];
+		UTexture2D* Tex2D = Cast<UTexture2D>(Tex);
+
+		if (Tex2D)
+		{
+			// Texture uses SRGB
+			if (Tex->SRGB != 0)
+			{
+				++TextureVotes[TextureIndex];
+			}
+
+#if WITH_EDITORONLY_DATA
+			// Texture has multiple channels
+			ETextureSourceFormat Format = Tex->Source.GetFormat();
+			if (Format == TSF_BGRA8 || Format == TSF_BGRE8 || Format == TSF_RGBA16 || Format == TSF_RGBA16F)
+			{
+				++TextureVotes[TextureIndex];
+			}
+#endif
+
+			// What else? Largest texture? Most layers? Most mipmaps?
+		}
+	}
+
+	int MaxIndex = -1;
+	int MaxVotes = -1;
+	for (int TextureIndex = 0; TextureIndex < Textures.Num(); ++TextureIndex)
+	{
+		if (TextureVotes[TextureIndex] > MaxVotes)
+		{
+			MaxIndex = TextureIndex;
+			MaxVotes = TextureVotes[TextureIndex];
+		}
+	}
+
+	return MaxIndex;
+}
+
+
+void UBakeMeshAttributeMapsToolBase::InitializeEmptyMaps()
+{
+	FTexture2DBuilder NormalsBuilder;
+	NormalsBuilder.Initialize(FTexture2DBuilder::ETextureType::NormalMap, FImageDimensions(16, 16));
+	NormalsBuilder.Commit(false);
+	EmptyNormalMap = NormalsBuilder.GetTexture2D();
+
+	FTexture2DBuilder ColorBuilderBlack;
+	ColorBuilderBlack.Initialize(FTexture2DBuilder::ETextureType::Color, FImageDimensions(16, 16));
+	ColorBuilderBlack.Clear(FColor(0,0,0));
+	ColorBuilderBlack.Commit(false);
+	EmptyColorMapBlack = ColorBuilderBlack.GetTexture2D();
+
+	FTexture2DBuilder ColorBuilderWhite;
+	ColorBuilderWhite.Initialize(FTexture2DBuilder::ETextureType::Color, FImageDimensions(16, 16));
+	ColorBuilderWhite.Clear(FColor::White);
+	ColorBuilderWhite.Commit(false);
+	EmptyColorMapWhite = ColorBuilderWhite.GetTexture2D();
+}
+
+#undef LOCTEXT_NAMESPACE
diff --git a/Engine/Plugins/Experimental/MeshModelingToolsetExp/Source/MeshModelingToolsExp/Private/BakeMultiMeshAttributeMapsTool.cpp b/Engine/Plugins/Experimental/MeshModelingToolsetExp/Source/MeshModelingToolsExp/Private/BakeMultiMeshAttributeMapsTool.cpp
new file mode 100644
--- /dev/null
+++ b/Engine/Plugins/Experimental/MeshModelingToolsetExp/Source/MeshModelingToolsExp/Private/BakeMultiMeshAttributeMapsTool.cpp
@@ -0,0 +1,758 @@
+// Copyright Epic Games, Inc. All Rights Reserved.
+
+#include "BakeMultiMeshAttributeMapsTool.h"
+#include "InteractiveToolManager.h"
+#include "ToolBuilderUtil.h"
+
+#include "DynamicMesh/DynamicMesh3.h"
+#include "DynamicMesh/DynamicMeshAttributeSet.h"
+#include "DynamicMesh/MeshTransforms.h"
+#include "MeshDescriptionToDynamicMesh.h"
+#include "Sampling/MeshNormalMapEvaluator.h"
+#include "Sampling/MeshOcclusionMapEvaluator.h"
+#include "Sampling/MeshCurvatureMapEvaluator.h"
+#include "Sampling/MeshPropertyMapEvaluator.h"
+#include "Sampling/MeshResampleImageEvaluator.h"
+
+#include "ImageUtils.h"
+
+#include "AssetUtils/Texture2DBuilder.h"
+#include "AssetUtils/Texture2DUtil.h"
+#include "ModelingObjectsCreationAPI.h"
+
+#include "TargetInterfaces/MaterialProvider.h"
+#include "TargetInterfaces/MeshDescriptionProvider.h"
+#include "TargetInterfaces/PrimitiveComponentBackedTarget.h"
+#include "TargetInterfaces/StaticMeshBackedTarget.h"
+#include "ToolTargetManager.h"
+#include "ModelingToolTargetUtil.h"
+
+// required to pass UStaticMesh asset so we can save at same location
+#include "Engine/Classes/Components/StaticMeshComponent.h"
+#include "Engine/Classes/Engine/StaticMesh.h"
+
+#include "ExplicitUseGeometryMathTypes.h"		// using UE::Geometry::(math types)
+using namespace UE::Geometry;
+
+#define LOCTEXT_NAMESPACE "UBakeMultiMeshAttributeMapsTool"
+
+/*
+* ToolBuilder
+*/
+
+const FToolTargetTypeRequirements& UBakeMultiMeshAttributeMapsToolBuilder::GetTargetRequirements() const
+{
+	static FToolTargetTypeRequirements TypeRequirements({
+		UMeshDescriptionProvider::StaticClass(),
+		UPrimitiveComponentBackedTarget::StaticClass(),
+		UStaticMeshBackedTarget::StaticClass(),			// currently only supports StaticMesh targets
+		UMaterialProvider::StaticClass()
+		});
+	return TypeRequirements;
+}
+
+bool UBakeMultiMeshAttributeMapsToolBuilder::CanBuildTool(const FToolBuilderState& SceneState) const
+{
+	const int32 NumTargets = SceneState.TargetManager->CountSelectedAndTargetable(SceneState, GetTargetRequirements());
+	return (NumTargets > 1);
+}
+
+UInteractiveTool* UBakeMultiMeshAttributeMapsToolBuilder::BuildTool(const FToolBuilderState& SceneState) const
+{
+	UBakeMultiMeshAttributeMapsTool* NewTool = NewObject<UBakeMultiMeshAttributeMapsTool>(SceneState.ToolManager);
+	TArray<TObjectPtr<UToolTarget>> Targets = SceneState.TargetManager->BuildAllSelectedTargetable(SceneState, GetTargetRequirements());
+	NewTool->SetTargets(MoveTemp(Targets));
+	NewTool->SetWorld(SceneState.World);
+	return NewTool;
+}
+
+
+TArray<FString> UBakeMultiMeshAttributeMapsToolProperties::GetUVLayerNamesFunc()
+{
+	return UVLayerNamesList;
+}
+
+
+/**
+ * MeshSceneAdapter bake detail sampler for baking N detail meshes to 1 target mesh.
+ */
+class FMeshBakerMeshSceneSampler : public IMeshBakerDetailSampler
+{
+public:
+	using FDetailTextureMap = TMap<void*, FBakeDetailTexture>;
+
+public:
+	/**
+	 * Constructor.
+	 *
+	 * Input FMeshSceneAdapter's spatial evaluation cache is assumed to be built.
+	 */
+	FMeshBakerMeshSceneSampler(FMeshSceneAdapter* Scene)
+		: MeshScene(Scene)
+	{
+	}
+
+	// Begin IMeshBakerDetailSampler interface
+	virtual void ProcessMeshes(TFunctionRef<void(const void*)> ProcessFn) const override
+	{
+		auto ProcessChildMesh = [ProcessFn](const FActorAdapter*, const FActorChildMesh* ChildMesh)
+		{
+			if (ChildMesh)
+			{
+				ProcessFn(ChildMesh->MeshSpatial);
+			}
+		};
+		MeshScene->ProcessActorChildMeshes(ProcessChildMesh);
+	}
+
+	virtual void SetColorMap(const void* Mesh, const FBakeDetailTexture& Map) override
+	{
+		DetailColorMaps[Mesh] = Map;
+	}
+
+	virtual void SetNormalMap(const void* Mesh, const FBakeDetailTexture& Map) override
+	{
+		DetailNormalMaps[Mesh] = Map;
+	}
+
+	virtual const FBakeDetailTexture* GetColorMap(const void* Mesh) const override
+	{
+		return DetailColorMaps.Find(Mesh);
+	}
+	
+	virtual const FBakeDetailTexture* GetNormalMap(const void* Mesh) const override
+	{
+		return DetailNormalMaps.Find(Mesh);
+	}
+
+	virtual bool SupportsIdentityCorrespondence() const override
+	{
+		return false;
+	}
+
+	virtual bool SupportsNearestPointCorrespondence() const override
+	{
+		return false;
+	}
+
+	virtual bool SupportsRaycastCorrespondence() const override
+	{
+		return true;
+	}
+
+	virtual const void* FindNearestHitTriangle(
+		const FRay3d& Ray,
+		double& NearestT,
+		int& TriId,
+		FVector3d& TriBaryCoords,
+		const IMeshSpatial::FQueryOptions& Options = IMeshSpatial::FQueryOptions()) const override
+	{
+		// TODO: Pass-through max distance to FMeshSceneAdapter query
+		FMeshSceneRayHit HitResult;
+		const bool bHit = MeshScene->FindNearestRayIntersection(Ray, HitResult);
+		const void* HitMesh = nullptr;
+		if (bHit && NearestT < Options.MaxDistance)
+		{
+			TriId = HitResult.HitMeshTriIndex;
+			NearestT = HitResult.RayDistance;
+			TriBaryCoords = HitResult.HitMeshBaryCoords;
+			HitMesh = HitResult.HitMeshSpatialWrapper;
+		}
+		return HitMesh;
+	}
+
+	virtual bool TestAnyHitTriangle(
+		const FRay3d& Ray,
+		const IMeshSpatial::FQueryOptions& Options = IMeshSpatial::FQueryOptions()) const override
+	{
+		// TODO: Add proper test any hit triangle support (for Occlusion)
+		// TODO: Pass through max distance to FMeshSceneAdapter query
+		double NearestT = TNumericLimits<double>::Max();
+		int TriId = IndexConstants::InvalidID;
+		FVector3d TriBaryCoords;
+		const bool bHit = FindNearestHitTriangle(Ray, NearestT, TriId, TriBaryCoords, Options) != nullptr;
+		return bHit && NearestT < Options.MaxDistance;
+	}
+
+	virtual FAxisAlignedBox3d GetBounds() const override
+	{
+		return MeshScene->GetBoundingBox();
+	}
+
+	virtual bool IsTriangle(const void* Mesh, const int TriId) const override
+	{
+		const IMeshSpatialWrapper* Spatial = static_cast<const IMeshSpatialWrapper*>(Mesh);
+		return Spatial->IsTriangle(TriId);
+	}
+
+	virtual FIndex3i GetTriangle(const void* Mesh, const int TriId) const override
+	{
+		const IMeshSpatialWrapper* Spatial = static_cast<const IMeshSpatialWrapper*>(Mesh);
+		return Spatial->GetTriangle(TriId);
+	}
+
+	virtual FVector3d GetTriNormal(const void* Mesh, const int TriId) const override
+	{
+		// TODO
+		checkNoEntry();
+		return FVector3d::Zero();
+	}
+
+	virtual int32 GetMaterialID(const void* Mesh, const int TriId) const override
+	{
+		// TODO
+		checkNoEntry();
+		return IndexConstants::InvalidID;
+	}
+
+	virtual bool HasNormals(const void* Mesh) const override
+	{
+		const IMeshSpatialWrapper* Spatial = static_cast<const IMeshSpatialWrapper*>(Mesh);
+		return Spatial->HasNormals();
+	}
+
+	virtual bool HasUVs(const void* Mesh, const int UVLayer = 0) const override
+	{
+		const IMeshSpatialWrapper* Spatial = static_cast<const IMeshSpatialWrapper*>(Mesh);
+		return Spatial->HasUVs(UVLayer);
+	}
+
+	virtual bool HasTangents(const void* Mesh) const override
+	{
+		// TODO
+		checkNoEntry();
+		return false;
+	}
+
+	virtual bool HasColors(const void* Mesh) const override
+	{
+		// TODO
+		checkNoEntry();
+		return false;
+	}
+
+	virtual FVector3d TriBaryInterpolatePoint(
+		const void* Mesh,
+		const int32 TriId,
+		const FVector3d& BaryCoords) const override
+	{
+		const IMeshSpatialWrapper* Spatial = static_cast<const IMeshSpatialWrapper*>(Mesh);
+		return Spatial->TriBaryInterpolatePoint(TriId, BaryCoords);
+	}
+
+	virtual bool TriBaryInterpolateNormal(
+		const void* Mesh,
+		const int32 TriId,
+		const FVector3d& BaryCoords,
+		FVector3f& NormalOut) const override
+	{
+		const IMeshSpatialWrapper* Spatial = static_cast<const IMeshSpatialWrapper*>(Mesh);
+		return Spatial->TriBaryInterpolateNormal(TriId, BaryCoords, NormalOut);
+	}
+
+	virtual bool TriBaryInterpolateUV(
+		const void* Mesh,
+		const int32 TriId,
+		const FVector3d& BaryCoords,
+		const int UVLayer,
+		FVector2f& UVOut ) const override
+	{
+		const IMeshSpatialWrapper* Spatial = static_cast<const IMeshSpatialWrapper*>(Mesh);
+		return Spatial->TriBaryInterpolateUV(TriId, BaryCoords, UVLayer, UVOut);
+	}
+
+	virtual bool TriBaryInterpolateColor(
+		const void* Mesh,
+		const int32 TriId,
+		const FVector3d& BaryCoords,
+		FVector4f& ColorOut) const override
+	{
+		// TODO
+		checkNoEntry();
+		return false;
+	}
+
+	virtual bool TriBaryInterpolateTangents(
+		const void* Mesh,
+		const int32 TriId,
+		const FVector3d& BaryCoords,
+		FVector3d& TangentX,
+		FVector3d& TangentY) const override
+	{
+		// TODO
+		checkNoEntry();
+		return false;
+	}
+
+	virtual void GetCurvature(
+		const void* Mesh,
+		FMeshVertexCurvatureCache& CurvatureCache) const override
+	{
+		// TODO
+		checkNoEntry();
+	}
+	// End IMeshBakerDetailSampler interface
+
+	/** Initialize the mesh to color textures map */
+	void SetColorMaps(const FDetailTextureMap& Map)
+	{
+		DetailColorMaps = Map;
+	}
+
+	/** Initialize the mesh to normal textures map */
+	void SetNormalMaps(const FDetailTextureMap& Map)
+	{
+		DetailNormalMaps = Map;
+	}
+
+protected:
+	FMeshSceneAdapter* MeshScene = nullptr;
+	FDetailTextureMap DetailColorMaps;
+	FDetailTextureMap DetailNormalMaps;
+};
+
+
+/*
+ * Operators
+ */
+
+class FMultiMeshMapBakerOp : public TGenericDataOperator<FMeshMapBaker>
+{
+public:
+	// General bake settings
+	FMeshSceneAdapter* DetailMeshScene = nullptr;
+	UE::Geometry::FDynamicMesh3* BaseMesh;
+	TSharedPtr<UE::Geometry::FMeshTangentsd, ESPMode::ThreadSafe> BaseMeshTangents;
+	TUniquePtr<UE::Geometry::FMeshMapBaker> Baker;
+	UBakeMultiMeshAttributeMapsTool::FBakeCacheSettings BakeCacheSettings;
+
+	// Detail bake data
+	TArray<TSharedPtr<UE::Geometry::TImageBuilder<FVector4f>>> CachedColorImages;
+	UBakeMultiMeshAttributeMapsTool::FTextureImageMap CachedMeshToColorImageMap;
+
+	// Begin TGenericDataOperator interface
+	virtual void CalculateResult(FProgressCancel* Progress) override
+	{
+		Baker = MakeUnique<FMeshMapBaker>();
+		Baker->CancelF = [Progress]() {
+			return Progress && Progress->Cancelled();
+		};
+		Baker->SetTargetMesh(BaseMesh);
+		Baker->SetDimensions(BakeCacheSettings.Dimensions);
+		Baker->SetUVLayer(BakeCacheSettings.UVLayer);
+		Baker->SetThickness(BakeCacheSettings.Thickness);
+		Baker->SetMultisampling(BakeCacheSettings.Multisampling);
+		Baker->SetTargetMeshTangents(BaseMeshTangents);
+		
+		FMeshBakerMeshSceneSampler DetailSampler(DetailMeshScene);
+		Baker->SetDetailSampler(&DetailSampler);
+
+		for (const EBakeMapType MapType : ALL_BAKE_MAP_TYPES)
+		{
+			switch (BakeCacheSettings.BakeMapTypes & MapType)
+			{
+			case EBakeMapType::TangentSpaceNormalMap:
+			{
+				TSharedPtr<FMeshNormalMapEvaluator, ESPMode::ThreadSafe> NormalBaker = MakeShared<FMeshNormalMapEvaluator, ESPMode::ThreadSafe>();
+				Baker->AddBaker(NormalBaker);
+				break;	
+			}
+			case EBakeMapType::Texture2DImage:
+			{
+				TSharedPtr<FMeshResampleImageEvaluator, ESPMode::ThreadSafe> ResampleBaker = MakeShared<FMeshResampleImageEvaluator, ESPMode::ThreadSafe>();
+				DetailSampler.SetColorMaps(CachedMeshToColorImageMap);
+				Baker->AddBaker(ResampleBaker);
+				break;
+			}
+			default:
+				break;
+			}
+		}
+		Baker->Bake();
+		SetResult(MoveTemp(Baker));
+	}
+	// End TGenericDataOperator interface
+};
+
+
+/*
+ * Tool
+ */
+
+void UBakeMultiMeshAttributeMapsTool::Setup()
+{
+	Super::Setup();
+
+	// Initialize base mesh
+	const UE::Geometry::FTransform3d BaseToWorld = UE::ToolTarget::GetLocalToWorldTransform(Targets[0]);
+	PreviewMesh->ProcessMesh([this, BaseToWorld](const FDynamicMesh3& Mesh)
+	{
+		BaseMesh.Copy(Mesh);
+		BaseMeshTangents = MakeShared<FMeshTangentsd, ESPMode::ThreadSafe>(&BaseMesh);
+		BaseMeshTangents->CopyTriVertexTangents(Mesh);
+		
+		// FMeshSceneAdapter operates in world space, so ensure our mesh transformed to world.
+		MeshTransforms::ApplyTransform(BaseMesh, BaseToWorld);
+		BaseSpatial.SetMesh(&BaseMesh, true);
+	});
+
+	// Initialize detail sampler
+	const int NumTargets = Targets.Num();
+	TArray<UActorComponent*> DetailComponents;
+	for (int Idx = 1; Idx < NumTargets; ++Idx)
+	{
+		UPrimitiveComponent* Component = UE::ToolTarget::GetTargetComponent(Targets[Idx]);
+		if (Component)
+		{
+			DetailComponents.Add(Component);
+		}
+	}
+	DetailMeshScene.AddComponents(DetailComponents);
+	DetailMeshScene.Build(FMeshSceneAdapterBuildOptions());
+	DetailMeshScene.BuildSpatialEvaluationCache();
+
+	// Setup tool property sets
+	Settings = NewObject<UBakeMultiMeshAttributeMapsToolProperties>(this);
+	Settings->RestoreProperties(this);
+	Settings->UVLayerNamesList.Reset();
+	int32 FoundIndex = -1;
+	for (int32 k = 0; k < BaseMesh.Attributes()->NumUVLayers(); ++k)
+	{
+		Settings->UVLayerNamesList.Add(FString::FromInt(k));
+		if (Settings->UVLayer == Settings->UVLayerNamesList.Last())
+		{
+			FoundIndex = k;
+		}
+	}
+	if (FoundIndex == -1)
+	{
+		Settings->UVLayer = Settings->UVLayerNamesList[0];
+	}
+	AddToolPropertySource(Settings);
+
+	Settings->WatchProperty(Settings->MapTypes, [this](int32) { bInputsDirty = true; UpdateOnModeChange(); });
+	Settings->WatchProperty(Settings->MapPreview, [this](int32) { UpdateVisualization(); GetToolManager()->PostInvalidation(); });
+	Settings->WatchProperty(Settings->Resolution, [this](EBakeTextureResolution) { bInputsDirty = true; });
+	Settings->WatchProperty(Settings->UVLayer, [this](FString) { bInputsDirty = true; });
+	Settings->WatchProperty(Settings->Thickness, [this](float) { bInputsDirty = true; });
+	Settings->WatchProperty(Settings->Multisampling, [this](EBakeMultisampling) { bInputsDirty = true; });
+
+	DetailProps = NewObject<UBakeMultiMeshDetailToolProperties>(this);
+	DetailProps->RestoreProperties(this);
+	AddToolPropertySource(DetailProps);
+	SetToolPropertySourceEnabled(DetailProps, true);
+
+	// Pre-populate detail mesh data
+	TArray<UTexture2D*> DetailColorTextures;
+	for (int Idx = 1; Idx < NumTargets; ++Idx)
+	{
+		UToolTarget* DetailTarget = Targets[Idx];
+
+		// Hide each of our detail targets since this baker operates solely in world space
+		// which will occlude the preview of the target mesh. 
+		UE::ToolTarget::HideSourceObject(DetailTarget);
+		
+		IStaticMeshBackedTarget* DetailStaticMeshTarget = Cast<IStaticMeshBackedTarget>(DetailTarget);
+		UStaticMesh* DetailStaticMesh = DetailStaticMeshTarget ? DetailStaticMeshTarget->GetStaticMesh() : nullptr;
+		const UPrimitiveComponent* Component = UE::ToolTarget::GetTargetComponent(Targets[Idx]);
+		UTexture2D* DetailColorTexture = nullptr;
+		ProcessComponentTextures(Component, [&DetailColorTexture](const int MaterialID, const TArray<UTexture*>& Textures)
+		{
+			// TODO: Support multiple materialIDs per detail mesh
+			if (MaterialID == 0)
+			{
+				const int SelectedTextureIndex = SelectColorTextureToBake(Textures);
+				if (SelectedTextureIndex >= 0)
+				{
+					DetailColorTexture = Cast<UTexture2D>(Textures[SelectedTextureIndex]);
+				}
+			}
+		});
+
+		FBakeMultiMeshDetailProperties DetailMeshProp;
+		DetailMeshProp.DetailMesh = DetailStaticMesh;
+		DetailMeshProp.DetailColorMap = DetailColorTexture;
+		DetailProps->DetailProperties.Add(DetailMeshProp);
+		DetailProps->WatchProperty(DetailProps->DetailProperties[Idx-1].DetailColorMap, [this](UTexture2D*) { bInputsDirty = true; });
+		DetailProps->WatchProperty(DetailProps->DetailProperties[Idx-1].DetailColorMapUVLayer, [this](int) { bInputsDirty = true; });
+	}
+
+	UpdateOnModeChange();
+
+	bInputsDirty = true;
+
+	SetToolDisplayName(LOCTEXT("ToolName", "Bake Textures"));
+	GetToolManager()->DisplayMessage(
+		LOCTEXT("OnStartTool", "Bake Maps. Select Bake Mesh (LowPoly) first, then select Detail Meshes (HiPoly) to bake. Texture Assets will be created on Accept. "),
+		EToolMessageLevel::UserNotification);
+
+	SetupBaseToolProperties();
+}
+
+
+
+
+bool UBakeMultiMeshAttributeMapsTool::CanAccept() const
+{
+	bool bCanAccept = Compute ? Compute->HaveValidResult() && Settings->MapTypes != (int) EBakeMapType::None : false;
+	if (bCanAccept)
+	{
+		// Allow Accept if all non-None types have valid results.
+		const int NumResults = Settings->Result.Num();
+		for (int ResultIdx = 0; ResultIdx < NumResults; ++ResultIdx)
+		{
+			bCanAccept = bCanAccept && Settings->Result[ResultIdx];
+		}
+	}
+	return bCanAccept;
+}
+
+
+TUniquePtr<UE::Geometry::TGenericDataOperator<FMeshMapBaker>> UBakeMultiMeshAttributeMapsTool::MakeNewOperator()
+{
+	TUniquePtr<FMultiMeshMapBakerOp> Op = MakeUnique<FMultiMeshMapBakerOp>();
+	Op->DetailMeshScene = &DetailMeshScene;
+	Op->BaseMesh = &BaseMesh;
+	Op->BakeCacheSettings = CachedBakeCacheSettings;
+
+	constexpr EBakeMapType RequiresTangents = EBakeMapType::TangentSpaceNormalMap | EBakeMapType::BentNormal;
+	if ((bool)(CachedBakeCacheSettings.BakeMapTypes & RequiresTangents))
+	{
+		Op->BaseMeshTangents = BaseMeshTangents;
+	}
+
+	if ((bool)(CachedBakeCacheSettings.BakeMapTypes & EBakeMapType::Texture2DImage))
+	{
+		Op->CachedColorImages = CachedColorImages;
+		Op->CachedMeshToColorImageMap = CachedMeshToColorImagesMap;
+	}
+	return Op;
+}
+
+
+void UBakeMultiMeshAttributeMapsTool::Shutdown(EToolShutdownType ShutdownType)
+{
+	Super::Shutdown(ShutdownType);
+	
+	Settings->SaveProperties(this);
+	DetailProps->SaveProperties(this);
+
+	if (Compute)
+	{
+		Compute->Shutdown();
+	}
+
+	// Restore visibility of detail targets
+	const int NumTargets = Targets.Num();
+	for (int Idx = 1; Idx < NumTargets; ++Idx)
+	{
+		UE::ToolTarget::ShowSourceObject(Targets[Idx]);
+	}
+	
+	if (ShutdownType == EToolShutdownType::Accept)
+	{
+		UStaticMeshComponent* StaticMeshComponent = CastChecked<UStaticMeshComponent>(UE::ToolTarget::GetTargetComponent(Targets[0]));
+		UStaticMesh* StaticMeshAsset = StaticMeshComponent->GetStaticMesh();
+		check(StaticMeshAsset);
+		FString BaseName = UE::ToolTarget::GetTargetActor(Targets[0])->GetName();
+
+		bool bCreatedAssetOK = true;
+		const int NumResults = Settings->Result.Num();
+		for (int ResultIdx = 0; ResultIdx < NumResults; ++ResultIdx)
+		{
+			const FTexture2DBuilder::ETextureType TexType = GetTextureType(ResultTypes[ResultIdx]);
+			FString TexName;
+			GetTextureName(ResultTypes[ResultIdx], BaseName, TexName);
+
+			FTexture2DBuilder::CopyPlatformDataToSourceData(Settings->Result[ResultIdx], TexType);
+			bCreatedAssetOK = bCreatedAssetOK && UE::Modeling::CreateTextureObject(GetToolManager(), FCreateTextureObjectParams{ 0, StaticMeshAsset->GetWorld(), StaticMeshAsset, TexName, Settings->Result[ResultIdx] }).IsOK();
+		}
+		ensure(bCreatedAssetOK);
+	}
+}
+
+
+void UBakeMultiMeshAttributeMapsTool::UpdateResult()
+{
+	// bInputsDirty ensures that we only validate parameters once per param
+	// change. Parameter validation can be expensive (ex. UpdateResult_Texture2DImage).
+	if (!bInputsDirty)
+	{
+		return;
+	}
+
+	// clear warning (ugh)
+	GetToolManager()->DisplayMessage(FText(), EToolMessageLevel::UserWarning);
+
+	const int32 ImageSize = (int32)Settings->Resolution;
+	const FImageDimensions Dimensions(ImageSize, ImageSize);
+
+	FBakeCacheSettings BakeCacheSettings;
+	BakeCacheSettings.Dimensions = Dimensions;
+	BakeCacheSettings.UVLayer = FCString::Atoi(*Settings->UVLayer);
+	BakeCacheSettings.Thickness = Settings->Thickness;
+	BakeCacheSettings.Multisampling = (int32)Settings->Multisampling;
+
+	// process the raw bitfield before caching which may add additional targets.
+	BakeCacheSettings.BakeMapTypes = GetMapTypes(Settings->MapTypes);
+
+	// update bake cache settings
+	if (!(CachedBakeCacheSettings == BakeCacheSettings))
+	{
+		CachedBakeCacheSettings = BakeCacheSettings;
+		CachedDetailSettings = FBakeMultiMeshDetailSettings();
+	}
+
+	// Clear our invalid bitflag to check again for valid inputs.
+	OpState = EBakeOpState::Evaluate;
+
+	// Update map type settings
+	OpState |= UpdateResult_DetailMeshes();
+
+	// Early exit if op input parameters are invalid.
+	if ((bool)(OpState & EBakeOpState::Invalid))
+	{
+		return;
+	}
+
+	// This should be the only point of compute invalidation to
+	// minimize synchronization issues.
+	const bool bInvalidate = bInputsDirty || (bool)(OpState & EBakeOpState::Evaluate);
+	if (!Compute)
+	{
+		Compute = MakeUnique<TGenericDataBackgroundCompute<FMeshMapBaker>>();
+		Compute->Setup(this);
+		Compute->OnResultUpdated.AddLambda([this](const TUniquePtr<FMeshMapBaker>& NewResult) { OnMapsUpdated(NewResult); });
+		Compute->InvalidateResult();
+	}
+	else if (bInvalidate)
+	{
+		Compute->InvalidateResult();
+	}
+	bInputsDirty = false;
+}
+
+
+EBakeOpState UBakeMultiMeshAttributeMapsTool::UpdateResult_DetailMeshes()
+{
+	// This method will always force a re-evaluation.
+	EBakeOpState ResultState = EBakeOpState::Evaluate;
+
+	FBakeMultiMeshDetailSettings NewSettings;
+
+	// Iterate through our detail props to build our detail mesh data.
+	const int32 NumDetail = DetailProps->DetailProperties.Num();
+	CachedColorImages.SetNum(NumDetail);
+	CachedColorUVLayers.SetNum(NumDetail);
+	TMap<UActorComponent*, int> ActorToDataMap;
+	for (int Idx = 0; Idx < NumDetail; ++Idx)
+	{
+		UActorComponent* ActorComponent = UE::ToolTarget::GetTargetComponent(Targets[Idx+1]);
+		ActorToDataMap.Emplace(ActorComponent, Idx);
+		
+		// Color map data
+		if ((bool)(CachedBakeCacheSettings.BakeMapTypes & EBakeMapType::Texture2DImage))
+		{
+			UTexture2D* ColorMapSourceTexture = DetailProps->DetailProperties[Idx].DetailColorMap;
+			const int ColorMapUVLayer = DetailProps->DetailProperties[Idx].DetailColorMapUVLayer;
+			if (!ColorMapSourceTexture)
+			{
+				GetToolManager()->DisplayMessage(LOCTEXT("InvalidTextureWarning", "The Source Texture is not valid"), EToolMessageLevel::UserWarning);
+				return EBakeOpState::Invalid;
+			}
+
+			TSharedPtr<UE::Geometry::TImageBuilder<FVector4f>, ESPMode::ThreadSafe> ColorTextureImage = MakeShared<UE::Geometry::TImageBuilder<FVector4f>, ESPMode::ThreadSafe>();
+			if (!UE::AssetUtils::ReadTexture(ColorMapSourceTexture, *ColorTextureImage, bPreferPlatformData))
+			{
+				GetToolManager()->DisplayMessage(LOCTEXT("CannotReadTextureWarning", "Cannot read from the source texture"), EToolMessageLevel::UserWarning);
+				return EBakeOpState::Invalid;
+			}
+			CachedColorImages[Idx] = ColorTextureImage;
+			CachedColorUVLayers[Idx] = ColorMapUVLayer;
+		}
+	}
+
+	// Iterate through mesh scene adapter and build mesh to data maps.
+	CachedMeshToColorImagesMap.Empty();
+	auto BuildMeshToDataMaps = [this, ActorToDataMap](const FActorAdapter*, const FActorChildMesh* ChildMesh)
+	{
+		if (ChildMesh)
+		{
+			if (const int* DataIndex = ActorToDataMap.Find(ChildMesh->SourceComponent))
+			{
+				if ((bool)(CachedBakeCacheSettings.BakeMapTypes & EBakeMapType::Texture2DImage))
+				{
+					CachedMeshToColorImagesMap.Emplace(
+					ChildMesh->MeshSpatial,
+					FTextureImageData(CachedColorImages[*DataIndex].Get(), CachedColorUVLayers[*DataIndex]));
+				}
+			}
+		}
+	};
+	DetailMeshScene.ProcessActorChildMeshes(BuildMeshToDataMaps);
+	
+	return ResultState;
+}
+
+
+void UBakeMultiMeshAttributeMapsTool::UpdateVisualization()
+{
+	PreviewMesh->SetOverrideRenderMaterial(PreviewMaterial);
+
+	// Map CachedMaps to Settings->Result
+	const int NumResults = Settings->Result.Num();
+	for (int ResultIdx = 0; ResultIdx < NumResults; ResultIdx++)
+	{
+		Settings->Result[ResultIdx] = CachedMaps[CachedMapIndices[ResultTypes[ResultIdx]]];
+	}
+
+	// Set the preview material according to the preview index.
+	if (Settings->MapPreview >= 0 && Settings->MapPreview < Settings->Result.Num())
+	{
+		UpdatePreview(Settings->MapPreview);
+	}
+}
+
+
+void UBakeMultiMeshAttributeMapsTool::UpdateOnModeChange()
+{
+	OnMapTypesUpdated(Settings->MapTypes);
+	
+	for (const EBakeMapType MapType : ALL_BAKE_MAP_TYPES)
+	{
+		switch ((EBakeMapType)Settings->MapTypes & MapType)
+		{
+		case EBakeMapType::TangentSpaceNormalMap:
+			//SetToolPropertySourceEnabled(NormalMapProps, true);
+			break;
+		case EBakeMapType::AmbientOcclusion:
+		case EBakeMapType::BentNormal:
+		case EBakeMapType::Occlusion:
+			//SetToolPropertySourceEnabled(OcclusionMapProps, true);
+			break;
+		case EBakeMapType::Curvature:
+			//SetToolPropertySourceEnabled(CurvatureMapProps, true);
+			break;
+		case EBakeMapType::NormalImage:
+		case EBakeMapType::FaceNormalImage:
+		case EBakeMapType::PositionImage:
+		case EBakeMapType::MaterialID:
+		case EBakeMapType::VertexColorImage:
+			break;
+		case EBakeMapType::Texture2DImage:
+			//SetToolPropertySourceEnabled(Texture2DProps, true);
+			break;
+		case EBakeMapType::MultiTexture:
+			//SetToolPropertySourceEnabled(MultiTextureProps, true);
+			break;
+		default:
+			break;
+		}
+	}
+	
+	Settings->Result.Empty();
+	Settings->Result.SetNum(ResultTypes.Num());
+}
+
+
+#undef LOCTEXT_NAMESPACE
diff --git a/Engine/Plugins/Experimental/MeshModelingToolsetExp/Source/MeshModelingToolsExp/Public/BakeMeshAttributeMapsTool.h b/Engine/Plugins/Experimental/MeshModelingToolsetExp/Source/MeshModelingToolsExp/Public/BakeMeshAttributeMapsTool.h
--- a/Engine/Plugins/Experimental/MeshModelingToolsetExp/Source/MeshModelingToolsExp/Public/BakeMeshAttributeMapsTool.h
+++ b/Engine/Plugins/Experimental/MeshModelingToolsetExp/Source/MeshModelingToolsExp/Public/BakeMeshAttributeMapsTool.h
@@ -11,32 +11,23 @@
 #include "Image/ImageBuilder.h"
 #include "Sampling/MeshMapBaker.h"
 #include "ModelingOperators.h"
-#include "MeshOpPreviewHelpers.h"
 #include "PreviewMesh.h"
-#include "BakeMeshAttributeToolCommon.h"
+#include "BakeMeshAttributeMapsToolBase.h"
 #include "BakeMeshAttributeMapsTool.generated.h"
 
 
 // predeclarations
-struct FMeshDescription;
-class UDynamicMeshComponent;
 class UMaterialInstanceDynamic;
 class UTexture2D;
 PREDECLARE_GEOMETRY(template<typename RealType> class TMeshTangents);
 PREDECLARE_GEOMETRY(class FMeshImageBakingCache);
 using UE::Geometry::FImageDimensions;
-class IPrimitiveComponentBackedTarget;
-class IMeshDescriptionProvider;
-class IMaterialProvider;
-class FBakeNormalMapOp;
-class FBakeOcclusionMapOp;
-class FBakeCurvatureMapOp;
-class FBakeMeshPropertyMapOp;
-class FBakeTexture2DImageMapOp;
+
 
 /**
  * Tool Builder
  */
+
 UCLASS()
 class MESHMODELINGTOOLSEXP_API UBakeMeshAttributeMapsToolBuilder : public UInteractiveToolBuilder
 {
@@ -52,53 +43,8 @@ protected:
 
 
 
-UENUM(meta=(Bitflags, UseEnumValuesAsMaskValuesInEditor="true"))
-enum class EBakeMapType
-{
-	None                   = 0 UMETA(Hidden),
-	TangentSpaceNormalMap  = 1 << 0,
-	AmbientOcclusion       = 1 << 1,
-	BentNormal             = 1 << 2,
-	Curvature              = 1 << 3,
-	Texture2DImage         = 1 << 4,
-	NormalImage            = 1 << 5,
-	FaceNormalImage        = 1 << 6,
-	PositionImage          = 1 << 7,
-	MaterialID             = 1 << 8,
-	MultiTexture           = 1 << 9,
-	VertexColorImage       = 1 << 10,
-	Occlusion              = (AmbientOcclusion | BentNormal) UMETA(Hidden),
-	All                    = 0x7FF UMETA(Hidden)
-};
-ENUM_CLASS_FLAGS(EBakeMapType);
-
-
-UENUM()
-enum class EBakeTextureResolution 
-{
-	Resolution16 = 16 UMETA(DisplayName = "16 x 16"),
-	Resolution32 = 32 UMETA(DisplayName = "32 x 32"),
-	Resolution64 = 64 UMETA(DisplayName = "64 x 64"),
-	Resolution128 = 128 UMETA(DisplayName = "128 x 128"),
-	Resolution256 = 256 UMETA(DisplayName = "256 x 256"),
-	Resolution512 = 512 UMETA(DisplayName = "512 x 512"),
-	Resolution1024 = 1024 UMETA(DisplayName = "1024 x 1024"),
-	Resolution2048 = 2048 UMETA(DisplayName = "2048 x 2048"),
-	Resolution4096 = 4096 UMETA(DisplayName = "4096 x 4096"),
-	Resolution8192 = 8192 UMETA(DisplayName = "8192 x 8192")
-};
 
 
-UENUM()
-enum class EBakeMultisampling
-{
-	None = 1 UMETA(DisplayName = "None"),
-	Sample2x2 = 2 UMETA(DisplayName = "2 x 2"),
-	Sample4x4 = 4 UMETA(DisplayName = "4 x 4"),
-	Sample8x8 = 8 UMETA(DisplayName = "8 x 8"),
-	Sample16x16 = 16 UMETA(DisplayName = "16 x 16")
-};
-
 
 UCLASS()
 class MESHMODELINGTOOLSEXP_API UBakeMeshAttributeMapsToolProperties : public UInteractiveToolPropertySet
@@ -111,7 +57,7 @@ public:
 	int32 MapTypes = (int32) EBakeMapType::None;
 
 	/** The map type index to preview */
-	UPROPERTY(EditAnywhere, Category = MapSettings, meta=(ArrayClamp="Result"))
+	UPROPERTY(EditAnywhere, Category = MapSettings, meta=(ArrayClamp="Result", TransientToolProperty))
 	int MapPreview = 0;
 
 	/** The pixel resolution of the generated map */
@@ -148,27 +94,25 @@ public:
  * Detail Map Baking Tool
  */
 UCLASS()
-class MESHMODELINGTOOLSEXP_API UBakeMeshAttributeMapsTool : public UMultiSelectionTool, public UE::Geometry::IGenericDataOperatorFactory<UE::Geometry::FMeshMapBaker>
+class MESHMODELINGTOOLSEXP_API UBakeMeshAttributeMapsTool : public UBakeMeshAttributeMapsToolBase
 {
 	GENERATED_BODY()
 
 public:
 	UBakeMeshAttributeMapsTool() = default;
 
+	// Begin UInteractiveTool interface
 	virtual void Setup() override;
 	virtual void Shutdown(EToolShutdownType ShutdownType) override;
 
-	virtual void OnTick(float DeltaTime) override;
-	virtual void Render(IToolsContextRenderAPI* RenderAPI) override;
-
 	virtual bool HasCancel() const override { return true; }
 	virtual bool HasAccept() const override { return true; }
 	virtual bool CanAccept() const override;
+	// End UInteractiveTool interface
 
-	// IGenericDataOperatorFactory API
+	// Begin IGenericDataOperatorFactory interface
 	virtual TUniquePtr<UE::Geometry::TGenericDataOperator<UE::Geometry::FMeshMapBaker>> MakeNewOperator() override;
-
-	void SetWorld(UWorld* World);
+	// End IGenericDataOperatorFactory interface
 
 protected:
 	// need to update bResultValid if these are modified, so we don't publicly expose them. 
@@ -195,31 +139,14 @@ protected:
 	UPROPERTY()
 	TObjectPtr<UBakedMultiTexture2DImageProperties> MultiTextureProps;
 
-	UPROPERTY()
-	TObjectPtr<UBakedOcclusionMapVisualizationProperties> VisualizationProps;
-
+protected:
+	// Begin UBakeMeshAttributeMapsToolBase interface
+	virtual void UpdateResult() override;
+	virtual void UpdateVisualization() override;
+	// End UBakeMeshAttributeMapsToolBase interface
 
 protected:
 	friend class FMeshMapBakerOp;
-	
-	UPROPERTY()
-	TObjectPtr<UPreviewMesh> PreviewMesh;
-
-	UPROPERTY()
-	TObjectPtr<UMaterialInstanceDynamic> PreviewMaterial;
-
-    UPROPERTY()
-	TObjectPtr<UMaterialInstanceDynamic> BentNormalPreviewMaterial;
-
-	UPROPERTY()
-	TObjectPtr<UMaterialInstanceDynamic> WorkingPreviewMaterial;
-	float SecondsBeforeWorkingMaterial = 0.75;
-
-	UWorld* TargetWorld = nullptr;
-
-	TSharedPtr<UE::Geometry::TMeshTangents<double>, ESPMode::ThreadSafe> BaseMeshTangents;
-	UE::Geometry::FDynamicMesh3 BaseMesh;
-	UE::Geometry::FDynamicMeshAABBTree3 BaseSpatial;
 
 	bool bIsBakeToSelf = false;
 
@@ -231,43 +158,7 @@ protected:
 	bool bDetailMeshValid = false;
 
 	bool bInputsDirty = false;
-	void UpdateResult();
-
 	void UpdateOnModeChange();
-	void UpdateVisualization();
-
-	TUniquePtr<TGenericDataBackgroundCompute<UE::Geometry::FMeshMapBaker>> Compute = nullptr;
-	void OnMapsUpdated(const TUniquePtr<UE::Geometry::FMeshMapBaker>& NewResult);
-
-	/** @return A single bitfield of map types from an array of map types. */
-	EBakeMapType GetMapTypes(const int32& MapTypes) const;
-	TArray<EBakeMapType> GetMapTypesArray(const int32& MapTypes) const;
-
-	struct FBakeCacheSettings
-	{
-		EBakeMapType BakeMapTypes = EBakeMapType::None;
-		FImageDimensions Dimensions;
-		int32 UVLayer;
-		int32 DetailTimestamp;
-		float Thickness;
-		int32 Multisampling;
-
-		bool operator==(const FBakeCacheSettings& Other) const
-		{
-			return BakeMapTypes == Other.BakeMapTypes && Dimensions == Other.Dimensions &&
-				UVLayer == Other.UVLayer && DetailTimestamp == Other.DetailTimestamp &&
-				Thickness == Other.Thickness && Multisampling == Other.Multisampling;
-		}
-	};
-	FBakeCacheSettings CachedBakeCacheSettings;
-	TArray<EBakeMapType> ResultTypes;
-
-	EBakeOpState OpState = EBakeOpState::Evaluate;
-
-	UPROPERTY()
-	TArray<TObjectPtr<UTexture2D>> CachedMaps;
-	using CachedMapIndex = TMap<EBakeMapType, int32>;
-	CachedMapIndex CachedMapIndices;
 
 	FDetailMeshSettings CachedDetailMeshSettings;
 	TSharedPtr<UE::Geometry::TImageBuilder<FVector4f>, ESPMode::ThreadSafe> CachedDetailNormalMap;
@@ -287,25 +178,9 @@ protected:
 
 	TSharedPtr<UE::Geometry::TImageBuilder<FVector4f>, ESPMode::ThreadSafe> CachedTextureImage;
 	FTexture2DImageSettings CachedTexture2DImageSettings;
-	bool bPreferPlatformData = true;
 	EBakeOpState UpdateResult_Texture2DImage();
 
 	TMap<int32, TSharedPtr<UE::Geometry::TImageBuilder<FVector4f>, ESPMode::ThreadSafe>> CachedMultiTextures;
 	EBakeOpState UpdateResult_MultiTexture();
-
-
-	// empty maps are shown when nothing is computed
-	UPROPERTY()
-	TObjectPtr<UTexture2D> EmptyNormalMap;
-
-	UPROPERTY()
-	TObjectPtr<UTexture2D> EmptyColorMapBlack;
-
-	UPROPERTY()
-	TObjectPtr<UTexture2D> EmptyColorMapWhite;
-
-	void InitializeEmptyMaps();
-
-	void GetTexturesFromDetailMesh(const UPrimitiveComponent* DetailComponent);
-
 };
+
diff --git a/Engine/Plugins/Experimental/MeshModelingToolsetExp/Source/MeshModelingToolsExp/Public/BakeMeshAttributeMapsToolBase.h b/Engine/Plugins/Experimental/MeshModelingToolsetExp/Source/MeshModelingToolsExp/Public/BakeMeshAttributeMapsToolBase.h
new file mode 100644
--- /dev/null
+++ b/Engine/Plugins/Experimental/MeshModelingToolsetExp/Source/MeshModelingToolsExp/Public/BakeMeshAttributeMapsToolBase.h
@@ -0,0 +1,304 @@
+// Copyright Epic Games, Inc. All Rights Reserved.
+
+#pragma once
+
+#include "CoreMinimal.h"
+#include "Engine/Classes/Engine/Texture2D.h"
+#include "Materials/MaterialInstanceDynamic.h"
+#include "MultiSelectionTool.h"
+#include "InteractiveToolManager.h"
+#include "AssetUtils/Texture2DBuilder.h"
+#include "DynamicMesh/DynamicMesh3.h"
+#include "DynamicMesh/DynamicMeshAABBTree3.h"
+#include "Image/ImageDimensions.h"
+#include "Sampling/MeshMapBaker.h"
+#include "ModelingOperators.h"
+#include "MeshOpPreviewHelpers.h"
+#include "PreviewMesh.h"
+#include "BakeMeshAttributeToolCommon.h"
+#include "BakeMeshAttributeMapsToolBase.generated.h"
+
+/**
+* Bake maps enums
+*/
+
+
+UENUM(meta=(Bitflags, UseEnumValuesAsMaskValuesInEditor="true"))
+enum class EBakeMapType
+{
+	None                   = 0 UMETA(Hidden),
+	TangentSpaceNormalMap  = 1 << 0,
+	AmbientOcclusion       = 1 << 1,
+	BentNormal             = 1 << 2,
+	Curvature              = 1 << 3,
+	Texture2DImage         = 1 << 4,
+	NormalImage            = 1 << 5,
+	FaceNormalImage        = 1 << 6,
+	PositionImage          = 1 << 7,
+	MaterialID             = 1 << 8,
+	MultiTexture           = 1 << 9,
+	VertexColorImage       = 1 << 10,
+	Occlusion              = (AmbientOcclusion | BentNormal) UMETA(Hidden),
+	All                    = 0x7FF UMETA(Hidden)
+};
+ENUM_CLASS_FLAGS(EBakeMapType);
+
+
+// Only include the Occlusion bitmask rather than its components
+// (AmbientOcclusion | BentNormal). Since the Occlusion baker can
+// bake both types in a single pass, only iterating over the Occlusion
+// bitmask gives direct access to both types without the need to
+// externally track if we've handled the Occlusion evaluator in a prior
+// iteration loop.
+static constexpr EBakeMapType ALL_BAKE_MAP_TYPES[] =
+{
+	EBakeMapType::TangentSpaceNormalMap,
+	EBakeMapType::Occlusion, // (AmbientOcclusion | BentNormal)
+	EBakeMapType::Curvature,
+	EBakeMapType::Texture2DImage,
+	EBakeMapType::NormalImage,
+	EBakeMapType::FaceNormalImage,
+	EBakeMapType::PositionImage,
+	EBakeMapType::MaterialID,
+	EBakeMapType::MultiTexture,
+	EBakeMapType::VertexColorImage
+};
+
+
+UENUM()
+enum class EBakeTextureResolution 
+{
+	Resolution16 = 16 UMETA(DisplayName = "16 x 16"),
+	Resolution32 = 32 UMETA(DisplayName = "32 x 32"),
+	Resolution64 = 64 UMETA(DisplayName = "64 x 64"),
+	Resolution128 = 128 UMETA(DisplayName = "128 x 128"),
+	Resolution256 = 256 UMETA(DisplayName = "256 x 256"),
+	Resolution512 = 512 UMETA(DisplayName = "512 x 512"),
+	Resolution1024 = 1024 UMETA(DisplayName = "1024 x 1024"),
+	Resolution2048 = 2048 UMETA(DisplayName = "2048 x 2048"),
+	Resolution4096 = 4096 UMETA(DisplayName = "4096 x 4096"),
+	Resolution8192 = 8192 UMETA(DisplayName = "8192 x 8192")
+};
+
+
+UENUM()
+enum class EBakeMultisampling
+{
+	None = 1 UMETA(DisplayName = "None"),
+	Sample2x2 = 2 UMETA(DisplayName = "2 x 2"),
+	Sample4x4 = 4 UMETA(DisplayName = "4 x 4"),
+	Sample8x8 = 8 UMETA(DisplayName = "8 x 8"),
+	Sample16x16 = 16 UMETA(DisplayName = "16 x 16")
+};
+
+
+// TODO: Refactor common bake tool functionality with vertex bakes into common base class.
+/**
+ * Base Bake Maps tool
+ */
+UCLASS()
+class MESHMODELINGTOOLSEXP_API UBakeMeshAttributeMapsToolBase : public UMultiSelectionTool, public UE::Geometry::IGenericDataOperatorFactory<UE::Geometry::FMeshMapBaker>
+{
+	GENERATED_BODY()
+
+public:
+	UBakeMeshAttributeMapsToolBase() = default;
+
+	// Begin UInteractiveTool interface
+	virtual void Setup() override;
+	virtual void Shutdown(EToolShutdownType ShutdownType) override;
+
+	virtual void OnTick(float DeltaTime) override;
+	virtual void Render(IToolsContextRenderAPI* RenderAPI) override;
+
+	virtual bool HasCancel() const override { return true; }
+	virtual bool HasAccept() const override { return true; }
+	virtual bool CanAccept() const override { return true; }
+	// End UInteractiveTool interface
+
+	// Begin IGenericDataOperatorFactory interface
+	virtual TUniquePtr<UE::Geometry::TGenericDataOperator<UE::Geometry::FMeshMapBaker>> MakeNewOperator() override;
+	// End IGenericDataOperatorFactory interface
+
+	void SetWorld(UWorld* World);
+
+protected:
+	//
+	// Tool property sets
+	//
+	UPROPERTY()
+	TObjectPtr<UBakedOcclusionMapVisualizationProperties> VisualizationProps;
+
+	//
+	// Preview mesh and materials
+	//
+	UPROPERTY()
+	TObjectPtr<UPreviewMesh> PreviewMesh;
+
+	UPROPERTY()
+	TObjectPtr<UMaterialInstanceDynamic> PreviewMaterial;
+
+	UPROPERTY()
+	TObjectPtr<UMaterialInstanceDynamic> BentNormalPreviewMaterial;
+
+	UPROPERTY()
+	TObjectPtr<UMaterialInstanceDynamic> WorkingPreviewMaterial;
+	float SecondsBeforeWorkingMaterial = 0.75;
+
+protected:
+	/** To be invoked at end of client Setup methods. */
+	void SetupBaseToolProperties();
+	
+	/**
+	 * Process dirty props and update background compute.
+	 * Invoked during Render.
+	 */
+	virtual void UpdateResult();
+
+	/**
+	 * Updates the preview material on the preview mesh with the
+	 * computed results. Invoked by OnMapsUpdated.
+	 */
+	virtual void UpdateVisualization();
+
+
+protected:
+	//
+	// Bake parameters
+	//
+	UPROPERTY()
+	TObjectPtr<UWorld> TargetWorld = nullptr;
+
+	UE::Geometry::FDynamicMesh3 BaseMesh;
+	TSharedPtr<UE::Geometry::TMeshTangents<double>, ESPMode::ThreadSafe> BaseMeshTangents;
+	UE::Geometry::FDynamicMeshAABBTree3 BaseSpatial;
+	
+	EBakeOpState OpState = EBakeOpState::Evaluate;
+
+	struct FBakeCacheSettings
+	{
+		EBakeMapType BakeMapTypes = EBakeMapType::None;
+		FImageDimensions Dimensions;
+		int32 UVLayer = 0;
+		int32 DetailTimestamp = 0;
+		float Thickness = 3.0;
+		int32 Multisampling = 1;
+
+		bool operator==(const FBakeCacheSettings& Other) const
+		{
+			return BakeMapTypes == Other.BakeMapTypes && Dimensions == Other.Dimensions &&
+				UVLayer == Other.UVLayer && DetailTimestamp == Other.DetailTimestamp &&
+				Thickness == Other.Thickness && Multisampling == Other.Multisampling;
+		}
+	};
+	FBakeCacheSettings CachedBakeCacheSettings;
+
+	/**
+	 * To be invoked by client when bake map types change.
+	 */
+	void OnMapTypesUpdated(int32 MapTypes);
+
+	//
+	// Background compute
+	//
+	TUniquePtr<TGenericDataBackgroundCompute<UE::Geometry::FMeshMapBaker>> Compute = nullptr;
+
+	/**
+	 * Retrieves the result of the FMeshMapBaker and generates UTexture2D into the CachedMaps.
+	 * It is the responsibility of the client to ensure that CachedMaps is appropriately sized for
+	 * the range of index values in MapIndex.
+	 * 
+	 * @param NewResult the resulting FMeshMapBaker from the background Compute
+	 */
+	void OnMapsUpdated(const TUniquePtr<UE::Geometry::FMeshMapBaker>& NewResult);
+
+	/**
+	 * Update the preview material parameters for a given a result index.
+	 * @param PreviewIdx index into the ResultTypes array to preview
+	 */
+	void UpdatePreview(int PreviewIdx);
+
+	/** Internal cache of bake texture results. */
+	UPROPERTY()
+	TArray<TObjectPtr<UTexture2D>> CachedMaps;
+
+	/** Internal map of map type to index into CachedMaps array */
+	TMap<EBakeMapType, int32> CachedMapIndices;
+
+	/** Internal array of CachedMaps index to map type */
+	TArray<EBakeMapType> ResultTypes;
+	
+	//
+	// Utilities
+	//
+	const bool bPreferPlatformData = true;
+	
+	/** @return An enumerated map type from a bitfield */
+	static EBakeMapType GetMapTypes(const int32& MapTypes);
+
+	/** @return An array of map types from a bitfield */
+	static TArray<EBakeMapType> GetMapTypesArray(const int32& MapTypes);
+
+	/** @return the Texture2D type for a given map type */
+	static UE::Geometry::FTexture2DBuilder::ETextureType GetTextureType(EBakeMapType MapType);
+
+	/** @return the texture name given a base name and map type */
+	static void GetTextureName(EBakeMapType MapType, const FString& BaseName, FString& TexName);
+
+	/**
+	 * Given an array of textures associated with a material,
+	 * use heuristics to identify the color/albedo texture.
+	 * @param Textures array of textures associated with a material.
+	 * @return integer index into the Textures array representing the color/albedo texture
+	 */
+	static int SelectColorTextureToBake(const TArray<UTexture*>& Textures);
+
+	/**
+	 * Iterate through a primitive component's textures by material ID.
+	 * @param Component the component to query
+	 * @param ProcessFunc enumeration function with signature: void(int MaterialID, const TArray<UTexture*>& Textures)
+	 */
+	template <typename ProcessFn>
+	static void ProcessComponentTextures(const UPrimitiveComponent* Component, ProcessFn&& ProcessFunc);
+
+	// empty maps are shown when nothing is computed
+	UPROPERTY()
+	TObjectPtr<UTexture2D> EmptyNormalMap;
+
+	UPROPERTY()
+	TObjectPtr<UTexture2D> EmptyColorMapBlack;
+
+	UPROPERTY()
+	TObjectPtr<UTexture2D> EmptyColorMapWhite;
+
+	void InitializeEmptyMaps();
+};
+
+
+template <typename ProcessFn>
+void UBakeMeshAttributeMapsToolBase::ProcessComponentTextures(const UPrimitiveComponent* Component, ProcessFn&& ProcessFunc)
+{
+	if (!Component)
+	{
+		return;
+	}
+
+	TArray<UMaterialInterface*> Materials;
+	Component->GetUsedMaterials(Materials);
+	
+	for (int32 MaterialID = 0; MaterialID < Materials.Num(); ++MaterialID)	// TODO: This won't match MaterialIDs on the FDynamicMesh3 in general, will it?
+	{
+		UMaterialInterface* MaterialInterface = Materials[MaterialID];
+		if (MaterialInterface == nullptr)
+		{
+			continue;
+		}
+
+		TArray<UTexture*> Textures;
+		MaterialInterface->GetUsedTextures(Textures, EMaterialQualityLevel::High, true, ERHIFeatureLevel::SM5, true);
+		ProcessFunc(MaterialID, Textures);
+	}
+}
+
+
+
diff --git a/Engine/Plugins/Experimental/MeshModelingToolsetExp/Source/MeshModelingToolsExp/Public/BakeMultiMeshAttributeMapsTool.h b/Engine/Plugins/Experimental/MeshModelingToolsetExp/Source/MeshModelingToolsExp/Public/BakeMultiMeshAttributeMapsTool.h
new file mode 100644
--- /dev/null
+++ b/Engine/Plugins/Experimental/MeshModelingToolsetExp/Source/MeshModelingToolsExp/Public/BakeMultiMeshAttributeMapsTool.h
@@ -0,0 +1,208 @@
+// Copyright Epic Games, Inc. All Rights Reserved.
+
+#pragma once
+
+#include "CoreMinimal.h"
+#include "MultiSelectionTool.h"
+#include "InteractiveToolBuilder.h"
+#include "DynamicMesh/DynamicMesh3.h"
+#include "Image/ImageDimensions.h"
+#include "Image/ImageBuilder.h"
+#include "Sampling/MeshMapBaker.h"
+#include "Scene/MeshSceneAdapter.h"
+#include "ModelingOperators.h"
+#include "MeshOpPreviewHelpers.h"
+#include "PreviewMesh.h"
+#include "BakeMeshAttributeMapsToolBase.h"
+#include "BakeMultiMeshAttributeMapsTool.generated.h"
+
+/**
+ * Tool Builder
+ */
+UCLASS()
+class MESHMODELINGTOOLSEXP_API UBakeMultiMeshAttributeMapsToolBuilder : public UInteractiveToolBuilder
+{
+	GENERATED_BODY()
+
+	public:
+	virtual bool CanBuildTool(const FToolBuilderState& SceneState) const override;
+	virtual UInteractiveTool* BuildTool(const FToolBuilderState& SceneState) const override;
+
+	protected:
+	virtual const FToolTargetTypeRequirements& GetTargetRequirements() const override;
+};
+
+
+// This enumeration must match EBakeMapType. This duplicate enum is
+// intended to only be used for UI and internally converted to EBakeMapType.
+UENUM(meta=(Bitflags, UseEnumValuesAsMaskValuesInEditor="true"))
+enum class EBakeMultiMapType
+{
+	None                   = 0 UMETA(Hidden),
+	TangentSpaceNormalMap  = 1 << 0,
+	AmbientOcclusion       = 1 << 1 UMETA(Hidden),
+	BentNormal             = 1 << 2 UMETA(Hidden),
+	Curvature              = 1 << 3 UMETA(Hidden),
+	Texture2DImage         = 1 << 4,
+	NormalImage            = 1 << 5 UMETA(Hidden),
+	FaceNormalImage        = 1 << 6 UMETA(Hidden),
+	PositionImage          = 1 << 7 UMETA(Hidden),
+	MaterialID             = 1 << 8 UMETA(Hidden),
+	MultiTexture           = 1 << 9 UMETA(Hidden),
+	VertexColorImage       = 1 << 10 UMETA(Hidden),
+	Occlusion              = (AmbientOcclusion | BentNormal) UMETA(Hidden),
+	All                    = 0x7FF UMETA(Hidden)
+};
+ENUM_CLASS_FLAGS(EBakeMultiMapType);
+
+
+UCLASS()
+class MESHMODELINGTOOLSEXP_API UBakeMultiMeshAttributeMapsToolProperties : public UInteractiveToolPropertySet
+{
+	GENERATED_BODY()
+
+public:
+	/** The map types to generate */
+	UPROPERTY(EditAnywhere, Category = MapSettings, meta=(Bitmask, BitmaskEnum=EBakeMultiMapType))
+	int32 MapTypes = (int32) EBakeMapType::None;
+
+	/** The map type index to preview */
+	UPROPERTY(EditAnywhere, Category = MapSettings, meta=(ArrayClamp="Result"))
+	int MapPreview = 0;
+
+	/** The pixel resolution of the generated map */
+	UPROPERTY(EditAnywhere, Category = MapSettings, meta = (TransientToolProperty))
+	EBakeTextureResolution Resolution = EBakeTextureResolution::Resolution256;
+
+	/** The multisampling configuration per texel */
+	UPROPERTY(EditAnywhere, Category = MapSettings)
+	EBakeMultisampling Multisampling = EBakeMultisampling::None;
+
+	/** Distance to search for the correspondence between the source and target meshes */
+	UPROPERTY(EditAnywhere, Category = MapSettings, meta = (ClampMin = "0.001"))
+	float Thickness = 3.0;
+
+	/** The base mesh UV layer to use to create the map */
+	UPROPERTY(EditAnywhere, Category = MapSettings, meta = (GetOptions = GetUVLayerNamesFunc))
+	FString UVLayer;
+
+	UFUNCTION()
+	TArray<FString> GetUVLayerNamesFunc();
+	UPROPERTY(meta = (TransientToolProperty))
+	TArray<FString> UVLayerNamesList;
+
+	UPROPERTY(VisibleAnywhere, Category = MapSettings, meta = (TransientToolProperty))
+	TArray<TObjectPtr<UTexture2D>> Result;
+
+};
+
+
+USTRUCT()
+struct MESHMODELINGTOOLSEXP_API FBakeMultiMeshDetailProperties
+{
+	GENERATED_BODY()
+
+	/** The detail mesh to sample */
+	UPROPERTY(VisibleAnywhere, Category = DetailMesh, meta = (TransientToolProperty))
+	TObjectPtr<UStaticMesh> DetailMesh = nullptr;
+
+	/** The detail mesh normal map to sample. If empty, the geometric normals will be used. */
+	UPROPERTY(EditAnywhere, Category = DetailMesh, meta = (TransientToolProperty))
+	TObjectPtr<UTexture2D> DetailColorMap = nullptr;
+
+	/** UV layer to sample from on the detail mesh */
+	UPROPERTY(EditAnywhere, Category = DetailMesh, meta = (TransientToolProperty))
+	int32 DetailColorMapUVLayer = 0;
+};
+
+
+UCLASS()
+class MESHMODELINGTOOLSEXP_API UBakeMultiMeshDetailToolProperties : public UInteractiveToolPropertySet
+{
+	GENERATED_BODY()
+
+public:
+	UPROPERTY(EditAnywhere, Category = MapSettings, meta = (TransientToolProperty, EditFixedSize))
+	TArray<FBakeMultiMeshDetailProperties> DetailProperties;
+};
+
+
+struct FBakeMultiMeshDetailSettings
+{
+	using FColorMapData = TTuple<int32, bool>;
+	TArray<FColorMapData> ColorMapData;
+	
+	bool operator==(const FBakeMultiMeshDetailSettings& Other) const
+	{
+		const int NumData = ColorMapData.Num();
+		bool bIsEqual = Other.ColorMapData.Num() == NumData;
+		for (int Idx = 0; bIsEqual && Idx < NumData; ++Idx)
+		{
+			bIsEqual = bIsEqual && ColorMapData[Idx] == Other.ColorMapData[Idx];
+		}
+		return bIsEqual;
+	}
+};
+
+
+// TODO: Refactor shared code into common base class.
+/**
+ * N-to-1 Detail Map Baking Tool
+ */
+UCLASS()
+class MESHMODELINGTOOLSEXP_API UBakeMultiMeshAttributeMapsTool : public UBakeMeshAttributeMapsToolBase
+{
+	GENERATED_BODY()
+
+public:
+	UBakeMultiMeshAttributeMapsTool() = default;
+
+	// Begin UInteractiveTool interface
+	virtual void Setup() override;
+	virtual void Shutdown(EToolShutdownType ShutdownType) override;
+
+	virtual bool HasCancel() const override { return true; }
+	virtual bool HasAccept() const override { return true; }
+	virtual bool CanAccept() const override;
+	// End UInteractiveTool interface
+
+	// Begin IGenericDataOperatorFactory interface
+	virtual TUniquePtr<UE::Geometry::TGenericDataOperator<UE::Geometry::FMeshMapBaker>> MakeNewOperator() override;
+	// End IGenericDataOperatorFactory interface
+
+protected:
+	// need to update bResultValid if these are modified, so we don't publicly expose them. 
+	// @todo setters/getters for these
+
+	UPROPERTY()
+	TObjectPtr<UBakeMultiMeshAttributeMapsToolProperties> Settings;
+
+	UPROPERTY()
+	TObjectPtr<UBakeMultiMeshDetailToolProperties> DetailProps;
+
+
+protected:
+	// Begin UBakeMeshAttributeMapsToolBase interface
+	virtual void UpdateResult() override;
+	virtual void UpdateVisualization() override;
+	// End UBakeMeshAttributeMapsToolBase interface
+	
+
+protected:
+	friend class FMultiMeshMapBakerOp;
+
+	UE::Geometry::FMeshSceneAdapter DetailMeshScene;
+
+	bool bInputsDirty = false;
+	void UpdateOnModeChange();
+
+	// Cached detail mesh data
+	FBakeMultiMeshDetailSettings CachedDetailSettings;
+	EBakeOpState UpdateResult_DetailMeshes();
+	
+	using FTextureImageData = TTuple<UE::Geometry::TImageBuilder<FVector4f>*, int>;
+	using FTextureImageMap = TMap<void*, UE::Geometry::IMeshBakerDetailSampler::FBakeDetailTexture>; 
+	TArray<TSharedPtr<UE::Geometry::TImageBuilder<FVector4f>>> CachedColorImages;
+	TArray<int> CachedColorUVLayers;
+	FTextureImageMap CachedMeshToColorImagesMap;
+};
diff --git a/Engine/Plugins/Experimental/MeshModelingToolsetExp/Source/ModelingComponents/Private/Scene/MeshSceneAdapter.cpp b/Engine/Plugins/Experimental/MeshModelingToolsetExp/Source/ModelingComponents/Private/Scene/MeshSceneAdapter.cpp
--- a/Engine/Plugins/Experimental/MeshModelingToolsetExp/Source/ModelingComponents/Private/Scene/MeshSceneAdapter.cpp
+++ b/Engine/Plugins/Experimental/MeshModelingToolsetExp/Source/ModelingComponents/Private/Scene/MeshSceneAdapter.cpp
@@ -289,6 +289,63 @@ public:
 		return Mesh.TriangleCount();
 	}
 
+	virtual bool IsTriangle(int32 TriId) const override
+	{
+		return Mesh.IsTriangle(TriId);
+	}
+
+	virtual FIndex3i GetTriangle(int32 TriId) const override
+	{
+		return Mesh.GetTriangle(TriId);
+	}
+
+	virtual bool HasNormals() const override
+	{
+		return Mesh.HasAttributes() && Mesh.Attributes()->PrimaryNormals();
+	}
+
+	virtual bool HasUVs(const int UVLayer = 0) const override
+	{
+		return Mesh.HasAttributes() && Mesh.Attributes()->GetUVLayer(UVLayer);
+	}
+
+	virtual FVector3d TriBaryInterpolatePoint(int32 TriId, const FVector3d& BaryCoords) const override
+	{
+		return Mesh.GetTriBaryPoint(TriId, BaryCoords.X, BaryCoords.Y, BaryCoords.Z);
+	}
+
+	virtual bool TriBaryInterpolateNormal(int32 TriId, const FVector3d& BaryCoords, FVector3f& NormalOut) const override
+	{
+		if (Mesh.HasAttributes())
+		{
+			const FDynamicMeshNormalOverlay* NormalOverlay = Mesh.Attributes()->PrimaryNormals();
+			if (NormalOverlay)
+			{
+				FVector3d Normal;
+				NormalOverlay->GetTriBaryInterpolate(TriId, &BaryCoords.X, &Normal.X);
+				NormalOut = FVector3f(Normal);
+				return true;
+			}
+		}
+		return false;
+	}
+
+	virtual bool TriBaryInterpolateUV(const int32 TriId, const FVector3d& BaryCoords, const int UVLayer, FVector2f& UVOut) const override
+	{
+		if (Mesh.HasAttributes())
+		{
+			const FDynamicMeshUVOverlay* UVOverlay = Mesh.Attributes()->GetUVLayer(UVLayer);
+			if (UVOverlay)
+			{
+				FVector2d UV;
+				UVOverlay->GetTriBaryInterpolate(TriId, &BaryCoords.X, &UV.X);
+				UVOut = FVector2f(UV);
+				return true;
+			}
+		}
+		return false;
+	}
+
 	virtual FAxisAlignedBox3d GetWorldBounds(TFunctionRef<FVector3d(const FVector3d&)> LocalToWorldFunc) override
 	{
 		FAxisAlignedBox3d Bounds = bHasBakedTransform ?
@@ -364,8 +421,8 @@ public:
 			UseRay = FRay3d(LocalOrigin, LocalDir);
 		}
 
-		double RayHitT; int32 HitTID;
-		if (AABBTree->FindNearestHitTriangle(UseRay, RayHitT, HitTID))
+		double RayHitT; int32 HitTID; FVector3d HitBaryCoords;
+		if (AABBTree->FindNearestHitTriangle(UseRay, RayHitT, HitTID, HitBaryCoords))
 		{
 			WorldHitResultOut.HitMeshTriIndex = HitTID;
 			WorldHitResultOut.HitMeshSpatialWrapper = this;
@@ -378,6 +435,7 @@ public:
 				FVector3d WorldPos = LocalToWorldTransform.TransformPosition(UseRay.PointAt(RayHitT));
 				WorldHitResultOut.RayDistance = WorldRay.Project(WorldPos);
 			}
+			WorldHitResultOut.HitMeshBaryCoords = HitBaryCoords;
 			return true;
 		}
 		return false;
@@ -620,6 +678,62 @@ public:
 		return Adapter->TriangleCount();
 	}
 
+	virtual bool IsTriangle(int32 TriId) const override
+	{
+		return SourceMesh ? Adapter->IsTriangle(TriId) : false;
+	}
+
+	virtual FIndex3i GetTriangle(int32 TriId) const override
+	{
+		return SourceMesh ? Adapter->GetTriangle(TriId) : FIndex3i::Invalid();
+	}
+
+	virtual bool HasNormals() const override
+	{
+		return SourceMesh ? Adapter->HasNormals() : false;
+	}
+
+	virtual bool HasUVs(const int UVLayer = 0) const override
+	{
+		return SourceMesh ? Adapter->HasUVs() : false;
+	}
+
+	virtual FVector3d TriBaryInterpolatePoint(int32 TriId, const FVector3d& BaryCoords) const override
+	{
+		if (SourceMesh)
+		{
+			FVector3d P0, P1, P2;
+			Adapter->GetTriVertices(TriId, P0, P1, P2);
+			return BaryCoords[0] * P0 + BaryCoords[1] * P1 + BaryCoords[2] * P2;
+		}
+		return FVector3d::Zero();
+	}
+
+	virtual bool TriBaryInterpolateNormal(int32 TriId, const FVector3d& BaryCoords, FVector3f& NormalOut) const override
+	{
+		if (SourceMesh && Adapter->HasNormals())
+		{
+			FIndex3i VertIds = Adapter->GetTriangle(TriId);
+			
+			FVector3f N0, N1, N2;
+			Adapter->GetTriNormals(TriId, N0, N1, N2);
+			NormalOut = (float) BaryCoords[0] * N0 + (float) BaryCoords[1] * N1 + (float) BaryCoords[2] * N2;
+			return true;
+		}
+		return false;
+	}
+
+	virtual bool TriBaryInterpolateUV(const int32 TriId, const FVector3d& BaryCoords, const int UVLayer, FVector2f& UVOut) const override
+	{
+		if (SourceMesh && Adapter->HasUVs())
+		{
+			FVector2f UV0, UV1, UV2;
+			Adapter->GetTriUVs(TriId, UVLayer, UV0, UV1, UV2);
+			UVOut = (float) BaryCoords[0] * UV0 + (float) BaryCoords[1] * UV1 + (float) BaryCoords[2] * UV2;
+			return true;
+		}
+		return false;
+	}
 
 	virtual FAxisAlignedBox3d GetWorldBounds(TFunctionRef<FVector3d(const FVector3d&)> LocalToWorldFunc) override
 	{
@@ -644,13 +758,14 @@ public:
 		FVector3d LocalDir = Normalized(LocalToWorldTransform.InverseTransformPosition(WorldRay.PointAt(1.0)) - LocalOrigin);
 		FRay3d LocalRay(LocalOrigin, LocalDir);
 
-		double LocalHitT; int32 HitTID;
-		if (AABBTree->FindNearestHitTriangle(LocalRay, LocalHitT, HitTID))
+		double LocalHitT; int32 HitTID; FVector3d HitBaryCoords;
+		if (AABBTree->FindNearestHitTriangle(LocalRay, LocalHitT, HitTID, HitBaryCoords))
 		{
 			WorldHitResultOut.HitMeshTriIndex = HitTID;
 			WorldHitResultOut.HitMeshSpatialWrapper = this;
 			FVector3d WorldPos = LocalToWorldTransform.TransformPosition(LocalRay.PointAt(LocalHitT));
 			WorldHitResultOut.RayDistance = WorldRay.Project(WorldPos);
+			WorldHitResultOut.HitMeshBaryCoords = HitBaryCoords;
 			return true;
 		}
 		return false;
@@ -1962,4 +2077,22 @@ void FMeshSceneAdapter::ParallelMeshVertexEnumeration(
 	}, bSingleThread ? EParallelForFlags::ForceSingleThread : EParallelForFlags::Unbalanced );
 
 
-}
\ No newline at end of file
+}
+
+void FMeshSceneAdapter::ProcessActorChildMeshes(TFunctionRef<void(const FActorAdapter* ActorAdapter, const FActorChildMesh* ChildMesh)> ProcessFunc)
+{
+	for (const TUniquePtr<FActorAdapter>& Actor : SceneActors)
+	{
+		if (Actor)
+		{
+			for (const TUniquePtr<FActorChildMesh>& ChildMesh : Actor->ChildMeshes)
+			{
+				if (ChildMesh)
+				{
+					ProcessFunc(Actor.Get(), ChildMesh.Get());
+				}
+			}
+		}
+	}
+}
+
diff --git a/Engine/Plugins/Experimental/MeshModelingToolsetExp/Source/ModelingComponents/Public/Scene/MeshSceneAdapter.h b/Engine/Plugins/Experimental/MeshModelingToolsetExp/Source/ModelingComponents/Public/Scene/MeshSceneAdapter.h
--- a/Engine/Plugins/Experimental/MeshModelingToolsetExp/Source/ModelingComponents/Public/Scene/MeshSceneAdapter.h
+++ b/Engine/Plugins/Experimental/MeshModelingToolsetExp/Source/ModelingComponents/Public/Scene/MeshSceneAdapter.h
@@ -105,6 +105,8 @@ struct MODELINGCOMPONENTS_API FMeshSceneRayHit
 	int32 HitComponentElementIndex = -1;
 	// Triangle Index/ID on mesh of hit Component
 	int32 HitMeshTriIndex = -1;
+	// Triangle barycentric coordinates of the hit location.
+	FVector3d HitMeshBaryCoords = FVector3d::Zero();
 
 	// SpatialWrapper for the mesh geometry that was hit, if available. This is a pointer to data owned by the FMeshSceneAdapter that was queried
 	IMeshSpatialWrapper* HitMeshSpatialWrapper = nullptr;
@@ -147,6 +149,27 @@ public:
 
 	/** Append the geometry represented by this wrapper to the accumulated AppendTo mesh, under the given world transform */
 	virtual void AppendMesh(FDynamicMesh3& AppendTo, const FTransformSequence3d& TransformSeq) = 0;
+
+	/** @return true if the given index is a valid triangle */
+	virtual bool IsTriangle(int32 TriId) const = 0;
+
+	/** @return vertex indices of the given triangle */
+	virtual FIndex3i GetTriangle(int32 TriId) const = 0;
+
+	/** @return true if the mesh has normals */
+	virtual bool HasNormals() const = 0;
+
+	/** @return true if the mesh has UVs */
+	virtual bool HasUVs(int UVLayer = 0) const = 0;
+
+	/** Compute the barycentric interpolated normal for the given tri */
+	virtual FVector3d TriBaryInterpolatePoint(int32 TriId, const FVector3d& BaryCoords) const = 0;
+
+	/** Compute the barycentric interpolated normal for the given tri */
+	virtual bool TriBaryInterpolateNormal(int32 TriId, const FVector3d& BaryCoords, FVector3f& NormalOut) const = 0;
+
+	/** Compute the barycentric interpolated UV for the given tri */
+	virtual bool TriBaryInterpolateUV(int32 TriId, const FVector3d& BaryCoords, int UVLayer, FVector2f& UVOut) const = 0;
 };
 
 
@@ -347,6 +370,11 @@ public:
 		bool bForceSingleThreaded = false
 	);
 
+	/**
+	 * Run a custom query across all scene actor child meshes.
+	 */
+	virtual void ProcessActorChildMeshes(TFunctionRef<void(const FActorAdapter* ActorAdapter, const FActorChildMesh* ChildMesh)> ProcessFunc);
+
 protected:
 	// top-level list of ActorAdapters, which represent each Actor and set of Components
 	TArray<TUniquePtr<FActorAdapter>> SceneActors;
diff --git a/Engine/Plugins/Runtime/GeometryProcessing/Source/DynamicMesh/Private/Sampling/MeshBakerCommon.cpp b/Engine/Plugins/Runtime/GeometryProcessing/Source/DynamicMesh/Private/Sampling/MeshBakerCommon.cpp
--- a/Engine/Plugins/Runtime/GeometryProcessing/Source/DynamicMesh/Private/Sampling/MeshBakerCommon.cpp
+++ b/Engine/Plugins/Runtime/GeometryProcessing/Source/DynamicMesh/Private/Sampling/MeshBakerCommon.cpp
@@ -70,7 +70,6 @@ const void* UE::Geometry::GetDetailMeshTrianglePoint_Raycast(
 	{
 		DetailTriangleOut = HitTID;
 		DetailTriBaryCoords = HitBaryCoords;
-		return HitMesh;
 	}
 	else
 	{
@@ -79,24 +78,19 @@ const void* UE::Geometry::GetDetailMeshTrianglePoint_Raycast(
 		OnSurfQueryOptions.MaxDistance = Thickness;
 		double NearDistSqr = 0;
 		int32 NearestTriID = IndexConstants::InvalidID;
-		const void* NearestMesh = nullptr;
+		
 		// if we are using absolute nearest point as a fallback, then ignore max distance
-		if (bFailToNearestPoint)
-		{
-			NearestMesh = DetailSampler->FindNearestTriangle(BasePoint, NearDistSqr, NearestTriID, DetailTriBaryCoords);
-		}
-		else
-		{
-			NearestMesh = DetailSampler->FindNearestTriangle(BasePoint, NearDistSqr, NearestTriID, DetailTriBaryCoords, OnSurfQueryOptions);
-		}
-		if (NearestMesh && DetailSampler->IsTriangle(NearestMesh, NearestTriID))
+		HitMesh = bFailToNearestPoint ? 
+			DetailSampler->FindNearestTriangle(BasePoint, NearDistSqr, NearestTriID, HitBaryCoords) :
+			DetailSampler->FindNearestTriangle(BasePoint, NearDistSqr, NearestTriID, HitBaryCoords, OnSurfQueryOptions);
+		if (HitMesh && DetailSampler->IsTriangle(HitMesh, NearestTriID))
 		{
 			DetailTriangleOut = NearestTriID;
-			return NearestMesh;
+			DetailTriBaryCoords = HitBaryCoords;
 		}
 	}
 
-	return nullptr;
+	return HitMesh;
 }
 
 
@@ -111,10 +105,12 @@ const void* UE::Geometry::GetDetailMeshTrianglePoint_Nearest(
 {
 	double NearDistSqr = 0;
 	int32 NearestTriID = IndexConstants::InvalidID;
-	const void* NearestMesh = DetailSampler->FindNearestTriangle(BasePoint, NearDistSqr, NearestTriID, DetailTriBaryCoords);
+	FVector3d NearestBaryCoords = FVector3d::Zero();
+	const void* NearestMesh = DetailSampler->FindNearestTriangle(BasePoint, NearDistSqr, NearestTriID, NearestBaryCoords);
 	if (NearestMesh && DetailSampler->IsTriangle(NearestMesh, NearestTriID))
 	{
 		DetailTriangleOut = NearestTriID;
+		DetailTriBaryCoords = NearestBaryCoords;
 		return NearestMesh;
 	}
 
diff --git a/Engine/Plugins/Runtime/GeometryProcessing/Source/DynamicMesh/Private/Sampling/MeshMapBaker.cpp b/Engine/Plugins/Runtime/GeometryProcessing/Source/DynamicMesh/Private/Sampling/MeshMapBaker.cpp
--- a/Engine/Plugins/Runtime/GeometryProcessing/Source/DynamicMesh/Private/Sampling/MeshMapBaker.cpp
+++ b/Engine/Plugins/Runtime/GeometryProcessing/Source/DynamicMesh/Private/Sampling/MeshMapBaker.cpp
@@ -323,7 +323,7 @@ void FMeshMapBaker::Bake()
 
 						FMeshMapEvaluator::FCorrespondenceSample Sample;
 						DetailCorrespondenceSampler.SampleUV(UVTriangleID, UVPosition, Sample);
-						if (DetailSampler->IsTriangle(Sample.DetailMesh, Sample.DetailTriID))
+						if (Sample.DetailMesh && DetailSampler->IsTriangle(Sample.DetailMesh, Sample.DetailTriID))
 						{
 							BakeSample(*TileBuffer, Sample, UVPosition, ImageCoords);
 						}
diff --git a/Engine/Plugins/Runtime/GeometryProcessing/Source/DynamicMesh/Private/Sampling/MeshNormalMapEvaluator.cpp b/Engine/Plugins/Runtime/GeometryProcessing/Source/DynamicMesh/Private/Sampling/MeshNormalMapEvaluator.cpp
--- a/Engine/Plugins/Runtime/GeometryProcessing/Source/DynamicMesh/Private/Sampling/MeshNormalMapEvaluator.cpp
+++ b/Engine/Plugins/Runtime/GeometryProcessing/Source/DynamicMesh/Private/Sampling/MeshNormalMapEvaluator.cpp
@@ -11,24 +11,25 @@ void FMeshNormalMapEvaluator::Setup(const FMeshBaseBaker& Baker, FEvaluationCont
 {
 	// Cache data from the baker
 	DetailSampler = Baker.GetDetailSampler();
-	auto GetNormalMaps = [this](const void* Mesh)
+	auto GetDetailNormalMaps = [this](const void* Mesh)
 	{
-		const FDetailNormalTexture NormalMap = DetailSampler->GetNormalMap(Mesh);
-		DetailNormalTextures.Add(Mesh, NormalMap);
+		const FDetailNormalTexture* NormalMap = DetailSampler->GetNormalMap(Mesh);
+		if (NormalMap)
+		{
+			// Require valid normal map, UV layer and tangents to enable normal map transfer.
+			const bool bEnableNormalMapTransfer = DetailSampler->HasUVs(Mesh, NormalMap->Get<1>()) && DetailSampler->HasTangents(Mesh);
+			if (bEnableNormalMapTransfer)
+			{
+				DetailNormalTextures.Add(Mesh, *NormalMap);
+				bHasDetailNormalTextures = true;
+			}
+		}
 	};
-	DetailSampler->ProcessMeshes(GetNormalMaps);
-	bool bHasDetailNormalMap = false;
-	for (const auto& Map : DetailNormalTextures)
-	{
-		const TImageBuilder<FVector4f>* NormalMap;
-		int NormalUVLayer;
-		Tie(NormalMap, NormalUVLayer) = Map.Value;
-		bHasDetailNormalMap = bHasDetailNormalMap || NormalMap != nullptr;
-	}
+	DetailSampler->ProcessMeshes(GetDetailNormalMaps);
 	
 	BaseMeshTangents = Baker.GetTargetMeshTangents();
 
-	Context.Evaluate = bHasDetailNormalMap ? &EvaluateSample<true> : &EvaluateSample<false>;
+	Context.Evaluate = bHasDetailNormalTextures ? &EvaluateSample<true> : &EvaluateSample<false>;
 	Context.EvaluateDefault = &EvaluateDefault;
 	Context.EvaluateColor = &EvaluateColor;
 	Context.EvalData = this;
@@ -84,9 +85,13 @@ FVector3f FMeshNormalMapEvaluator::SampleFunction(const FCorrespondenceSample& S
 	
 	if constexpr (bUseDetailNormalMap)
 	{
-		const TImageBuilder<FVector4f>* DetailNormalMap;
-		int DetailNormalUVLayer;
-		Tie(DetailNormalMap, DetailNormalUVLayer) = DetailNormalTextures[DetailMesh];
+		const TImageBuilder<FVector4f>* DetailNormalMap = nullptr;
+		int DetailNormalUVLayer = 0;
+		const FDetailNormalTexture* DetailNormalTexture = DetailNormalTextures.Find(DetailMesh);
+		if (DetailNormalTexture)
+		{
+			Tie(DetailNormalMap, DetailNormalUVLayer) = *DetailNormalTexture;
+		}
 
 		if (DetailNormalMap)
 		{
diff --git a/Engine/Plugins/Runtime/GeometryProcessing/Source/DynamicMesh/Private/Sampling/MeshPropertyMapEvaluator.cpp b/Engine/Plugins/Runtime/GeometryProcessing/Source/DynamicMesh/Private/Sampling/MeshPropertyMapEvaluator.cpp
--- a/Engine/Plugins/Runtime/GeometryProcessing/Source/DynamicMesh/Private/Sampling/MeshPropertyMapEvaluator.cpp
+++ b/Engine/Plugins/Runtime/GeometryProcessing/Source/DynamicMesh/Private/Sampling/MeshPropertyMapEvaluator.cpp
@@ -11,22 +11,23 @@ void FMeshPropertyMapEvaluator::Setup(const FMeshBaseBaker& Baker, FEvaluationCo
 {
 	// Cache data from the baker
 	DetailSampler = Baker.GetDetailSampler();
-	auto GetNormalMaps = [this](const void* Mesh)
+	auto GetDetailNormalMaps = [this](const void* Mesh)
 	{
-		const FDetailNormalTexture NormalMap = DetailSampler->GetNormalMap(Mesh);
-		DetailNormalTextures.Add(Mesh, NormalMap);
+		const FDetailNormalTexture* NormalMap = DetailSampler->GetNormalMap(Mesh);
+		if (NormalMap)
+		{
+			// Require valid normal map, UV layer and tangents to enable normal map transfer.
+			const bool bEnableNormalMapTransfer = DetailSampler->HasUVs(Mesh, NormalMap->Get<1>()) && DetailSampler->HasTangents(Mesh);
+			if (bEnableNormalMapTransfer)
+			{
+				DetailNormalTextures.Add(Mesh, *NormalMap);
+				bHasDetailNormalTextures = true;
+			}
+		}
 	};
-	DetailSampler->ProcessMeshes(GetNormalMaps);
-	bool bHasDetailNormalMap = false;
-	for (const auto& Map : DetailNormalTextures)
-	{
-		const TImageBuilder<FVector4f>* NormalMap;
-		int NormalUVLayer;
-		Tie(NormalMap, NormalUVLayer) = Map.Value;
-		bHasDetailNormalMap = bHasDetailNormalMap || NormalMap != nullptr;
-	}
+	DetailSampler->ProcessMeshes(GetDetailNormalMaps);
 
-	Context.Evaluate = bHasDetailNormalMap ? &EvaluateSample<true> : &EvaluateSample<false>;
+	Context.Evaluate = bHasDetailNormalTextures ? &EvaluateSample<true> : &EvaluateSample<false>;
 	Context.EvaluateDefault = &EvaluateDefault;
 	Context.EvaluateColor = &EvaluateColor;
 	Context.EvalData = this;
@@ -118,9 +119,13 @@ FVector3f FMeshPropertyMapEvaluator::SampleFunction(const FCorrespondenceSample&
 
 			if constexpr (bUseDetailNormalMap)
 			{
-				const TImageBuilder<FVector4f>* DetailNormalMap;
-				int DetailNormalUVLayer;
-				Tie(DetailNormalMap, DetailNormalUVLayer) = DetailNormalTextures[DetailMesh];
+				const TImageBuilder<FVector4f>* DetailNormalMap = nullptr;
+				int DetailNormalUVLayer = 0;
+				const FDetailNormalTexture* DetailNormalTexture = DetailNormalTextures.Find(DetailMesh);
+				if (DetailNormalTexture)
+				{
+					Tie(DetailNormalMap, DetailNormalUVLayer) = *DetailNormalTexture;
+				}
 
 				if (DetailNormalMap)
 				{
diff --git a/Engine/Plugins/Runtime/GeometryProcessing/Source/DynamicMesh/Private/Sampling/MeshResampleImageEvaluator.cpp b/Engine/Plugins/Runtime/GeometryProcessing/Source/DynamicMesh/Private/Sampling/MeshResampleImageEvaluator.cpp
--- a/Engine/Plugins/Runtime/GeometryProcessing/Source/DynamicMesh/Private/Sampling/MeshResampleImageEvaluator.cpp
+++ b/Engine/Plugins/Runtime/GeometryProcessing/Source/DynamicMesh/Private/Sampling/MeshResampleImageEvaluator.cpp
@@ -45,13 +45,22 @@ FVector4f FMeshResampleImageEvaluator::ImageSampleFunction(const FCorrespondence
 {
 	const void* DetailMesh = SampleData.DetailMesh;
 	const int32 DetailTriID = SampleData.DetailTriID;
-	const TImageBuilder<FVector4f>* TextureImage;
-	int DetailUVLayer;
-	Tie(TextureImage, DetailUVLayer) = DetailSampler->GetColorMap(DetailMesh);
+	const TImageBuilder<FVector4f>* TextureImage = nullptr;
+	int DetailUVLayer = 0;
+	const IMeshBakerDetailSampler::FBakeDetailTexture* ColorMap = DetailSampler->GetColorMap(DetailMesh);
+	if (ColorMap)
+	{
+		Tie(TextureImage, DetailUVLayer) = *ColorMap;
+	}
 
-	FVector2f DetailUV;
-	DetailSampler->TriBaryInterpolateUV(DetailMesh, DetailTriID, SampleData.DetailBaryCoords, DetailUVLayer, DetailUV);
-	return TextureImage->BilinearSampleUV<float>(FVector2d(DetailUV), FVector4f(0, 0, 0, 1));
+	FVector4f Color(0, 0, 0, 1);
+	if (TextureImage)
+	{
+		FVector2f DetailUV;
+		DetailSampler->TriBaryInterpolateUV(DetailMesh, DetailTriID, SampleData.DetailBaryCoords, DetailUVLayer, DetailUV);
+		Color = TextureImage->BilinearSampleUV<float>(FVector2d(DetailUV), FVector4f(0, 0, 0, 1));
+	}
+	return Color;
 }
 
 //
diff --git a/Engine/Plugins/Runtime/GeometryProcessing/Source/DynamicMesh/Public/Sampling/MeshBakerCommon.h b/Engine/Plugins/Runtime/GeometryProcessing/Source/DynamicMesh/Public/Sampling/MeshBakerCommon.h
--- a/Engine/Plugins/Runtime/GeometryProcessing/Source/DynamicMesh/Public/Sampling/MeshBakerCommon.h
+++ b/Engine/Plugins/Runtime/GeometryProcessing/Source/DynamicMesh/Public/Sampling/MeshBakerCommon.h
@@ -110,10 +110,10 @@ public:
 	virtual void SetNormalMap(const void* Mesh, const FBakeDetailTexture& Map) = 0;
 
 	/** Retrieve a color map and UV layer index from a given mesh in the detail set */
-	virtual FBakeDetailTexture GetColorMap(const void* Mesh) const = 0;
+	virtual const FBakeDetailTexture* GetColorMap(const void* Mesh) const = 0;
 
 	/** Retrieve a normal map and UV layer index from a given mesh in the detail set */
-	virtual FBakeDetailTexture GetNormalMap(const void* Mesh) const = 0;
+	virtual const FBakeDetailTexture* GetNormalMap(const void* Mesh) const = 0;
 
 	/** @return true if identity correspondence is supported */
 	virtual bool SupportsIdentityCorrespondence() const = 0;
@@ -205,6 +205,31 @@ public:
 	 */
 	virtual int32 GetMaterialID(const void* Mesh, const int TriId) const = 0;
 
+	/**
+	 * @param Mesh pointer to mesh to query
+	 * @return true if this mesh has normals
+	 */
+	virtual bool HasNormals(const void* Mesh) const = 0;
+
+	/**
+	 * @param Mesh pointer to mesh to query
+	 * @param UVLayer the UV layer to check for existence
+	 * @return true if this mesh has UVs on the given layer
+	 */
+	virtual bool HasUVs(const void* Mesh, int UVLayer = 0) const = 0;
+
+	/**
+	 * @param Mesh pointer to mesh to query
+	 * @return true if this mesh has tangents
+	 */
+	virtual bool HasTangents(const void* Mesh) const = 0;
+
+	/**
+	 * @param Mesh pointer to mesh to query
+	 * @return true if this mesh has vertex colors on the given layer
+	 */
+	virtual bool HasColors(const void* Mesh) const = 0;
+
 	/**
 	 * Compute interpolated Position value inside triangle using barycentric coordinates
 	 * @param Mesh pointer to mesh to query
@@ -314,14 +339,14 @@ public:
 		DetailNormalMap = Map;
 	}
 
-	virtual FBakeDetailTexture GetColorMap(const void* Mesh) const override
+	virtual const FBakeDetailTexture* GetColorMap(const void* Mesh) const override
 	{
-		return DetailColorMap;
+		return &DetailColorMap;
 	}
 	
-	virtual FBakeDetailTexture GetNormalMap(const void* Mesh) const override
+	virtual const FBakeDetailTexture* GetNormalMap(const void* Mesh) const override
 	{
-		return DetailNormalMap;
+		return &DetailNormalMap;
 	}
 
 	virtual bool SupportsIdentityCorrespondence() const override
@@ -362,14 +387,8 @@ public:
 		FVector3d& TriBaryCoords,
 		const IMeshSpatial::FQueryOptions& Options = IMeshSpatial::FQueryOptions()) const override
 	{
-		const bool bHit = DetailSpatial->FindNearestHitTriangle(Ray, NearestT, TriId, Options);
-		if (bHit)
-		{
-			const FIntrRay3Triangle3d IntrQuery = TMeshQueries<FDynamicMesh3>::TriangleIntersection(*DetailMesh, TriId, Ray);
-			TriBaryCoords = IntrQuery.TriangleBaryCoords;
-			return DetailMesh;
-		}
-		return nullptr;
+		const bool bHit = DetailSpatial->FindNearestHitTriangle(Ray, NearestT, TriId, TriBaryCoords, Options);
+		return bHit ? DetailMesh : nullptr;
 	}
 
 	virtual bool TestAnyHitTriangle(
@@ -409,6 +428,29 @@ public:
 		return MaterialIDOverlay ? MaterialIDOverlay->GetValue(TriId) : IndexConstants::InvalidID;
 	}
 
+	virtual bool HasNormals(const void* Mesh) const override
+	{
+		const FDynamicMesh3* DynamicMesh = static_cast<const FDynamicMesh3*>(Mesh);
+		return DynamicMesh && DynamicMesh->HasAttributes() && DynamicMesh->Attributes()->PrimaryNormals();
+	}
+
+	virtual bool HasUVs(const void* Mesh, const int UVLayer = 0) const override
+	{
+		const FDynamicMesh3* DynamicMesh = static_cast<const FDynamicMesh3*>(Mesh);
+		return DynamicMesh && DynamicMesh->HasAttributes() && DynamicMesh->Attributes()->GetUVLayer(UVLayer);
+	}
+
+	virtual bool HasTangents(const void* Mesh) const override
+	{
+		return DetailTangents != nullptr;
+	}
+
+	virtual bool HasColors(const void* Mesh) const override
+	{
+		const FDynamicMesh3* DynamicMesh = static_cast<const FDynamicMesh3*>(Mesh);
+		return DynamicMesh && DynamicMesh->HasAttributes() && DynamicMesh->Attributes()->PrimaryColors();
+	}
+
 	virtual FVector3d TriBaryInterpolatePoint(
 		const void* Mesh,
 		const int32 TriId,
diff --git a/Engine/Plugins/Runtime/GeometryProcessing/Source/DynamicMesh/Public/Sampling/MeshNormalMapEvaluator.h b/Engine/Plugins/Runtime/GeometryProcessing/Source/DynamicMesh/Public/Sampling/MeshNormalMapEvaluator.h
--- a/Engine/Plugins/Runtime/GeometryProcessing/Source/DynamicMesh/Public/Sampling/MeshNormalMapEvaluator.h
+++ b/Engine/Plugins/Runtime/GeometryProcessing/Source/DynamicMesh/Public/Sampling/MeshNormalMapEvaluator.h
@@ -38,6 +38,7 @@ protected:
 	using FDetailNormalTexture = IMeshBakerDetailSampler::FBakeDetailTexture;
 	using FDetailNormalTextureMap = TMap<const void*, FDetailNormalTexture>;
 	FDetailNormalTextureMap DetailNormalTextures;
+	bool bHasDetailNormalTextures = false;
 	const TMeshTangents<double>* BaseMeshTangents = nullptr;
 
 	FVector3f DefaultNormal = FVector3f::UnitZ();
diff --git a/Engine/Plugins/Runtime/GeometryProcessing/Source/DynamicMesh/Public/Sampling/MeshPropertyMapEvaluator.h b/Engine/Plugins/Runtime/GeometryProcessing/Source/DynamicMesh/Public/Sampling/MeshPropertyMapEvaluator.h
--- a/Engine/Plugins/Runtime/GeometryProcessing/Source/DynamicMesh/Public/Sampling/MeshPropertyMapEvaluator.h
+++ b/Engine/Plugins/Runtime/GeometryProcessing/Source/DynamicMesh/Public/Sampling/MeshPropertyMapEvaluator.h
@@ -51,6 +51,7 @@ protected:
 	using FDetailNormalTexture = IMeshBakerDetailSampler::FBakeDetailTexture;
 	using FDetailNormalTextureMap = TMap<const void*, FDetailNormalTexture>;
 	FDetailNormalTextureMap DetailNormalTextures;
+	bool bHasDetailNormalTextures = false;
 	FAxisAlignedBox3d Bounds;
 	FVector3f DefaultValue = FVector3f::Zero();
 
diff --git a/Engine/Source/Runtime/GeometryCore/Public/Spatial/MeshAABBTree3.h b/Engine/Source/Runtime/GeometryCore/Public/Spatial/MeshAABBTree3.h
--- a/Engine/Source/Runtime/GeometryCore/Public/Spatial/MeshAABBTree3.h
+++ b/Engine/Source/Runtime/GeometryCore/Public/Spatial/MeshAABBTree3.h
@@ -362,15 +362,23 @@ public:
 	{
 		double NearestT;
 		int TID;
-		FindNearestHitTriangle(Ray, NearestT, TID, Options);
+		FVector3d BaryCoords;
+		FindNearestHitTriangle(Ray, NearestT, TID, BaryCoords, Options);
 		return TID;
 	}
 
+	virtual bool FindNearestHitTriangle(const FRay3d& Ray, double& NearestT, int& TID, const FQueryOptions& Options = FQueryOptions()) const override
+	{
+		FVector3d BaryCoords;
+		return FindNearestHitTriangle(Ray, NearestT, TID, BaryCoords, Options);
+	}
+
 	virtual bool FindNearestHitTriangle(
-		const FRay3d& Ray, double& NearestT, int& TID, 
+		const FRay3d& Ray, double& NearestT, int& TID, FVector3d& BaryCoords,
 		const FQueryOptions& Options = FQueryOptions()) const override
 	{
 		TID = IndexConstants::InvalidID;
+		BaryCoords = FVector3d::Zero();
 
 		// Note: using TNumericLimits<float>::Max() here because we need to use <= to compare Box hit
 		//   to NearestT, and Box hit returns TNumericLimits<double>::Max() on no-hit. So, if we set
@@ -382,13 +390,13 @@ public:
 			return false;
 		}
 
-		FindHitTriangle(RootIndex, Ray, NearestT, TID, Options);
+		FindHitTriangle(RootIndex, Ray, NearestT, TID, BaryCoords, Options);
 		return TID != IndexConstants::InvalidID;
 	}
 
 protected:
 	void FindHitTriangle(
-		int IBox, const FRay3d& Ray, double& NearestT, int& TID,
+		int IBox, const FRay3d& Ray, double& NearestT, int& TID, FVector3d& BaryCoords,
 		const FQueryOptions& Options = FQueryOptions()) const
 	{
 		int idx = BoxToIndex[IBox];
@@ -412,6 +420,7 @@ protected:
 					{
 						NearestT = Query.RayParameter;
 						TID = ti;
+						BaryCoords = Query.TriangleBaryCoords;
 					}
 				}
 			}
@@ -427,7 +436,7 @@ protected:
 				double fChild1T = box_ray_intersect_t(iChild1, Ray);
 				if (fChild1T <= NearestT + e)
 				{
-					FindHitTriangle(iChild1, Ray, NearestT, TID, Options);
+					FindHitTriangle(iChild1, Ray, NearestT, TID, BaryCoords, Options);
 				}
 			}
 			else
@@ -441,10 +450,10 @@ protected:
 				{
 					if (fChild1T <= NearestT + e)
 					{
-						FindHitTriangle(iChild1, Ray, NearestT, TID, Options);
+						FindHitTriangle(iChild1, Ray, NearestT, TID, BaryCoords, Options);
 						if (fChild2T <= NearestT + e)
 						{
-							FindHitTriangle(iChild2, Ray, NearestT, TID, Options);
+							FindHitTriangle(iChild2, Ray, NearestT, TID, BaryCoords, Options);
 						}
 					}
 				}
@@ -452,10 +461,10 @@ protected:
 				{
 					if (fChild2T <= NearestT + e)
 					{
-						FindHitTriangle(iChild2, Ray, NearestT, TID, Options);
+						FindHitTriangle(iChild2, Ray, NearestT, TID, BaryCoords, Options);
 						if (fChild1T <= NearestT + e)
 						{
-							FindHitTriangle(iChild1, Ray, NearestT, TID, Options);
+							FindHitTriangle(iChild1, Ray, NearestT, TID, BaryCoords, Options);
 						}
 					}
 				}
diff --git a/Engine/Source/Runtime/GeometryCore/Public/Spatial/SpatialInterfaces.h b/Engine/Source/Runtime/GeometryCore/Public/Spatial/SpatialInterfaces.h
--- a/Engine/Source/Runtime/GeometryCore/Public/Spatial/SpatialInterfaces.h
+++ b/Engine/Source/Runtime/GeometryCore/Public/Spatial/SpatialInterfaces.h
@@ -75,7 +75,7 @@ public:
 	/**
 	 * @param Query point
 	 * @param NearestDistSqrOut returned nearest squared distance, if triangle is found
-	 * @param MaxDistance maximum search distance
+	 * @param Options Query options (ex. max distance)
 	 * @return ID of triangle nearest to Point within MaxDistance, or InvalidID if not found
 	 */
 	virtual int FindNearestTriangle(const FVector3d& Point, double& NearestDistSqrOut, const FQueryOptions& Options = FQueryOptions()) const = 0;
@@ -88,7 +88,7 @@ public:
 
 	/**
 	 * @param Ray query ray
-	 * @param MaxDistance maximum hit distance
+	 * @param Options Query options (ex. max distance)
 	 * @return ID of triangle intersected by ray within MaxDistance, or InvalidID if not found
 	 */
 	virtual int FindNearestHitTriangle(const FRay3d& Ray, const FQueryOptions& Options = FQueryOptions()) const
@@ -104,10 +104,25 @@ public:
 	 * @param Ray query ray
 	 * @param NearestT returned-by-reference parameter of the nearest hit
 	 * @param TID returned-by-reference ID of triangle intersected by ray within MaxDistance, or InvalidID if not found
-	 * @param MaxDistance maximum hit distance
+	 * @param Options Query options (ex. max distance)
 	 * @return true if hit, false if no hit found
 	 */
-	virtual bool FindNearestHitTriangle(const FRay3d& Ray, double& NearestT, int& TID, const FQueryOptions& Options = FQueryOptions()) const = 0;
+	virtual bool FindNearestHitTriangle(const FRay3d& Ray, double& NearestT, int& TID, const FQueryOptions& Options = FQueryOptions()) const
+	{
+		FVector3d BaryCoords;
+		return FindNearestHitTriangle(Ray, NearestT, TID, BaryCoords, Options);
+	}
+
+	/**
+	 * Find nearest triangle from the given ray
+	 * @param Ray query ray
+	 * @param NearestT returned-by-reference parameter of the nearest hit
+	 * @param TID returned-by-reference ID of triangle intersected by ray within MaxDistance, or InvalidID if not found
+	 * @param BaryCoords returned-by-reference Barycentric coordinates of the triangle intersected by ray within MaxDistance, or FVector3d::Zero if not found.
+	 * @param Options Query options (ex. max distance)
+	 * @return true if hit, false if no hit found
+	 */
+	virtual bool FindNearestHitTriangle(const FRay3d& Ray, double& NearestT, int& TID, FVector3d& BaryCoords, const FQueryOptions& Options = FQueryOptions()) const = 0;
 
 
 };
diff --git a/Engine/Source/Runtime/MeshConversion/Public/MeshDescriptionAdapter.h b/Engine/Source/Runtime/MeshConversion/Public/MeshDescriptionAdapter.h
--- a/Engine/Source/Runtime/MeshConversion/Public/MeshDescriptionAdapter.h
+++ b/Engine/Source/Runtime/MeshConversion/Public/MeshDescriptionAdapter.h
@@ -25,6 +25,8 @@ protected:
 	const FMeshDescription* Mesh;
 	TVertexAttributesConstRef<FVector3f> VertexPositions;
 	TVertexInstanceAttributesConstRef<FVector3f> VertexInstanceNormals;
+	TVertexInstanceAttributesConstRef<FVector2D> VertexInstanceUVs;
+	TArrayView<const FVertexInstanceID> TriangleVertexInstanceIndices;
 
 	FVector3d BuildScale = FVector3d::One();
 	bool bScaleNormals = false;
@@ -35,6 +37,8 @@ public:
 		FStaticMeshConstAttributes Attributes(*MeshIn);
 		VertexPositions = Attributes.GetVertexPositions();
 		VertexInstanceNormals = Attributes.GetVertexInstanceNormals();
+		VertexInstanceUVs = Attributes.GetVertexInstanceUVs();
+		TriangleVertexInstanceIndices = Attributes.GetTriangleVertexInstanceIndices().GetRawArray();
 		// @todo: can we hold TArrayViews of the attribute arrays here? Do we guarantee not to mutate the mesh description for the duration of this object?
 	}
 
@@ -125,12 +129,56 @@ public:
 	{
 		return HasNormals() ? VertexInstanceNormals.GetNumElements() : 0;
 	}
+
+	/** Get Normal by VertexInstanceID */
 	FVector3f GetNormal(int32 IDValue) const
 	{
 		const FVector& InstanceNormal = VertexInstanceNormals[FVertexInstanceID(IDValue)];
 		return (!bScaleNormals) ? FVector3f(InstanceNormal) :
 			UE::Geometry::Normalized(FVector3f(InstanceNormal.X/BuildScale.X, InstanceNormal.Y/BuildScale.Y, InstanceNormal.Z/BuildScale.Z));
 	}
+
+	/** Get Normals for a given Triangle */
+	template<typename VectorType>
+	void GetTriNormals(int32 TriId, VectorType& N0, VectorType& N1, VectorType& N2)
+	{
+		N0 = GetNormal(TriangleVertexInstanceIndices[TriId*3]);
+		N1 = GetNormal(TriangleVertexInstanceIndices[TriId*3+1]);
+		N2 = GetNormal(TriangleVertexInstanceIndices[TriId*3+2]);
+	}
+
+	inline bool HasUVs() const
+	{
+		return VertexInstanceUVs.IsValid();
+	}
+	inline bool IsUV(const int32 UVId) const
+	{
+		return HasUVs() && UVId >= 0 && UVId < UVCount();
+	}
+	inline int32 MaxUVID() const
+	{
+		return HasUVs() ? VertexInstanceUVs.GetNumElements() : 0;
+	}
+	inline int32 UVCount() const
+	{
+		return HasUVs() ? VertexInstanceUVs.GetNumElements() : 0;
+	}
+
+	/** Get UV by VertexInstanceID for a given UVLayer */
+	FVector2f GetUV(const int32 IDValue, const int32 UVLayer) const
+	{
+		const FVector2D& InstanceUV = VertexInstanceUVs.Get(FVertexInstanceID(IDValue), UVLayer);
+		return FVector2f(InstanceUV);
+	}
+
+	/** Get UVs for a given UVLayer and Triangle */
+	template<typename VectorType>
+	void GetTriUVs(const int32 TriId, const int32 UVLayer, VectorType& UV0, VectorType& UV1, VectorType& UV2)
+	{
+		UV0 = GetUV(TriangleVertexInstanceIndices[TriId*3], UVLayer);
+		UV1 = GetUV(TriangleVertexInstanceIndices[TriId*3+1], UVLayer);
+		UV2 = GetUV(TriangleVertexInstanceIndices[TriId*3+2], UVLayer);
+	}
 };
 
 
