Mesh Modeling Tools:  DeformMeshPolygonsTool - change the laplacian flavor from using clamped-cotangent to using intrinsic delaunay cotangent

diff --git a/Engine/Plugins/Experimental/MeshModelingToolsetExp/Source/MeshModelingToolsExp/Public/DeformMeshPolygonsTool.h b/Engine/Plugins/Experimental/MeshModelingToolsetExp/Source/MeshModelingToolsExp/Public/DeformMeshPolygonsTool.h
--- a/Engine/Plugins/Experimental/MeshModelingToolsetExp/Source/MeshModelingToolsExp/Public/DeformMeshPolygonsTool.h
+++ b/Engine/Plugins/Experimental/MeshModelingToolsetExp/Source/MeshModelingToolsExp/Public/DeformMeshPolygonsTool.h
@@ -60,7 +60,8 @@ enum class EWeightScheme
 	Valence				UMETA(DisplayName = "Valence"),
 	MeanValue			UMETA(DisplayName = "MeanValue"),
 	Cotangent			UMETA(DisplayName = "Cotangent"),
-	ClampedCotangent	UMETA(DisplayName = "ClampedCotangent")
+	ClampedCotangent	UMETA(DisplayName = "ClampedCotangent"),
+	IDTCotangent        UMETA(DisplayName = "IDTCotangent")
 };
 
 /** The ELaplacianWeightScheme enum is the same..*/
@@ -115,7 +116,7 @@ public:
 	//Laplacian Deformation Options, currently not exposed.
 
 	UPROPERTY() 
-	EWeightScheme SelectedWeightScheme = EWeightScheme::ClampedCotangent;
+	EWeightScheme SelectedWeightScheme = EWeightScheme::IDTCotangent;
 
 	UPROPERTY() 
 	double HandleWeight = 1000.0;
