MeshModelingTools: Added a CubeGrid tool, which allows for blocky boolean operations on an orientable power-of-two grid.

#rnx

diff --git a/Engine/Plugins/Editor/ModelingToolsEditorMode/Source/ModelingToolsEditorMode/Private/ModelingToolsActions.cpp b/Engine/Plugins/Editor/ModelingToolsEditorMode/Source/ModelingToolsEditorMode/Private/ModelingToolsActions.cpp
--- a/Engine/Plugins/Editor/ModelingToolsEditorMode/Source/ModelingToolsEditorMode/Private/ModelingToolsActions.cpp
+++ b/Engine/Plugins/Editor/ModelingToolsEditorMode/Source/ModelingToolsEditorMode/Private/ModelingToolsActions.cpp
@@ -7,6 +7,7 @@
 #include "MeshVertexSculptTool.h"
 #include "MeshGroupPaintTool.h"
 #include "MeshInspectorTool.h"
+#include "CubeGridTool.h"
 #include "DrawPolygonTool.h"
 #include "EditMeshPolygonsTool.h"
 #include "ShapeSprayTool.h"
@@ -89,6 +90,7 @@ void FModelingToolActionCommands::RegisterAllToolActions()
 	FMeshSelectionToolActionCommands::Register();
 	FEditMeshMaterialsToolActionCommands::Register();
 	FMeshPlaneCutToolActionCommands::Register();
+	FCubeGridToolActionCommands::Register();
 	FEditMeshPolygonsToolActionCommands::Register();
 	FDrawAndRevolveToolActionCommands::Register();
 }
@@ -104,6 +106,7 @@ void FModelingToolActionCommands::UnregisterAllToolActions()
 	FMeshSelectionToolActionCommands::Unregister();
 	FEditMeshMaterialsToolActionCommands::Unregister();
 	FMeshPlaneCutToolActionCommands::Unregister();
+	FCubeGridToolActionCommands::Unregister();
 	FEditMeshPolygonsToolActionCommands::Unregister();
 	FDrawAndRevolveToolActionCommands::Unregister();
 }
@@ -147,6 +150,10 @@ void FModelingToolActionCommands::UpdateToolCommandBinding(UInteractiveTool* Too
 	{
 		UPDATE_BINDING(FMeshPlaneCutToolActionCommands);
 	}
+	else if (ExactCast<UCubeGridTool>(Tool) != nullptr)
+	{
+		UPDATE_BINDING(FCubeGridToolActionCommands);
+	}
 	else if (ExactCast<UEditMeshPolygonsTool>(Tool) != nullptr)
 	{
 		UPDATE_BINDING(FEditMeshPolygonsToolActionCommands);
@@ -184,6 +191,7 @@ DEFINE_TOOL_ACTION_COMMANDS(FDrawPolygonToolActionCommands, "ModelingToolsDrawPo
 DEFINE_TOOL_ACTION_COMMANDS(FMeshSelectionToolActionCommands, "ModelingToolsMeshSelectionTool", "Modeling Tools - Mesh Selection Tool", UMeshSelectionTool);
 DEFINE_TOOL_ACTION_COMMANDS(FEditMeshMaterialsToolActionCommands, "ModelingToolsEditMeshMaterials", "Modeling Tools - Edit Materials Tool", UEditMeshMaterialsTool);
 DEFINE_TOOL_ACTION_COMMANDS(FMeshPlaneCutToolActionCommands, "ModelingToolsMeshPlaneCutTool", "Modeling Tools - Mesh Plane Cut Tool", UPlaneCutTool);
+DEFINE_TOOL_ACTION_COMMANDS(FCubeGridToolActionCommands, "ModelingToolsCubeGridTool", "Modeling Tools - Cube Grid Tool", UCubeGridTool);
 DEFINE_TOOL_ACTION_COMMANDS(FEditMeshPolygonsToolActionCommands, "ModelingToolsEditMeshPolygonsTool", "Modeling Tools - Edit Mesh Polygons Tool", UEditMeshPolygonsTool);
 DEFINE_TOOL_ACTION_COMMANDS(FDrawAndRevolveToolActionCommands, "ModelingToolsDrawAndRevolveTool", "Modeling Tools - Draw-and-Revolve Tool", UDrawAndRevolveTool);
 
diff --git a/Engine/Plugins/Editor/ModelingToolsEditorMode/Source/ModelingToolsEditorMode/Private/ModelingToolsEditorMode.cpp b/Engine/Plugins/Editor/ModelingToolsEditorMode/Source/ModelingToolsEditorMode/Private/ModelingToolsEditorMode.cpp
--- a/Engine/Plugins/Editor/ModelingToolsEditorMode/Source/ModelingToolsEditorMode/Private/ModelingToolsEditorMode.cpp
+++ b/Engine/Plugins/Editor/ModelingToolsEditorMode/Source/ModelingToolsEditorMode/Private/ModelingToolsEditorMode.cpp
@@ -33,6 +33,7 @@
 #include "ConvertToPolygonsTool.h"
 #include "AddPrimitiveTool.h"
 #include "AddPatchTool.h"
+#include "CubeGridTool.h"
 #include "RevolveBoundaryTool.h"
 #include "SmoothMeshTool.h"
 #include "OffsetMeshTool.h"
@@ -411,6 +412,8 @@ void UModelingToolsEditorMode::Enter()
 	auto ShapeSprayToolBuilder = NewObject<UShapeSprayToolBuilder>();
 	RegisterTool(ToolManagerCommands.BeginShapeSprayTool, TEXT("BeginShapeSprayTool"), ShapeSprayToolBuilder);
 
+	auto CubeGridToolBuilder = NewObject<UCubeGridToolBuilder>();
+	RegisterTool(ToolManagerCommands.BeginCubeGridTool, TEXT("BeginCubeGridTool"), CubeGridToolBuilder);
 
 	//
 	// vertex deform tools
diff --git a/Engine/Plugins/Editor/ModelingToolsEditorMode/Source/ModelingToolsEditorMode/Private/ModelingToolsEditorModeToolkit.cpp b/Engine/Plugins/Editor/ModelingToolsEditorMode/Source/ModelingToolsEditorMode/Private/ModelingToolsEditorModeToolkit.cpp
--- a/Engine/Plugins/Editor/ModelingToolsEditorMode/Source/ModelingToolsEditorMode/Private/ModelingToolsEditorModeToolkit.cpp
+++ b/Engine/Plugins/Editor/ModelingToolsEditorMode/Source/ModelingToolsEditorMode/Private/ModelingToolsEditorModeToolkit.cpp
@@ -615,6 +615,7 @@ void FModelingToolsEditorModeToolkit::BuildToolPalette_Experimental(FName Palett
 	{
 		ToolbarBuilder.AddToolBarButton(Commands.BeginPolyEditTool);
 		ToolbarBuilder.AddToolBarButton(Commands.BeginPolyDeformTool);
+		ToolbarBuilder.AddToolBarButton(Commands.BeginCubeGridTool);
 		ToolbarBuilder.AddToolBarButton(Commands.BeginMeshBooleanTool);
 		ToolbarBuilder.AddToolBarButton(Commands.BeginCutMeshWithMeshTool);
 		ToolbarBuilder.AddToolBarButton(Commands.BeginSubdividePolyTool);
diff --git a/Engine/Plugins/Editor/ModelingToolsEditorMode/Source/ModelingToolsEditorMode/Private/ModelingToolsManagerActions.cpp b/Engine/Plugins/Editor/ModelingToolsEditorMode/Source/ModelingToolsEditorMode/Private/ModelingToolsManagerActions.cpp
--- a/Engine/Plugins/Editor/ModelingToolsEditorMode/Source/ModelingToolsEditorMode/Private/ModelingToolsManagerActions.cpp
+++ b/Engine/Plugins/Editor/ModelingToolsEditorMode/Source/ModelingToolsEditorMode/Private/ModelingToolsManagerActions.cpp
@@ -40,6 +40,7 @@ void FModelingToolsManagerCommands::RegisterCommands()
 	UI_COMMAND(BeginShapeSprayTool, "Spray", "Shape Spray", EUserInterfaceActionType::ToggleButton, FInputChord());
 
 	UI_COMMAND(BeginSculptMeshTool, "VSclpt", "Vertex Sculpting", EUserInterfaceActionType::ToggleButton, FInputChord());
+	UI_COMMAND(BeginCubeGridTool, "CubeGr", "Create block-like geometry from a repositionable grid", EUserInterfaceActionType::ToggleButton, FInputChord());
 	UI_COMMAND(BeginPolyEditTool, "PolyEd", "Edit Mesh via PolyGroups", EUserInterfaceActionType::ToggleButton, FInputChord());
 	UI_COMMAND(BeginTriEditTool, "TriEd", "Edit Mesh via Triangles", EUserInterfaceActionType::ToggleButton, FInputChord());
 	UI_COMMAND(BeginPolyDeformTool, "PolyDef", "Deform Mesh via PolyGroups", EUserInterfaceActionType::ToggleButton, FInputChord());
diff --git a/Engine/Plugins/Editor/ModelingToolsEditorMode/Source/ModelingToolsEditorMode/Public/ModelingToolsActions.h b/Engine/Plugins/Editor/ModelingToolsEditorMode/Source/ModelingToolsEditorMode/Public/ModelingToolsActions.h
--- a/Engine/Plugins/Editor/ModelingToolsEditorMode/Source/ModelingToolsEditorMode/Public/ModelingToolsActions.h
+++ b/Engine/Plugins/Editor/ModelingToolsEditorMode/Source/ModelingToolsEditorMode/Public/ModelingToolsActions.h
@@ -83,6 +83,7 @@ DECLARE_TOOL_ACTION_COMMANDS(FDrawPolygonToolActionCommands);
 DECLARE_TOOL_ACTION_COMMANDS(FMeshSelectionToolActionCommands);
 DECLARE_TOOL_ACTION_COMMANDS(FEditMeshMaterialsToolActionCommands);
 DECLARE_TOOL_ACTION_COMMANDS(FMeshPlaneCutToolActionCommands);
+DECLARE_TOOL_ACTION_COMMANDS(FCubeGridToolActionCommands);
 DECLARE_TOOL_ACTION_COMMANDS(FEditMeshPolygonsToolActionCommands);
 DECLARE_TOOL_ACTION_COMMANDS(FDrawAndRevolveToolActionCommands);
 
diff --git a/Engine/Plugins/Editor/ModelingToolsEditorMode/Source/ModelingToolsEditorMode/Public/ModelingToolsManagerActions.h b/Engine/Plugins/Editor/ModelingToolsEditorMode/Source/ModelingToolsEditorMode/Public/ModelingToolsManagerActions.h
--- a/Engine/Plugins/Editor/ModelingToolsEditorMode/Source/ModelingToolsEditorMode/Public/ModelingToolsManagerActions.h
+++ b/Engine/Plugins/Editor/ModelingToolsEditorMode/Source/ModelingToolsEditorMode/Public/ModelingToolsManagerActions.h
@@ -33,6 +33,7 @@ public:
 	TSharedPtr<FUICommandInfo> BeginShapeSprayTool;
 
 	TSharedPtr<FUICommandInfo> BeginSculptMeshTool;
+	TSharedPtr<FUICommandInfo> BeginCubeGridTool;
 	TSharedPtr<FUICommandInfo> BeginPolyEditTool;
 	TSharedPtr<FUICommandInfo> BeginSubdividePolyTool;
 	TSharedPtr<FUICommandInfo> BeginTriEditTool;
diff --git a/Engine/Plugins/Experimental/MeshModelingToolsetExp/Source/MeshModelingToolsExp/Private/CubeGridTool.cpp b/Engine/Plugins/Experimental/MeshModelingToolsetExp/Source/MeshModelingToolsExp/Private/CubeGridTool.cpp
new file mode 100644
--- /dev/null
+++ b/Engine/Plugins/Experimental/MeshModelingToolsetExp/Source/MeshModelingToolsExp/Private/CubeGridTool.cpp
@@ -0,0 +1,2068 @@
+// Copyright Epic Games, Inc. All Rights Reserved.
+
+#include "CubeGridTool.h"
+
+#include "BaseBehaviors/ClickDragBehavior.h"
+#include "BaseBehaviors/SingleClickBehavior.h"
+#include "BaseBehaviors/MouseHoverBehavior.h"
+#include "BaseGizmos/GizmoMath.h"
+#include "BaseGizmos/TransformGizmoUtil.h"
+#include "BaseGizmos/TransformProxy.h"
+#include "CompositionOps/CubeGridBooleanOp.h"
+#include "Distance/DistLine3Ray3.h"
+#include "Drawing/PreviewGeometryActor.h"
+#include "Drawing/LineSetComponent.h"
+#include "DynamicMesh/DynamicMeshChangeTracker.h"
+#include "DynamicMesh/MeshTransforms.h"
+#include "DynamicMeshToMeshDescription.h"
+#include "InteractiveToolChange.h"
+#include "InteractiveToolManager.h"
+#include "InputState.h"
+#include "Mechanics/DragAlignmentMechanic.h"
+#include "MeshOpPreviewHelpers.h" //UMeshOpPreviewWithBackgroundCompute
+#include "MeshDescriptionToDynamicMesh.h"
+#include "ModelingObjectsCreationAPI.h"
+#include "ModelingToolTargetUtil.h"
+#include "Properties/MeshMaterialProperties.h"
+#include "PropertySets/CreateMeshObjectTypeProperties.h"
+#include "Selection/ToolSelectionUtil.h"
+#include "ToolContextInterfaces.h"
+#include "ToolTargetManager.h"
+#include "ToolTargets/ToolTarget.h"
+#include "TargetInterfaces/MeshDescriptionCommitter.h"
+#include "TargetInterfaces/MeshDescriptionProvider.h"
+#include "TargetInterfaces/PrimitiveComponentBackedTarget.h"
+#include "ToolSceneQueriesUtil.h"
+#include "ToolSetupUtil.h"
+
+#include "ExplicitUseGeometryMathTypes.h"		// using UE::Geometry::(math types)
+using namespace UE::Geometry;
+
+#define LOCTEXT_NAMESPACE "UCubeGridTool"
+
+// General note: This tool used to largely operate in cube grid space. However it turned out
+// to be valuable to be able to keep a selection more or less constant while resizing the grid,
+// and this was easier to do by generally operating in just the frame space of the grid (not
+// scaled into the grid size).
+
+namespace CubeGridToolLocals
+{
+	FText PushPullModeMessage = LOCTEXT("CubeGridPushPullModeDescription", "Select grid cells and push/pull them to create geometry. "
+		"Refer to side panel for shortcuts.");
+
+	FText CornerModeMessage = LOCTEXT("CubeGridCornerModeDescription", "Toggle corner selection for push/pulling by clicking or dragging. "
+		"Press Enter or click \"Done\" in the side panel to accept the result.");
+
+	const FText SelectionChangeTransactionName = LOCTEXT("SelectionChangeTransaction", "Cube Grid Selection Change");
+	const FText ModeChangeTransactionName = LOCTEXT("ModeChangeTransaction", "Cube Grid Mode Change");
+
+	const FString& HoverLineSetID(TEXT("HoverLines"));
+	const FString& GridLineSetID(TEXT("GridLines"));
+	const FString& SelectionLineSetID(TEXT("SelectionLines"));
+	const FString& CornerModeLineSetID(TEXT("CornerModeLines"));
+
+	const FColor HoverLineColor(255, 255, 128); // Pale yellow
+	const double HoverLineThickness = 2;
+	const double HoverLineDepthBias = 0.1;
+
+	const FColor SelectionLineColor(255, 128, 0); // Orange
+	const double SelectionLineDepthBias = 0.1;
+	const double SelectionGridLineThickness = 1;
+	const double SelectionMainLineThickness = 3;
+
+	const FColor GridLineColor(200, 200, 200);
+	const double GridLineDepthBias = 0.05;
+	const double GridLineThickness = 0.5;
+
+	const FColor UnselectedCornerLineColor(64, 0, 128); // dark purple
+	const FColor SelectedCornerLineColor = FColor::Yellow;
+	const double CornerLineThickness = 3;
+	const int32 CornerCircleNumSteps = 12;
+	const FColor CornerModeWireframeColor = FColor::Red;
+	const double CornerModeWireframeThickness = SelectionGridLineThickness;
+	const double CornerModeWireframeDepthBias = 0.1;
+
+
+	/**
+	 * Undoes the actual mesh change (i.e. after boolean operations)
+	 */
+	class FCubeGridToolMeshChange : public FToolCommandChange
+	{
+	public:
+		FCubeGridToolMeshChange(TUniquePtr<UE::Geometry::FDynamicMeshChange> MeshChangeIn)
+			: MeshChange(MoveTemp(MeshChangeIn))
+		{};
+
+		virtual void Apply(UObject* Object) override
+		{
+			UCubeGridTool* Tool = Cast<UCubeGridTool>(Object);
+			Tool->UpdateUsingMeshChange(*MeshChange, false);
+		}
+		virtual void Revert(UObject* Object) override
+		{
+			UCubeGridTool* Tool = Cast<UCubeGridTool>(Object);
+			Tool->UpdateUsingMeshChange(*MeshChange, true);
+		}
+
+		virtual FString ToString() const override
+		{
+			return TEXT("CubeGridToolLocals::FCubeGridToolMeshChange");
+		}
+
+	protected:
+		TUniquePtr<UE::Geometry::FDynamicMeshChange> MeshChange;
+	};
+
+	/** Undoes selection changes */
+	class FCubeGridToolSelectionChange : public FToolCommandChange
+	{
+	public:
+		FCubeGridToolSelectionChange(
+			bool bHaveStartSelectionBeforeIn, bool bHaveStartSelectionAfterIn,
+			const UCubeGridTool::FSelection& SelectionBeforeIn,
+			const UCubeGridTool::FSelection& SelectionAfterIn
+			)
+			: bHaveStartSelectionBefore(bHaveStartSelectionBeforeIn)
+			, bHaveStartSelectionAfter(bHaveStartSelectionAfterIn)
+			, SelectionBefore(SelectionBeforeIn)
+			, SelectionAfter(SelectionAfterIn)
+		{};
+
+		virtual void Apply(UObject* Object) override
+		{
+			UCubeGridTool* Tool = Cast<UCubeGridTool>(Object);
+			if (!bHaveStartSelectionAfter)
+			{
+				Tool->ClearSelection(false);
+			}
+			else
+			{
+				Tool->SetSelection(SelectionAfter, false);
+			}
+		}
+		virtual void Revert(UObject* Object) override
+		{
+
+			UCubeGridTool* Tool = Cast<UCubeGridTool>(Object);
+			if (!bHaveStartSelectionBefore)
+			{
+				Tool->ClearSelection(false);
+			}
+			else
+			{
+				Tool->SetSelection(SelectionBefore, false);
+			}
+		}
+
+		virtual FString ToString() const override
+		{
+			return TEXT("CubeGridToolLocals::FCubeGridToolSelectionChange");
+		}
+
+	protected:
+		bool bHaveStartSelectionBefore;
+		bool bHaveStartSelectionAfter;
+		UCubeGridTool::FSelection SelectionBefore;
+		UCubeGridTool::FSelection SelectionAfter;
+	};
+
+	/** Undoes activating "corner" mode. Not redoable. */
+	class FCubeGridToolModeChange : public FToolCommandChange
+	{
+	public:
+		FCubeGridToolModeChange() {};
+
+		virtual bool HasExpired(UObject* Object) const override
+		{
+			return Cast<UCubeGridTool>(Object)->IsInDefaultMode();
+		}
+
+		virtual void Apply(UObject* Object) override
+		{
+			return;
+		}
+		virtual void Revert(UObject* Object) override
+		{
+			Cast<UCubeGridTool>(Object)->RevertToDefaultMode();
+		}
+
+		virtual FString ToString() const override
+		{
+			return TEXT("CubeGridToolLocals::FCubeGridToolModeChange");
+		}
+	};
+	
+	// Attach a frame to the box such that Z points along the given direction.
+	FOrientedBox3d ConvertToOrientedBox(const FAxisAlignedBox3d& Box, FCubeGrid::EFaceDirection Direction)
+	{
+		int FlatDim = FCubeGrid::DirToFlatDim(Direction);
+
+		FVector3d GridSpaceZ = FCubeGrid::DirToNormal(Direction);
+		FVector3d GridSpaceX = (FlatDim == 0) ? FVector3d::UnitY() : FVector3d::UnitX();
+		FVector3d GridSpaceY = GridSpaceZ.Cross(GridSpaceX);
+
+		FVector3d BoxExtents = Box.Extents();
+		FVector3d Extents = BoxExtents; // The case if selection axis is Z
+
+		if (FlatDim == 0)
+		{
+			// If selection dir was along x axis, then frame z is x, frame x is y, and frame y is x
+			BoxExtents = FVector3d(BoxExtents[1], BoxExtents[2], BoxExtents[0]);
+		}
+		else if (FlatDim == 1)
+		{
+			// If selection dir was along y axis, then frame z is y, frame x is x, and frame y is z
+			Swap(BoxExtents[2], BoxExtents[1]);
+		}
+
+		return FOrientedBox3d(FFrame3d(Box.Center(), GridSpaceX, GridSpaceY, GridSpaceZ), BoxExtents);
+	}
+
+	void GetNewSelectionFaceInBox(const FCubeGrid& Grid, const FAxisAlignedBox3d& Box, 
+		const FCubeGrid::FCubeFace& FaceIn, FCubeGrid::FCubeFace& FaceOut)
+	{
+		// Start at the corner of the selection and move a little bit to make sure
+		// you're in the first face in the corner.
+		FVector3d TowardOtherCorner(Box.Max - Box.Min);
+		TowardOtherCorner.Normalize();
+
+		FVector3d PointOnDesiredFace = (Box.Min + Grid.GetCurrentGridCellSize() * TowardOtherCorner)
+			/ Grid.GetCurrentGridCellSize();
+
+		FaceOut = FCubeGrid::FCubeFace(
+			PointOnDesiredFace,
+			FaceIn.GetDirection(),
+			Grid.GetGridPower());
+	}
+
+	/**
+	 * Given grid, start point, extrude direction, and number of blocks to extrude, produces
+	 * the frame-space extrusion distance. If the start point is not on the grid, the first
+	 * "block" is counted as the distance to get back onto grid in the extrusion direction.
+	 */
+	double GetFrameSpaceExtrudeDist(const FCubeGrid& CubeGrid, const FVector3d& FrameSpaceStartPoint,
+		int32 CurrentExtrudeAmount, FCubeGrid::EFaceDirection Direction)
+	{
+		double GridSpaceExtrudeDist = CurrentExtrudeAmount; // Will be adjusted
+
+		int FlatDim = FCubeGrid::DirToFlatDim(Direction);
+		double GridSpaceExtrudeCoord = FrameSpaceStartPoint[FlatDim] / CubeGrid.GetCurrentGridCellSize();
+		double ClosestOnGridCoord = FMath::RoundToDouble(GridSpaceExtrudeCoord);
+
+		// See if we're actually off-grid
+		if (FMath::Abs(GridSpaceExtrudeCoord - ClosestOnGridCoord) > KINDA_SMALL_NUMBER)
+		{
+			double NextOnGrid = CurrentExtrudeAmount > 0 ? FMath::CeilToDouble(GridSpaceExtrudeCoord)
+				: FMath::FloorToDouble(GridSpaceExtrudeCoord);
+			GridSpaceExtrudeDist += (NextOnGrid - GridSpaceExtrudeCoord);
+			GridSpaceExtrudeDist += (CurrentExtrudeAmount > 0 ? -1 : 1);
+		}
+
+		return GridSpaceExtrudeDist * CubeGrid.GetCurrentGridCellSize();
+	}
+
+	bool IsAnyCornerSelected(bool CornerSelectedFlags[])
+	{
+		for (int i = 0; i < 4; ++i)
+		{
+			if (CornerSelectedFlags[i])
+			{
+				return true;
+			}
+		}
+		return false;
+	}
+
+	/** 
+	 * @param CornerVector Lines to draw on the corners, for instance to show the depth direction
+	 */
+	void DrawGridRectangle(ULineSetComponent& LineSet, const FCubeGrid& Grid,
+		const FVector3d& GridMin, const FVector3d& GridMax,
+		const FColor& Color, double Thickness,
+		double DepthBias, const FVector3d* CornerVector = nullptr)
+	{
+		// We'll step from max/min along one of the differing dimensions to get corners
+		int32 DifferingDimension = (GridMin[0] == GridMax[0]) ? 1 : 0;
+
+		FVector3d Corner1 = GridMin;
+		Corner1[DifferingDimension] = GridMax[DifferingDimension];
+
+		FVector3d Corner2 = GridMax;
+		Corner2[DifferingDimension] = GridMin[DifferingDimension];
+
+		LineSet.AddLine(GridMin, Corner1, Color, Thickness, DepthBias);
+		LineSet.AddLine(GridMin, Corner2, Color, Thickness, DepthBias);
+		LineSet.AddLine(GridMax, Corner1, Color, Thickness, DepthBias);
+		LineSet.AddLine(GridMax, Corner2, Color, Thickness, DepthBias);
+
+		if (CornerVector)
+		{
+			LineSet.AddLine(GridMin, GridMin + *CornerVector, Color, Thickness, DepthBias);
+			LineSet.AddLine(GridMax, GridMax + *CornerVector, Color, Thickness, DepthBias);
+			LineSet.AddLine(Corner1, Corner1 + *CornerVector, Color, Thickness, DepthBias);
+			LineSet.AddLine(Corner2, Corner2 + *CornerVector, Color, Thickness, DepthBias);
+		}
+	}
+
+	const int32 MAX_NUM_INTERIOR_GRID_LINES = 1000;
+
+	void DrawGridSection(ULineSetComponent& LineSet, const FCubeGrid& Grid, 
+		const FAxisAlignedBox3d& BBox,
+		const FColor& Color, double Thickness, 
+		double DepthBias, const FVector3d* CornerVector = nullptr)
+	{
+		// Draw the boundary
+		DrawGridRectangle(LineSet, Grid, BBox.Min, BBox.Max, Color, Thickness, DepthBias, CornerVector);
+
+		// Find the two nonzero dimensions of the box
+		const FVector3d BoxDimensions = BBox.Max - BBox.Min;
+		int32 Dim1 = BoxDimensions[0] != 0 ? 0 : 2;
+		int32 Dim2 = BoxDimensions[1] != 0 ? 1 : 2;
+
+		double StepSize = Grid.GetCurrentGridCellSize();
+
+		// Draw the inside only if there aren't too many lines to draw (approximate)
+		if (BoxDimensions[Dim1] / StepSize + BoxDimensions[Dim2] / StepSize > MAX_NUM_INTERIOR_GRID_LINES)
+		{
+			return;
+		}
+
+		// Draws lines that lie in the DimToDrawAlong dimension, along border in DimToStepAlong dimension
+		auto DrawParallelInteriorLines = [&BBox, &LineSet, &Grid, &Color, &BoxDimensions, StepSize,
+			Thickness, DepthBias](int32 DimToStepAlong, int32 DimToDrawAlong)
+		{
+			FVector3d BorderDirection = FVector3d::Zero();
+			BorderDirection[DimToStepAlong] = StepSize;
+
+			int NumSteps = BoxDimensions[DimToStepAlong] / StepSize;
+			if (NumSteps * StepSize == BoxDimensions[DimToStepAlong])
+			{
+				--NumSteps;
+			}
+
+			for (int32 i = 0; i < NumSteps; ++i)
+			{
+				FVector3d SidePoint = BBox.Min + BorderDirection * (i + 1);
+				FVector3d OtherSidePoint = SidePoint;
+				OtherSidePoint[DimToDrawAlong] = BBox.Max[DimToDrawAlong];
+
+				LineSet.AddLine(SidePoint, OtherSidePoint, Color, Thickness, DepthBias);
+			}
+		};
+		DrawParallelInteriorLines(Dim1, Dim2);
+		DrawParallelInteriorLines(Dim2, Dim1);
+	}
+
+	/** Given a world hit, get a hit face. */
+	void ConvertToFaceHit(const FCubeGrid& CubeGrid, ECubeGridToolFaceSelectionMode SelectionMode, 
+		 const FRay& WorldRay, double HitT, const FVector3d& Normal, FCubeGrid::FCubeFace& FaceOut,
+		 double Tolerance)
+	{
+		FVector3d WorldPoint = WorldRay.PointAt(HitT);
+		bool bSuccess = false;
+		switch (SelectionMode)
+		{
+		case ECubeGridToolFaceSelectionMode::OutsideBasedOnNormal:
+			bSuccess = CubeGrid.GetHitGridFaceBasedOnRay(WorldPoint, Normal, FaceOut, false, Tolerance);
+			break;
+		case ECubeGridToolFaceSelectionMode::InsideBasedOnNormal:
+			bSuccess = CubeGrid.GetHitGridFaceBasedOnRay(WorldPoint, Normal, FaceOut, true, Tolerance);
+			break;
+		case ECubeGridToolFaceSelectionMode::OutsideBasedOnViewRay:
+			bSuccess = CubeGrid.GetHitGridFaceBasedOnRay(WorldPoint, -WorldRay.Direction, FaceOut, false, Tolerance);
+			break;
+		case ECubeGridToolFaceSelectionMode::InsideBasedOnViewRay:
+			bSuccess = CubeGrid.GetHitGridFaceBasedOnRay(WorldPoint, -WorldRay.Direction, FaceOut, true, Tolerance);
+			break;
+		}
+
+		ensureMsgf(bSuccess, TEXT("CubeGridTool: Unable to convert hit location to proper grid face."));
+	}
+
+	/**
+	 * Given a world ray and a flat box in grid frame space, intersect the ray with the plane containing the box, 
+	 * find the selected cell in the cube grid, and project that cell onto the same plane to produce a new frame 
+	 * space box.
+	 * 
+	 * @param bExpandBoxWithStartBox If true, the output is expanded to contain the original box.
+	 * @return true If the plane was actually hit.
+	 */
+	bool GetCoplanarFrameSpaceSelectedBox(const FCubeGrid& CubeGrid, const FRay& WorldRay, 
+		const FAxisAlignedBox3d& StartBox, bool bExpandOutputBoxWithStartBox, FAxisAlignedBox3d& BoxOut)
+	{
+		FVector3d BoxDims = StartBox.Max - StartBox.Min;
+		int FlatDim = 0;
+		for (int i = 0; i < 3; ++i)
+		{
+			if (BoxDims[i] == 0)
+			{
+				FlatDim = i;
+				break;
+			}
+		}
+
+		bool bIntersects = false;
+		FVector IntersectionPoint;
+		GizmoMath::RayPlaneIntersectionPoint(
+			CubeGrid.GetFrame().FromFramePoint(StartBox.Min), 
+			CubeGrid.GetFrame().GetAxis(FlatDim), 
+			WorldRay.Origin, WorldRay.Direction, 
+			bIntersects, IntersectionPoint);
+
+		if (!bIntersects)
+		{
+			return false;
+		}
+
+		FVector3d GridSpaceIntersection = CubeGrid.ToGridPoint(IntersectionPoint);
+
+		FVector3d FrameSpaceMin(
+			FMath::Floor(GridSpaceIntersection.X) * CubeGrid.GetCurrentGridCellSize(),
+			FMath::Floor(GridSpaceIntersection.Y) * CubeGrid.GetCurrentGridCellSize(),
+			FMath::Floor(GridSpaceIntersection.Z) * CubeGrid.GetCurrentGridCellSize());
+		FVector3d FrameSpaceMax(
+			FMath::CeilToDouble(GridSpaceIntersection.X) * CubeGrid.GetCurrentGridCellSize(),
+			FMath::CeilToDouble(GridSpaceIntersection.Y) * CubeGrid.GetCurrentGridCellSize(),
+			FMath::CeilToDouble(GridSpaceIntersection.Z) * CubeGrid.GetCurrentGridCellSize());
+
+		// Project the cell we got back onto the plane
+		FrameSpaceMin[FlatDim] = StartBox.Min[FlatDim];
+		FrameSpaceMax[FlatDim] = StartBox.Min[FlatDim];
+
+		BoxOut = FAxisAlignedBox3d(FrameSpaceMin, FrameSpaceMax);
+		if (bExpandOutputBoxWithStartBox)
+		{
+			BoxOut.Contain(StartBox);
+		}
+
+		return true;
+	}
+}
+
+const FToolTargetTypeRequirements& UCubeGridToolBuilder::GetTargetRequirements() const
+{
+	static FToolTargetTypeRequirements TypeRequirements({
+		UMaterialProvider::StaticClass(),
+		UMeshDescriptionCommitter::StaticClass(),
+		UMeshDescriptionProvider::StaticClass(),
+		UPrimitiveComponentBackedTarget::StaticClass()
+		});
+	return TypeRequirements;
+}
+
+bool UCubeGridToolBuilder::CanBuildTool(const FToolBuilderState& SceneState) const
+{
+	return SceneState.TargetManager->CountSelectedAndTargetable(SceneState, GetTargetRequirements()) <= 1;
+}
+
+UInteractiveTool* UCubeGridToolBuilder::BuildTool(const FToolBuilderState& SceneState) const
+{
+	UCubeGridTool* NewTool = NewObject<UCubeGridTool>(SceneState.ToolManager);
+
+	TObjectPtr<UToolTarget> Target = SceneState.TargetManager->BuildFirstSelectedTargetable(SceneState, GetTargetRequirements());
+	NewTool->SetTarget(Target); // May be null
+	NewTool->SetWorld(SceneState.World);
+
+	return NewTool;
+}
+
+
+
+void UCubeGridTool::InvalidatePreview(bool bUpdateCornerLineSet)
+{
+	using namespace CubeGridToolLocals;
+
+	// Do the line set first before we get to all the early returns
+	if (bUpdateCornerLineSet && Mode == EMode::Corner)
+	{
+		UpdateCornerModeLineSet();
+	}
+
+	if (CurrentExtrudeAmount == 0)
+	{
+		// Reset the preview
+		Preview->CancelCompute();
+		if (bPreviewMayDiffer)
+		{
+			Preview->PreviewMesh->UpdatePreview(CurrentMesh.Get());
+			bPreviewMayDiffer = false;
+			bWaitingToApplyPreview = false;
+		}
+		return;
+	}
+	else if (CurrentMesh->TriangleCount() == 0 && CurrentExtrudeAmount < 0 
+		&& Mode != EMode::Corner)
+	{
+		// We're subtracting from an empty mesh. Just slide the selection.
+		// This will also reset the preview.
+		SlideSelection(CurrentExtrudeAmount, true);
+		bPreviewMayDiffer = false;
+		bWaitingToApplyPreview = false;
+		return;
+	}
+
+	// If we didn't start with an existing mesh, and we are adding to an empty
+	// starting mesh, set the transform such that it is in the (grid) minimum of
+	// the selection. This frequently (though not always, in the case of
+	// a pyramid) places the pivot in a handy corner for snapping.
+	if (!Target && CurrentMesh->TriangleCount() == 0 && CurrentExtrudeAmount > 0
+		&& ensure(bHaveSelection))
+	{
+		FFrame3d GridFrame = CubeGrid->GetFrame();
+		GridFrame.Origin = GridFrame.FromFramePoint(Selection.Box.Min);
+		CurrentMeshTransform = GridFrame.ToTransform();
+
+		GridGizmoAlignmentMechanic->InitializeDeformedMeshRayCast([this]() { return MeshSpatial.Get(); }, CurrentMeshTransform, nullptr);
+	}
+
+	// Finally invalidate the preview
+	Preview->InvalidateResult();
+	bPreviewMayDiffer = true;
+}
+
+TUniquePtr<FDynamicMeshOperator> UCubeGridTool::MakeNewOperator()
+{
+	using namespace CubeGridToolLocals;
+
+	// Figure out how far to extrude in grid space. This only becomes tricky if the selection
+	// is no longer on grid because we changed grid power.
+	double FrameSpaceExtrudeAmount = GetFrameSpaceExtrudeDist(*CubeGrid, Selection.Box.Min, CurrentExtrudeAmount, Selection.Direction);
+
+	FOrientedBox3d FrameSpaceBox = ConvertToOrientedBox(Selection.Box, Selection.Direction);
+
+	// Give the selection box depth
+	FrameSpaceBox.Frame.Origin += (FrameSpaceExtrudeAmount / 2.0f) * FCubeGrid::DirToNormal(Selection.Direction);
+	FrameSpaceBox.Extents.Z = FMath::Abs(FrameSpaceExtrudeAmount) / 2.0f;
+
+	// Translate the oriented box from grid space to world space
+	const FFrame3d& GridFrame = CubeGrid->GetFrame();
+	FOrientedBox3d WorldBox(FFrame3d(GridFrame.FromFrame(FrameSpaceBox.Frame)), FrameSpaceBox.Extents);
+
+	// Make the op.
+	TUniquePtr<FCubeGridBooleanOp> Op = MakeUnique<FCubeGridBooleanOp>();
+	Op->InputMesh = CurrentMesh;
+	Op->InputTransform = CurrentMeshTransform;
+	Op->bKeepInputTransform = true;
+	Op->WorldBox = WorldBox;
+	Op->bSubtract = CurrentExtrudeAmount < 0;
+	Op->bTrackChangedTids = true;
+
+	if (Mode == EMode::Corner)
+	{
+		Op->CornerInfo = MakeShared<FCubeGridBooleanOp::FCornerInfo>();
+		for (int i = 0; i < 4; ++i)
+		{
+			Op->CornerInfo->WeldedAtBase[i] = !CornerSelectedFlags[i];
+		}
+
+		Op->bCrosswiseDiagonal = Settings->bCrosswiseDiagonal;
+	}
+
+	return Op;
+}
+
+void UCubeGridTool::SlideSelection(int32 Amount, bool bEmitChange)
+{
+	using namespace CubeGridToolLocals;
+
+	FVector3d FrameSpaceDisplacement = FCubeGrid::DirToNormal(Selection.Direction)
+		* GetFrameSpaceExtrudeDist(*CubeGrid, Selection.StartBox.Min, Amount, Selection.Direction);
+
+	FSelection NewSelection = Selection;
+	NewSelection.StartBox = FAxisAlignedBox3d(
+		Selection.StartBox.Min + FrameSpaceDisplacement,
+		Selection.StartBox.Max + FrameSpaceDisplacement);
+	NewSelection.Box = FAxisAlignedBox3d(
+		Selection.Box.Min + FrameSpaceDisplacement,
+		Selection.Box.Max + FrameSpaceDisplacement);
+
+	SetSelection(NewSelection, true);
+}
+
+void UCubeGridTool::SetSelection(const FSelection& NewSelection, bool bEmitChange)
+{
+	using namespace CubeGridToolLocals;
+
+	// Clear op/preview
+	if (Mode != EMode::Corner)
+	{
+		CurrentExtrudeAmount = 0;
+		InvalidatePreview();
+	}
+
+	if (bEmitChange && (!bHaveSelection || Selection != NewSelection))
+	{
+		GetToolManager()->EmitObjectChange(this,
+			MakeUnique<FCubeGridToolSelectionChange>(bHaveSelection, true, Selection, NewSelection),
+			SelectionChangeTransactionName);
+	}
+	Selection = NewSelection;
+	bHaveSelection = true;
+
+	UpdateSelectionLineSet();
+	if (Mode == EMode::Corner)
+	{
+		InvalidatePreview();
+	}
+}
+
+void UCubeGridTool::ClearSelection(bool bEmitChange)
+{
+	using namespace CubeGridToolLocals;
+
+	if (bEmitChange && bHaveSelection)
+	{
+		GetToolManager()->EmitObjectChange(this,
+			MakeUnique<FCubeGridToolSelectionChange>(bHaveSelection, false, Selection, Selection),
+			SelectionChangeTransactionName);
+	}
+	bHaveSelection = false;
+
+	UpdateSelectionLineSet();
+}
+
+
+void UCubeGridTool::Setup()
+{
+	using namespace CubeGridToolLocals;
+
+	UInteractiveTool::Setup();
+
+	GetToolManager()->DisplayMessage(PushPullModeMessage, EToolMessageLevel::UserNotification);
+
+	DuringActivityActions = NewObject<UCubeGridDuringActivityActions>(this);
+	DuringActivityActions->Initialize(this);
+	AddToolPropertySource(DuringActivityActions);
+	SetToolPropertySourceEnabled(DuringActivityActions, false);
+
+	ToolActions = NewObject<UCubeGridToolActions>(this);
+	ToolActions->Initialize(this);
+	AddToolPropertySource(ToolActions);
+
+	Settings = NewObject<UCubeGridToolProperties>(this);
+	Settings->RestoreProperties(this);
+	AddToolPropertySource(Settings);
+
+	OutputTypeProperties = NewObject<UCreateMeshObjectTypeProperties>(this);
+	OutputTypeProperties->RestoreProperties(this);
+	OutputTypeProperties->InitializeDefault();
+	OutputTypeProperties->WatchProperty(OutputTypeProperties->OutputType, [this](FString) { OutputTypeProperties->UpdatePropertyVisibility(); });
+	AddToolPropertySource(OutputTypeProperties);
+
+	MaterialProperties = NewObject<UNewMeshMaterialProperties>(this);
+	MaterialProperties->RestoreProperties(this);
+	MaterialProperties->bShowExtendedOptions = false;
+	AddToolPropertySource(MaterialProperties);
+
+	CurrentMesh = MakeShared<FDynamicMesh3, ESPMode::ThreadSafe>();
+	CurrentMesh->EnableAttributes();
+	if (Target)
+	{
+		*CurrentMesh = UE::ToolTarget::GetDynamicMeshCopy(Target);
+		UE::ToolTarget::SetSourceObjectVisible(Target, false);
+	}
+
+	MeshSpatial = MakeShared<FDynamicMeshAABBTree3, ESPMode::ThreadSafe>();
+	MeshSpatial->SetMesh(CurrentMesh.Get(), true);
+
+	Preview = NewObject<UMeshOpPreviewWithBackgroundCompute>(this);
+	Preview->Setup(TargetWorld, this);
+	Preview->PreviewMesh->SetTangentsMode(EDynamicMeshComponentTangentsMode::AutoCalculated);
+	if (Target)
+	{
+		Preview->PreviewMesh->UpdatePreview(CurrentMesh.Get());
+		CurrentMeshTransform = UE::ToolTarget::GetLocalToWorldTransform(Target);
+
+		// Set materials
+		FComponentMaterialSet MaterialSet = UE::ToolTarget::GetMaterialSet(Target);
+		Preview->ConfigureMaterials(MaterialSet.Materials, ToolSetupUtil::GetDefaultWorkingMaterial(GetToolManager()));
+	}
+	Preview->PreviewMesh->SetTransform((FTransform)CurrentMeshTransform);
+	Preview->OnOpCompleted.AddWeakLambda(this, [this](const FDynamicMeshOperator* Op) {
+		LastOpChangedTids = static_cast<const FCubeGridBooleanOp*>(Op)->ChangedTids;
+		});
+
+	CubeGrid = MakeShared<FCubeGrid>();
+	CubeGrid->SetGridFrame(FFrame3d(Settings->GridFrameOrigin, Settings->GridFrameOrientation.Quaternion()));
+	CubeGrid->SetBaseGridCellSize(Settings->BlockBaseSize);
+	CubeGrid->SetGridPower(Settings->PowerOfTwo);
+
+	FTransform GridTransform = CubeGrid->GetFrame().ToFTransform();
+
+	GridGizmoTransformProxy = NewObject<UTransformProxy>(this);
+	GridGizmoTransformProxy->SetTransform(GridTransform);
+	GridGizmoTransformProxy->OnBeginTransformEdit.AddWeakLambda(this, [this](UTransformProxy* Proxy) 
+		{
+			bInGizmoDrag = true;
+		});
+	GridGizmoTransformProxy->OnTransformChanged.AddUObject(this, &UCubeGridTool::GridGizmoMoved);
+	GridGizmoTransformProxy->OnEndTransformEdit.AddWeakLambda(this,
+		[this](UTransformProxy* Proxy) {
+			NotifyOfPropertyChangeByTool(Settings);
+			UpdateCornerGeometrySet();
+			bInGizmoDrag = false;
+		});
+
+	GridGizmo = UE::TransformGizmoUtil::CreateCustomTransformGizmo(GetToolManager(),
+		ETransformGizmoSubElements::StandardTranslateRotate, this);
+	GridGizmo->bUseContextCoordinateSystem = false;
+	GridGizmo->SetActiveTarget(GridGizmoTransformProxy, GetToolManager());
+	GridGizmoAlignmentMechanic = NewObject<UDragAlignmentMechanic>(this);
+	GridGizmoAlignmentMechanic->Setup(this);
+	GridGizmoAlignmentMechanic->InitializeDeformedMeshRayCast([this]() { return MeshSpatial.Get(); }, CurrentMeshTransform, nullptr);
+	GridGizmoAlignmentMechanic->AddToGizmo(GridGizmo);
+
+	GridGizmo->SetVisibility(false);
+
+	LineSets = NewObject<UPreviewGeometry>();
+	LineSets->CreateInWorld(TargetWorld, GridTransform);
+
+	LineSets->AddLineSet(HoverLineSetID);
+	LineSets->AddLineSet(SelectionLineSetID);
+	LineSets->AddLineSet(CornerModeLineSetID);
+	LineSets->SetAllLineSetsMaterial(ToolSetupUtil::GetDefaultLineComponentMaterial(GetToolManager(), /*bDepthTested*/ false));
+
+	LineSets->AddLineSet(GridLineSetID);
+	LineSets->SetLineSetMaterial(GridLineSetID, 
+		ToolSetupUtil::GetDefaultLineComponentMaterial(GetToolManager(), /*bDepthTested*/ true));
+
+	UpdateGridLineSet();
+
+	SelectedCornerRenderer.LineThickness = CornerLineThickness;
+
+	ClickDragBehavior = NewObject<UClickDragInputBehavior>();
+	ClickDragBehavior->Initialize(this);
+	AddInputBehavior(ClickDragBehavior, this);
+
+	HoverBehavior = NewObject<UMouseHoverBehavior>();
+	HoverBehavior->Modifiers.RegisterModifier(ShiftModifierID, FInputDeviceState::IsShiftKeyDown);
+	HoverBehavior->Modifiers.RegisterModifier(CtrlModifierID, FInputDeviceState::IsCtrlKeyDown);
+	HoverBehavior->Initialize(this);
+	AddInputBehavior(HoverBehavior, this);
+
+	CtrlMiddleClickBehavior = NewObject<ULocalSingleClickInputBehavior>();
+	CtrlMiddleClickBehavior->Initialize();
+	CtrlMiddleClickBehavior->SetUseMiddleMouseButton();
+	CtrlMiddleClickBehavior->ModifierCheckFunc = [this](const FInputDeviceState& InputState) {
+		return FInputDeviceState::IsCtrlKeyDown(InputState);
+	};
+	CtrlMiddleClickBehavior->IsHitByClickFunc = [this](const FInputDeviceRay& InputRay) {
+		FCubeGrid::FCubeFace Face;
+		FInputRayHit OutResult;
+		OutResult.bHit = GetHitGridFace(InputRay.WorldRay, Face);
+		return OutResult;
+	};
+	CtrlMiddleClickBehavior->OnClickedFunc = [this](const FInputDeviceRay& ClickPos) {
+		OnCtrlMiddleClick(ClickPos);
+	};
+	AddInputBehavior(CtrlMiddleClickBehavior);
+
+	MiddleClickDragBehavior = NewObject<ULocalClickDragInputBehavior>();
+	MiddleClickDragBehavior->Initialize();
+	MiddleClickDragBehavior->SetUseMiddleMouseButton();
+	MiddleClickDragBehavior->ModifierCheckFunc = [this](const FInputDeviceState& InputState) {
+		return !FInputDeviceState::IsCtrlKeyDown(InputState);
+	};
+	MiddleClickDragBehavior->CanBeginClickDragFunc = [this](const FInputDeviceRay& ClickPos) {
+		return CanBeginMiddleClickDrag(ClickPos);
+	};
+	MiddleClickDragBehavior->OnClickPressFunc = [this](const FInputDeviceRay& ClickPos) {
+		PrepForSelectionChange();
+		RayCastSelectionPlane(ClickPos.WorldRay, MiddleClickDragStart);
+	};
+	MiddleClickDragBehavior->OnClickDragFunc = [this](const FInputDeviceRay& ClickPos) {
+		OnMiddleClickDrag(ClickPos);
+	};
+	MiddleClickDragBehavior->OnClickReleaseFunc = [this](const FInputDeviceRay&) {
+		EndSelectionChange();
+	};
+	MiddleClickDragBehavior->OnTerminateFunc = [this]() {
+		EndSelectionChange();
+	};
+	AddInputBehavior(MiddleClickDragBehavior);
+
+	PowerOfTwoPrevious = Settings->PowerOfTwo;
+	Settings->WatchProperty(Settings->PowerOfTwo,
+		[this](uint8 NewPower) {
+			SetPowerOfTwoClamped(Settings->PowerOfTwo);
+		});
+	Settings->WatchProperty(Settings->BlockBaseSize,
+		[this](double NewBaseSize) {
+			ClearHover();
+
+			CubeGrid->SetBaseGridCellSize(NewBaseSize);
+
+			UpdateSelectionLineSet();
+			UpdateGridLineSet();
+		});
+	Settings->WatchProperty(Settings->bShowGizmo,
+		[this](bool bOn) {
+			UpdateGizmoVisibility(bOn);
+		});
+	Settings->WatchProperty(Settings->bCrosswiseDiagonal,
+		[this](bool bOn) {
+			InvalidatePreview();
+		});
+	Settings->WatchProperty(Settings->PlaneTolerance,
+		[this](double Tolerance) {
+			InvalidatePreview(false);
+		});
+
+	auto UpdateFromDetailsPanelTransformChange = [this]()
+	{
+		CubeGrid->SetGridFrame(FFrame3d(Settings->GridFrameOrigin, Settings->GridFrameOrientation.Quaternion()));
+
+		FTransform GridTransform = CubeGrid->GetFrame().ToFTransform();
+		LineSets->SetTransform(GridTransform);
+		GridGizmo->ReinitializeGizmoTransform(GridTransform);
+
+		InvalidatePreview(false);
+	};
+
+	GridFrameOriginWatcherIdx = Settings->WatchProperty(Settings->GridFrameOrigin,
+		[UpdateFromDetailsPanelTransformChange](FVector) {
+			UpdateFromDetailsPanelTransformChange();
+		});
+	GridFrameOrientationWatcherIdx = Settings->WatchProperty(Settings->GridFrameOrientation,
+		[UpdateFromDetailsPanelTransformChange](FRotator) {
+			UpdateFromDetailsPanelTransformChange();
+		});
+
+	Settings->SilentUpdateWatched();
+}
+
+void UCubeGridTool::Shutdown(EToolShutdownType ShutdownType)
+{
+	using UE::Geometry::FTransform3d;
+
+	Settings->SaveProperties(this);
+	MaterialProperties->SaveProperties(this);
+
+	if (Mode == EMode::Corner)
+	{
+		ApplyCornerMode(true);
+	}
+
+	if (Target)
+	{
+		Cast<IPrimitiveComponentBackedTarget>(Target)->SetOwnerVisibility(true);
+
+		if (bChangesMade)
+		{
+			GetToolManager()->BeginUndoTransaction(LOCTEXT("CubeGridToolEditTransactionName", "Block Tool Edit"));
+			UE::ToolTarget::CommitDynamicMeshUpdate(Target, *CurrentMesh, true);
+			GetToolManager()->EndUndoTransaction();
+		}
+	}
+	else if (CurrentMesh->TriangleCount() > 0)
+	{
+		GetToolManager()->BeginUndoTransaction(LOCTEXT("CubeGridToolCreateTransactionName", "Block Tool Create New"));
+
+		FCreateMeshObjectParams NewMeshObjectParams;
+		NewMeshObjectParams.TargetWorld = TargetWorld;
+		NewMeshObjectParams.Transform = (FTransform)CurrentMeshTransform;
+		NewMeshObjectParams.BaseName = TEXT("CubeGridToolOutput");
+		NewMeshObjectParams.Materials.Add(MaterialProperties->Material.Get());
+		NewMeshObjectParams.SetMesh(CurrentMesh.Get());
+		OutputTypeProperties->ConfigureCreateMeshObjectParams(NewMeshObjectParams);
+		FCreateMeshObjectResult Result = UE::Modeling::CreateMeshObject(GetToolManager(), MoveTemp(NewMeshObjectParams));
+		if (Result.IsOK() && Result.NewActor != nullptr)
+		{
+			ToolSelectionUtil::SetNewActorSelection(GetToolManager(), Result.NewActor);
+		}
+
+		GetToolManager()->EndUndoTransaction();
+	}
+
+	Preview->OnOpCompleted.RemoveAll(this);
+	Preview->Shutdown();
+
+	if (LineSets)
+	{
+		LineSets->Disconnect();
+		LineSets = nullptr;
+	}
+
+	GridGizmoAlignmentMechanic->Shutdown();
+
+	GetToolManager()->GetPairedGizmoManager()->DestroyAllGizmosByOwner(this);
+}
+
+void UCubeGridTool::OnTick(float DeltaTime)
+{
+	if (Preview)
+	{
+		Preview->Tick(DeltaTime);
+	}
+
+	if (PendingAction != ECubeGridToolAction::NoAction)
+	{
+		ApplyAction(PendingAction);
+		PendingAction = ECubeGridToolAction::NoAction;
+	}
+
+	if (bWaitingToApplyPreview && Preview->HaveValidResult())
+	{
+		ApplyPreview();
+	}
+}
+
+void UCubeGridTool::ApplyPreview()
+{
+	using namespace CubeGridToolLocals;
+
+	FDynamicMeshChangeTracker ChangeTracker(CurrentMesh.Get());
+	ChangeTracker.BeginChange();
+	ChangeTracker.SaveTriangles(*LastOpChangedTids, true /*bSaveVertices*/);
+
+	// Update current mesh
+	bChangesMade = true; // TODO: make this undoable
+
+	CurrentMesh->Copy(*Preview->PreviewMesh->GetMesh());
+	MeshSpatial->Build();
+
+	const FText TransactionText = LOCTEXT("CubeGridToolTransactionName", "Block Tool Change");
+	GetToolManager()->BeginUndoTransaction(TransactionText);
+
+	bWaitingToApplyPreview = false;
+	bPreviewMayDiffer = false;
+	bBlockUntilPreviewUpdate = false;
+
+	if (bAdjustSelectionOnPreviewUpdate)
+	{
+		// Change the selection to the new location. Note that this should to happen
+		// after resetting bPreviewMayDiffer to avoid an extra preview reset when selection
+		// changes.
+		SlideSelection(CurrentExtrudeAmount, true);
+	}
+	CurrentExtrudeAmount = 0;
+
+	GetToolManager()->EmitObjectChange(this, MakeUnique<FCubeGridToolMeshChange>(ChangeTracker.EndChange()), TransactionText);
+	GetToolManager()->EndUndoTransaction();
+}
+
+void UCubeGridTool::Render(IToolsContextRenderAPI* RenderAPI)
+{
+	using namespace CubeGridToolLocals;
+
+	GetToolManager()->GetContextQueriesAPI()->GetCurrentViewState(CameraState);
+
+	if (Mode == EMode::Corner)
+	{
+		FOrientedBox3d OrientedBox = ConvertToOrientedBox(Selection.Box, Selection.Direction);
+
+		SelectedCornerRenderer.BeginFrame(RenderAPI, CameraState);
+		for (int i = 0; i < 4; ++i)
+		{
+			FVector WorldPosition = CubeGrid->GetFrame().FromFramePoint(OrientedBox.GetCorner(i));
+
+			// Depending on whether we're in an orthographic view or not, we set the radius based on visual angle or based on ortho 
+			// viewport width (divided into 90 segments like the FOV is divided into 90 degrees).
+			float Radius = (CameraState.bIsOrthographic) ? CameraState.OrthoWorldCoordinateWidth * 0.5 / 90.0
+				: (float)ToolSceneQueriesUtil::CalculateDimensionFromVisualAngleD(CameraState, (FVector3d)WorldPosition, 0.5);
+			bool bDepthTested = false;
+			SelectedCornerRenderer.DrawViewFacingCircle(
+				WorldPosition,
+				Radius,
+				CornerCircleNumSteps,
+				CornerSelectedFlags[i] ? SelectedCornerLineColor : UnselectedCornerLineColor,
+				CornerLineThickness, bDepthTested);
+		}
+		SelectedCornerRenderer.EndFrame();
+	}
+
+	GridGizmoAlignmentMechanic->Render(RenderAPI);
+}
+
+void UCubeGridTool::OnPropertyModified(UObject* PropertySet, FProperty* Property)
+{
+}
+
+void UCubeGridTool::ClearHover() 
+{
+	if (bHaveHoveredSelection)
+	{
+		UpdateHoverLineSet(false, HoveredSelectionBox);
+	}
+}
+
+void  UCubeGridTool::GridGizmoMoved(UTransformProxy* Proxy, FTransform Transform)
+{
+	Settings->GridFrameOrigin = Transform.GetTranslation();
+	Settings->SilentUpdateWatcherAtIndex(GridFrameOriginWatcherIdx);
+
+	Settings->GridFrameOrientation = Transform.GetRotation().Rotator();
+	Settings->SilentUpdateWatcherAtIndex(GridFrameOrientationWatcherIdx);
+
+	CubeGrid->SetGridFrame(FFrame3d(Settings->GridFrameOrigin, Settings->GridFrameOrientation.Quaternion()));
+	LineSets->SetTransform(CubeGrid->GetFrame().ToFTransform());
+
+	if (Mode == EMode::Corner)
+	{
+		InvalidatePreview(false);
+
+		if (!bInGizmoDrag)
+		{
+			UpdateCornerGeometrySet();
+		}
+	}
+}
+
+void UCubeGridTool::OnCtrlMiddleClick(const FInputDeviceRay& ClickPos)
+{
+	// Get the selected face
+	FCubeGrid::FCubeFace Face;
+	if (!ensure(GetHitGridFace(ClickPos.WorldRay, Face)))
+	{
+		return;
+	}
+
+	if (!GridGizmo->IsVisible())
+	{
+		UpdateGizmoVisibility(true);
+	}
+
+	// Get the face's four corners in grid space
+	FVector3d FaceMin = Face.GetMinCorner();
+	FVector3d FaceMax = Face.GetMaxCorner();
+
+	FVector3d Corners[4] = {
+		FaceMin,
+		FaceMin,
+		FaceMax,
+		FaceMax
+	};
+
+	int32 DifferingDimension = (FaceMin[0] == FaceMax[0]) ? 1 : 0;
+
+	Corners[1][DifferingDimension] = FaceMax[DifferingDimension];
+	Corners[3][DifferingDimension] = FaceMin[DifferingDimension];
+
+	// Transform the ray to grid space and see which of the corners is closest.
+	FVector3d GridSpaceRayOrigin = CubeGrid->ToGridPoint(ClickPos.WorldRay.Origin);
+	FVector3d GridSpaceRayDirection = CubeGrid->GetFrame().ToFrameVector(ClickPos.WorldRay.Direction);
+	GridSpaceRayDirection.Normalize();
+	FRay3d GizmoSpaceRay(GridSpaceRayOrigin, GridSpaceRayDirection);
+
+	double MinDistSquared = GizmoSpaceRay.DistanceSquared(Corners[0]);
+	int32 ClosestCornerIndex = 0;
+	for (int32 i = 1; i < 4; ++i)
+	{
+		double DistSquared = GizmoSpaceRay.DistanceSquared(Corners[i]);
+		if (DistSquared < MinDistSquared)
+		{
+			MinDistSquared = DistSquared;
+			ClosestCornerIndex = i;
+		}
+	}
+
+	GetToolManager()->BeginUndoTransaction(LOCTEXT("QuickAdjustGizmo", "Transform Gizmo"));
+	
+	// Adjust the selection if needed
+	if (bHaveSelection)
+	{
+		FVector3d GridSpaceDisplacement = Corners[ClosestCornerIndex];
+		FVector3d FrameSpaceDisplacement = GridSpaceDisplacement * CubeGrid->GetCurrentGridCellSize();
+
+		FSelection NewSelection = Selection;
+		NewSelection.StartBox = FAxisAlignedBox3d(
+			Selection.StartBox.Min - FrameSpaceDisplacement,
+			Selection.StartBox.Max - FrameSpaceDisplacement);
+		NewSelection.Box = FAxisAlignedBox3d(
+			Selection.Box.Min - FrameSpaceDisplacement,
+			Selection.Box.Max - FrameSpaceDisplacement);
+
+		SetSelection(NewSelection, true);
+	}
+
+	// Move the gizmo to that corner.
+	Settings->GridFrameOrigin = CubeGrid->ToWorldPoint(Corners[ClosestCornerIndex]);
+	Settings->SilentUpdateWatcherAtIndex(GridFrameOriginWatcherIdx);
+	CubeGrid->SetGridFrame(FFrame3d(Settings->GridFrameOrigin, Settings->GridFrameOrientation.Quaternion()));
+
+	FTransform GridTransform = CubeGrid->GetFrame().ToFTransform();
+	LineSets->SetTransform(GridTransform);
+	GridGizmo->SetNewGizmoTransform(GridTransform);
+
+	GetToolManager()->EndUndoTransaction();
+}
+
+// Tries to intersect the selected box. Used for middle mouse dragging the selection.
+FInputRayHit UCubeGridTool::RayCastSelectionPlane(const UE::Geometry::FRay3d& WorldRay,
+	FVector3d& HitPointOut)
+{
+	FVector3d Normal = CubeGrid->GetFrame().FromFrameVector(
+		FCubeGrid::DirToNormal(Selection.Direction));
+
+	FInputRayHit HitResult;
+	GizmoMath::RayPlaneIntersectionPoint(CubeGrid->GetFrame().FromFramePoint(Selection.Box.Min), Normal,
+		WorldRay.Origin, WorldRay.Direction, HitResult.bHit, HitPointOut);
+	
+	if (HitResult.bHit)
+	{
+		HitResult = FInputRayHit(WorldRay.Project(HitPointOut));
+	}
+	return HitResult;
+}
+
+FInputRayHit UCubeGridTool::CanBeginMiddleClickDrag(const FInputDeviceRay& ClickPos)
+{
+	FInputRayHit HitResult;
+	if (!bHaveSelection)
+	{
+		return HitResult;
+	}
+
+	FVector3d WorldHitPoint;
+	HitResult = RayCastSelectionPlane(ClickPos.WorldRay, WorldHitPoint);
+
+	int FlatDim = FCubeGrid::DirToFlatDim(Selection.Direction);
+
+	FVector3d FrameSpacePoint = CubeGrid->GetFrame().ToFramePoint(WorldHitPoint);
+	FrameSpacePoint[FlatDim] = Selection.Box.Min[FlatDim];
+	HitResult.bHit = Selection.Box.Contains(FrameSpacePoint);
+
+	return HitResult;
+}
+
+void UCubeGridTool::OnMiddleClickDrag(const FInputDeviceRay& DragPos)
+{
+	if (!bHaveSelection)
+	{
+		return;
+	}
+
+	FVector3d MiddleClickDragEnd;
+	RayCastSelectionPlane(DragPos.WorldRay, MiddleClickDragEnd);
+	FVector3d DisplacementInGridFrame = CubeGrid->GetFrame().ToFrameVector(MiddleClickDragEnd - MiddleClickDragStart); 
+	
+	// Clamp the relevant dimension in the displacement vector
+	DisplacementInGridFrame[FCubeGrid::DirToFlatDim(Selection.Direction)] = 0;
+
+	// Make the displacement be a multiple of the current grid cell size
+	DisplacementInGridFrame /= CubeGrid->GetCurrentGridCellSize();
+	DisplacementInGridFrame = FVector3d(
+		FMath::RoundToDouble(DisplacementInGridFrame.X), 
+		FMath::RoundToDouble(DisplacementInGridFrame.Y),
+		FMath::RoundToDouble(DisplacementInGridFrame.Z));
+	DisplacementInGridFrame *= CubeGrid->GetCurrentGridCellSize();
+
+	FAxisAlignedBox3d NewSelectionBox(
+		PreviousSelection.Box.Min + DisplacementInGridFrame,
+		PreviousSelection.Box.Max + DisplacementInGridFrame
+	);
+
+	// Adjust selection
+	if (NewSelectionBox != Selection.Box)
+	{
+		Selection.StartBox = FAxisAlignedBox3d(
+			PreviousSelection.StartBox.Min + DisplacementInGridFrame,
+			PreviousSelection.StartBox.Max + DisplacementInGridFrame
+		);
+
+		Selection.Box = FAxisAlignedBox3d(
+			PreviousSelection.Box.Min + DisplacementInGridFrame,
+			PreviousSelection.Box.Max + DisplacementInGridFrame
+		);
+
+		UpdateSelectionLineSet();
+		InvalidatePreview();
+	}
+}
+
+void UCubeGridTool::PrepForSelectionChange()
+{
+	bPreviousHaveSelection = bHaveSelection;
+	PreviousSelection = Selection;
+}
+
+void UCubeGridTool::EndSelectionChange()
+{
+	using namespace CubeGridToolLocals;
+
+	if (bPreviousHaveSelection != bHaveSelection
+		|| (bHaveSelection && PreviousSelection != Selection))
+	{
+		GetToolManager()->EmitObjectChange(this,
+			MakeUnique<FCubeGridToolSelectionChange>(bPreviousHaveSelection, bHaveSelection,
+				PreviousSelection, Selection),
+			SelectionChangeTransactionName);
+	}
+}
+
+void UCubeGridTool::UpdateGizmoVisibility(bool bVisible)
+{
+	GridGizmo->SetVisibility(bVisible);
+	LineSets->SetLineSetMaterial(CubeGridToolLocals::GridLineSetID,
+		ToolSetupUtil::GetDefaultLineComponentMaterial(GetToolManager(), /*bDepthTested*/ !bVisible));
+	Settings->bShowGizmo = bVisible;
+	Settings->SilentUpdateWatched();
+}
+
+bool UCubeGridTool::GetHitGridFace(const FRay& WorldRay, FCubeGrid::FCubeFace& FaceOut)
+{
+	using namespace CubeGridToolLocals;
+
+	double BestHitT = TNumericLimits<double>::Max();
+
+	// We always hit-test the ground plane...
+	bool bHitPlane;
+	FVector IntersectionPoint;
+	GizmoMath::RayPlaneIntersectionPoint(CubeGrid->GetFrame().Origin, CubeGrid->GetFrame().Z(), 
+		WorldRay.Origin, WorldRay.Direction, bHitPlane, IntersectionPoint);
+	if (bHitPlane)
+	{
+		FVector3d ClampedGridPoint = CubeGrid->ToGridPoint(IntersectionPoint);
+		ClampedGridPoint.Z = 0;
+		FaceOut = FCubeGrid::FCubeFace(ClampedGridPoint,
+			WorldRay.Origin.Z >= 0 ? FCubeGrid::EFaceDirection::PositiveZ : FCubeGrid::EFaceDirection::NegativeZ,
+			CubeGrid->GetGridPower());
+		BestHitT = WorldRay.GetParameter(IntersectionPoint);
+	}
+
+	// ...However depending on the settings, we may give everything else priority, which we do by
+	// keeping the plane hit distance maximal.
+	if (!Settings->bHitGridGroundPlaneIfCloser)
+	{
+		BestHitT = TNumericLimits<double>::Max();
+	}
+
+	if (Settings->bHitUnrelatedGeometry)
+	{
+		FHitResult HitResult;
+		if (ToolSceneQueriesUtil::FindNearestVisibleObjectHit(TargetWorld, HitResult, WorldRay)
+			&& HitResult.Distance < BestHitT)
+		{
+			BestHitT = HitResult.Distance;
+			ConvertToFaceHit(*CubeGrid, Settings->FaceSelectionMode, 
+				WorldRay, BestHitT, HitResult.ImpactNormal, FaceOut, Settings->PlaneTolerance);
+		}
+	}
+
+	if (MeshSpatial)
+	{
+		FRay3d LocalRay(CurrentMeshTransform.InverseTransformPosition((FVector3d)WorldRay.Origin),
+			CurrentMeshTransform.InverseTransformVectorNoScale((FVector3d)WorldRay.Direction));
+
+		int32 Tid;
+		double LocalHitT = TNumericLimits<double>::Max();
+		if (MeshSpatial->FindNearestHitTriangle(LocalRay, LocalHitT, Tid))
+		{
+			double HitT = WorldRay.GetParameter(CurrentMeshTransform.TransformPosition(LocalRay.PointAt(LocalHitT)));
+			if (HitT < BestHitT)
+			{
+				BestHitT = HitT;
+				ConvertToFaceHit(*CubeGrid, Settings->FaceSelectionMode,
+					WorldRay, BestHitT, CurrentMesh->GetTriNormal(Tid), FaceOut, Settings->PlaneTolerance);
+			}
+		}
+	}
+
+	// We can't go just off of BestHitT because we keep it maximal for plane hits when 
+	// bHitGridGroundPlaneIfCloser is false.
+	return bHitPlane || BestHitT != TNumericLimits<double>::Max();
+}
+
+FInputRayHit UCubeGridTool::CanBeginClickDragSequence(const FInputDeviceRay& PressPos)
+{
+	FInputRayHit HitResult;
+	HitResult.bHit = (Mode != EMode::FitGrid && !(bBlockUntilPreviewUpdate && bWaitingToApplyPreview));
+	return HitResult;
+}
+
+
+void UCubeGridTool::OnClickPress(const FInputDeviceRay& PressPos)
+{
+	using namespace CubeGridToolLocals;
+
+	UpdateHoverLineSet(false, HoveredSelectionBox); // clear hover
+	
+	if (Mode == EMode::Corner)
+	{
+		MouseState = EMouseState::DraggingCornerSelection;
+		for (int i = 0; i < 4; ++i)
+		{
+			PreDragCornerSelectedFlags[i] = CornerSelectedFlags[i];
+		}
+		AttemptToSelectCorner(PressPos.WorldRay);
+		return;
+	}
+
+	if (bMouseDragShouldPushPull)
+	{
+		MouseState = EMouseState::DraggingExtrudeDistance;
+		if (bHaveSelection)
+		{
+			DragProjectionAxis = FRay3d(CubeGrid->GetFrame().FromFramePoint(Selection.Box.Center()),
+				CubeGrid->GetFrame().FromFrameVector(FCubeGrid::DirToNormal(Selection.Direction)), true);
+
+			FDistLine3Ray3d DistanceCalculator(UE::Geometry::FLine3d(DragProjectionAxis.Origin, DragProjectionAxis.Direction), PressPos.WorldRay);
+			DistanceCalculator.ComputeResult();
+			DragProjectedStartParam = DistanceCalculator.LineParameter;
+		}
+		return;
+	}
+
+	// Otherwise, deal with selection
+	PrepForSelectionChange();
+
+	MouseState = EMouseState::DraggingRegularSelection;
+	FCubeGrid::FCubeFace HitFace;
+	if (bHaveSelection && bSelectionToggle)
+	{
+		// We're adding to existing selection
+		if (GetCoplanarFrameSpaceSelectedBox(*CubeGrid, PressPos.WorldRay, Selection.StartBox, true, Selection.Box))
+		{
+			UpdateSelectionLineSet();
+		}
+	} 
+	else if (GetHitGridFace(PressPos.WorldRay, HitFace))
+	{
+		// Reset start of the selection
+		bHaveSelection = true;
+		double GridScale = CubeGrid->GetCellSize(HitFace.GetSourceCubeGridPower());
+		Selection.Box = FAxisAlignedBox3d(HitFace.GetMinCorner() * GridScale, HitFace.GetMaxCorner() * GridScale);
+		Selection.StartBox = Selection.Box;
+		Selection.Direction = HitFace.GetDirection();
+		UpdateSelectionLineSet();
+	}
+	else
+	{
+		// Clear selection (the event emit, if needed, happens on click release) 
+		bHaveSelection = false;
+		Selection.Box = FAxisAlignedBox3d();
+	}
+}
+
+void UCubeGridTool::OnClickDrag(const FInputDeviceRay& DragPos)
+{
+	using namespace CubeGridToolLocals;
+
+	if (!bHaveSelection)
+	{
+		return;
+	}
+
+	if (MouseState == EMouseState::DraggingExtrudeDistance)
+	{
+		if (!bHaveSelection || (Mode == EMode::Corner && !IsAnyCornerSelected(CornerSelectedFlags)))
+		{
+			return;
+		}
+
+		FDistLine3Ray3d DistanceCalculator(
+			UE::Geometry::FLine3d(DragProjectionAxis.Origin, DragProjectionAxis.Direction), DragPos.WorldRay);
+		DistanceCalculator.ComputeResult();
+
+		double ParamDelta = DistanceCalculator.LineParameter - DragProjectedStartParam;
+		double CubeSize = CubeGrid->GetCurrentGridCellSize();
+		int32 NewExtrudeAmount = FMath::RoundToInt(ParamDelta / (CubeSize * Settings->BlocksPerStep)) * Settings->BlocksPerStep;
+		if (NewExtrudeAmount != CurrentExtrudeAmount)
+		{
+			CurrentExtrudeAmount = NewExtrudeAmount;
+			InvalidatePreview();
+		}
+		return;
+	}
+	else if (MouseState == EMouseState::DraggingCornerSelection)
+	{
+		AttemptToSelectCorner(DragPos.WorldRay);
+		return;
+	}
+	else // Grid selection
+	{
+		FCubeGrid::FCubeFace HitFace;
+		double HitT = TNumericLimits<double>::Max();
+		bool bHit = GetCoplanarFrameSpaceSelectedBox(*CubeGrid, DragPos.WorldRay, Selection.StartBox,
+			true, Selection.Box);
+
+		UpdateSelectionLineSet();
+	}
+}
+
+void UCubeGridTool::OnClickRelease(const FInputDeviceRay& ReleasePos)
+{
+	if (MouseState == EMouseState::DraggingExtrudeDistance && CurrentExtrudeAmount != 0)
+	{
+		bWaitingToApplyPreview = true;
+		bBlockUntilPreviewUpdate = false;
+		bAdjustSelectionOnPreviewUpdate = true;
+	}
+	else if (MouseState == EMouseState::DraggingRegularSelection)
+	{
+		EndSelectionChange();
+	}
+
+	MouseState = EMouseState::NotDragging;
+}
+
+void UCubeGridTool::OnTerminateDragSequence()
+{
+	if (MouseState == EMouseState::DraggingExtrudeDistance && CurrentExtrudeAmount != 0)
+	{
+		bWaitingToApplyPreview = true;
+		bBlockUntilPreviewUpdate = false;
+		bAdjustSelectionOnPreviewUpdate = true;
+	}
+
+	MouseState = EMouseState::NotDragging;
+}
+
+void UCubeGridTool::AttemptToSelectCorner(const FRay3d& WorldRay)
+{
+	TArray<FGeometrySet3::FNearest> HitCorners;
+	CornersGeometrySet.CollectPointsNearRay(WorldRay, HitCorners, [this](const FVector3d& Position1, const FVector3d& Position2) {
+		double ToleranceScale = 3;
+		if (CameraState.bIsOrthographic)
+		{
+			// We could just always use ToolSceneQueriesUtil::PointSnapQuery. But in ortho viewports, we happen to know
+			// that the only points that we will ever give this function will be the closest points between a ray and
+			// some geometry, meaning that the vector between them will be orthogonal to the view ray. With this knowledge,
+			// we can do the tolerance computation more efficiently than PointSnapQuery can, since we don't need to project
+			// down to the view plane.
+			// As in PointSnapQuery, we convert our angle-based tolerance to one we can use in an ortho viewport (instead of
+			// dividing our field of view into 90 visual angle degrees, we divide the plane into 90 units).
+			float OrthoTolerance = ToolSceneQueriesUtil::GetDefaultVisualAngleSnapThreshD() * CameraState.OrthoWorldCoordinateWidth / 90.0;
+			OrthoTolerance *= ToleranceScale;
+			return FVector3d::DistSquared(Position1, Position2) < OrthoTolerance * OrthoTolerance;
+		}
+		else
+		{
+			return ToolSceneQueriesUtil::PointSnapQuery(CameraState,
+				Position1, Position2,
+				ToolSceneQueriesUtil::GetDefaultVisualAngleSnapThreshD() * ToleranceScale);
+		}
+		});
+
+	for (FGeometrySet3::FNearest Hit : HitCorners)
+	{
+		CornerSelectedFlags[Hit.ID] = !PreDragCornerSelectedFlags[Hit.ID];
+	}
+
+	if (HitCorners.Num() > 0)
+	{
+		InvalidatePreview();
+	}
+}
+
+void UCubeGridTool::UpdateSelectionLineSet()
+{
+	using namespace CubeGridToolLocals;
+
+	ULineSetComponent* LineSet = LineSets->FindLineSet(SelectionLineSetID);
+	LineSet->Clear();
+	if (bHaveSelection)
+	{
+		FVector3d CornerVector = FCubeGrid::DirToNormal(Selection.Direction)
+			* GetFrameSpaceExtrudeDist(*CubeGrid, Selection.Box.Min, -Settings->BlocksPerStep, Selection.Direction);
+		DrawGridRectangle(*LineSet, *CubeGrid, Selection.StartBox.Min, Selection.StartBox.Max, SelectionLineColor,
+			SelectionMainLineThickness, SelectionLineDepthBias);
+		DrawGridRectangle(*LineSet, *CubeGrid, Selection.Box.Min, Selection.Box.Max, SelectionLineColor,
+			SelectionMainLineThickness, SelectionLineDepthBias);
+		DrawGridSection(*LineSet, *CubeGrid, Selection.Box, SelectionLineColor, 
+			SelectionGridLineThickness, SelectionLineDepthBias, &CornerVector);
+
+		if (Mode == EMode::Corner)
+		{
+			// This isn't quite relevant to updating the selection line set, but it's a convenient place
+			// to put this because if the selection set is changing, the geometry set probably needs
+			// to be doing the same.
+			UpdateCornerGeometrySet();
+		}
+	}
+}
+
+FInputRayHit UCubeGridTool::BeginHoverSequenceHitTest(const FInputDeviceRay& PressPos)
+{
+	FInputRayHit HitResult;
+	HitResult.bHit = UpdateHover(PressPos.WorldRay);
+	return HitResult;
+}
+
+bool UCubeGridTool::UpdateHover(const FRay& WorldRay)
+{
+	using namespace CubeGridToolLocals;
+
+	if (Mode != EMode::PushPull)
+	{
+		UpdateHoverLineSet(false, HoveredSelectionBox);
+		return false;
+	}
+
+	bool bHit;
+	UE::Geometry::FAxisAlignedBox3d Box;
+	if (bHaveSelection && bSelectionToggle)
+	{
+		bHit = GetCoplanarFrameSpaceSelectedBox(*CubeGrid, WorldRay, Selection.StartBox, false, Box);
+	}
+	else
+	{
+		FCubeGrid::FCubeFace HitFace;
+		bHit = GetHitGridFace(WorldRay, HitFace);
+		if (bHit)
+		{
+			double HoverScale = CubeGrid->GetCurrentGridCellSize();
+			Box = FAxisAlignedBox3d(HitFace.GetMinCorner() * HoverScale,
+				HitFace.GetMaxCorner() * HoverScale);
+		}
+	}
+
+	UpdateHoverLineSet(bHit, Box);
+
+	return bHit;
+}
+
+void UCubeGridTool::OnBeginHover(const FInputDeviceRay& DevicePos)
+{
+}
+
+bool UCubeGridTool::OnUpdateHover(const FInputDeviceRay& DevicePos)
+{
+	return UpdateHover(DevicePos.WorldRay);
+}
+
+// We could have not taken arguments here and done it the way we do selection, but we'd need
+// to keep track of previous hover to avoid unnecessary updates
+void UCubeGridTool::UpdateHoverLineSet(bool bNewHaveHover, const UE::Geometry::FAxisAlignedBox3d& NewHoveredBox)
+{
+	using namespace CubeGridToolLocals;
+
+	ULineSetComponent* LineSet = LineSets->FindLineSet(HoverLineSetID);
+
+	if (!bNewHaveHover)
+	{
+		if (bHaveHoveredSelection)
+		{
+			LineSet->Clear();
+		}
+	}
+	else if (!bHaveHoveredSelection || NewHoveredBox != HoveredSelectionBox)
+	{
+		HoveredSelectionBox = NewHoveredBox;
+		bHaveHoveredSelection = true;
+
+		LineSet->Clear();
+		DrawGridRectangle(*LineSet, *CubeGrid, HoveredSelectionBox.Min, HoveredSelectionBox.Max,
+			HoverLineColor, HoverLineThickness, HoverLineDepthBias);
+	}
+}
+
+void UCubeGridTool::OnEndHover()
+{
+}
+
+void UCubeGridTool::OnUpdateModifierState(int ModifierID, bool bIsOn)
+{
+	if (ModifierID == ShiftModifierID)
+	{
+		bSelectionToggle = bIsOn;
+	}
+	else if (ModifierID == CtrlModifierID)
+	{
+		bMouseDragShouldPushPull = bIsOn;
+	}
+}
+
+void UCubeGridTool::UpdateGridLineSet()
+{
+	using namespace CubeGridToolLocals;
+
+	double CurrentGridScale = CubeGrid->GetCurrentGridCellSize();
+
+	FAxisAlignedBox3d GridBox;
+	GridBox.Contain(FVector3d(-50, -50, 0) * CurrentGridScale);
+	GridBox.Contain(FVector3d(50, 50, 0) * CurrentGridScale);
+
+	ULineSetComponent* LineSet = LineSets->FindLineSet(GridLineSetID);
+	LineSet->Clear();
+	DrawGridSection(*LineSet, *CubeGrid, GridBox,
+		GridLineColor, GridLineThickness, GridLineDepthBias);
+}
+
+void UCubeGridTool::UpdateCornerModeLineSet()
+{
+	using namespace CubeGridToolLocals;
+	ULineSetComponent* LineSet = LineSets->FindLineSet(CornerModeLineSetID);
+
+	LineSet->Clear();
+	if (Mode == EMode::Corner && CurrentExtrudeAmount < 0)
+	{
+		FOrientedBox3d FrameSpaceBox = ConvertToOrientedBox(Selection.Box, Selection.Direction);
+
+		bool CornerWelded[4];
+		for (int i = 0; i < 4; ++i) {
+			CornerWelded[i] = !CornerSelectedFlags[i];
+		}
+
+		int DiagStartIdx = (CornerWelded[0] && CornerWelded[2]) ? 1 : 0;
+		DiagStartIdx = Settings->bCrosswiseDiagonal ? 1 - DiagStartIdx : DiagStartIdx;
+		
+		bool bDiagonalWelded = CornerWelded[DiagStartIdx] && CornerWelded[DiagStartIdx + 2];
+		int DeletedVert = -1;
+		if (bDiagonalWelded)
+		{
+			if (CornerWelded[DiagStartIdx + 1])
+			{
+				DeletedVert = DiagStartIdx + 1;
+			}
+			else if (CornerWelded[(DiagStartIdx + 3) % 4])
+			{
+				DeletedVert = (DiagStartIdx + 3) % 4;
+			}
+		}
+
+		FVector3d CornerExtrudeVector = CubeGrid->GetCurrentGridCellSize() * CurrentExtrudeAmount 
+			* FCubeGrid::DirToNormal(Selection.Direction);
+
+		for (int i = 0; i < 4; ++i)
+		{
+			FVector3d CurrentCorner = FrameSpaceBox.GetCorner(i);
+			if (!CornerWelded[i])
+			{
+				FVector3d UpCorner = CurrentCorner + CornerExtrudeVector;
+				LineSet->AddLine(CurrentCorner, UpCorner, CornerModeWireframeColor, 
+					CornerModeWireframeThickness, CornerModeWireframeDepthBias);
+				CurrentCorner = UpCorner;
+			}
+
+			int NextIdx = (i + 1) % 4;
+			if (i == DeletedVert || NextIdx == DeletedVert)
+			{
+				continue;
+			}
+
+			FVector3d NextCorner = CornerWelded[NextIdx] ? FrameSpaceBox.GetCorner(NextIdx)
+				: FrameSpaceBox.GetCorner(NextIdx) + CornerExtrudeVector;
+			LineSet->AddLine(CurrentCorner, NextCorner, CornerModeWireframeColor,
+				CornerModeWireframeThickness, CornerModeWireframeDepthBias);
+		}
+
+		FVector3d DiagCorner1 = CornerWelded[DiagStartIdx] ? FrameSpaceBox.GetCorner(DiagStartIdx)
+			: FrameSpaceBox.GetCorner(DiagStartIdx) + CornerExtrudeVector;
+		FVector3d DiagCorner2 = CornerWelded[DiagStartIdx + 2] ? FrameSpaceBox.GetCorner(DiagStartIdx + 2)
+			: FrameSpaceBox.GetCorner(DiagStartIdx + 2) + CornerExtrudeVector;
+		LineSet->AddLine(DiagCorner1, DiagCorner2, CornerModeWireframeColor,
+			CornerModeWireframeThickness, CornerModeWireframeDepthBias);
+	}
+}
+
+void UCubeGridTool::ApplyFlipSelection()
+{
+	if (!bHaveSelection)
+	{
+		return;
+	}
+
+	GetToolManager()->BeginUndoTransaction(LOCTEXT("FlipTransactionName", "Flip Selection"));
+
+	FSelection NewSelection = Selection;
+	Selection.Direction = FCubeGrid::FlipDir(Selection.Direction);
+	SetSelection(NewSelection, true);
+
+	GetToolManager()->EndUndoTransaction();
+
+	// TODO: We actually probably want some special handling here in Corner mode. For one thing,
+	// we're keeping the selected corners the same, which ends up rotating rather than mirroring
+	// the currently pushed/pulled portion (should fix this, but would need another undo transaction,
+	// at which point we probably want full undo support for corner mode, rather than our current
+	// approach of keeping corner selection and extrude distance...). For another, might a user want
+	// a flip in corner mode to equate to a reversal of push vs pull, instead of a mirror of the same
+	// operation (i.e. you flip a pull and you get a mirrored push rather than mirrored pull)? Not certain.
+}
+
+void UCubeGridTool::ApplySlide(int32 NumBlocks)
+{
+	if (!bHaveSelection)
+	{
+		return;
+	}
+
+	GetToolManager()->BeginUndoTransaction(LOCTEXT("SlideTransactionName", "Slide Selection"));
+	SlideSelection(NumBlocks, true);
+	GetToolManager()->EndUndoTransaction();
+}
+
+void UCubeGridTool::ApplyPushPull(int32 NumBlocks)
+{
+	using namespace CubeGridToolLocals;
+
+	if (!bHaveSelection || (Mode == EMode::Corner && !IsAnyCornerSelected(CornerSelectedFlags)))
+	{
+		return;
+	}
+
+	CurrentExtrudeAmount += NumBlocks;
+
+	InvalidatePreview();
+
+	if (Mode != EMode::Corner)
+	{
+		bWaitingToApplyPreview = true;
+		bBlockUntilPreviewUpdate = false;
+		bAdjustSelectionOnPreviewUpdate = true;
+	}
+}
+
+void UCubeGridTool::SetPowerOfTwoClamped(int32 PowerOfTwo)
+{
+	Settings->PowerOfTwo = FMath::Clamp(PowerOfTwo, 0, Settings->MaxPowerOfTwo);
+	CubeGrid->SetGridPower(Settings->PowerOfTwo);
+
+	// Update CurrentExtrudeAmount to reflect the new step size (mainly important
+	// for corner mode).
+	int32 PowerOfTwoDifference = static_cast<int32>(Settings->PowerOfTwo) - PowerOfTwoPrevious;
+	if (CurrentExtrudeAmount != 0 && PowerOfTwoDifference != 0)
+	{
+		int32 AbsExtrudeAmount = FMath::Abs(CurrentExtrudeAmount);
+
+		int32 NewAbsExtrudeAmount = AbsExtrudeAmount >> PowerOfTwoDifference;
+		if (NewAbsExtrudeAmount << PowerOfTwoDifference != AbsExtrudeAmount)
+		{
+			++NewAbsExtrudeAmount;
+			InvalidatePreview();
+		}
+
+		CurrentExtrudeAmount = NewAbsExtrudeAmount * FMath::Sign(CurrentExtrudeAmount);
+	}
+	UpdateCornerModeLineSet();
+
+	ClearHover();
+	UpdateSelectionLineSet(); // Updates the grid drawn inside
+	UpdateGridLineSet();
+
+	PowerOfTwoPrevious = Settings->PowerOfTwo;
+}
+
+// Action support
+
+void UCubeGridToolActions::PostAction(ECubeGridToolAction Action)
+{
+	if (ParentTool.IsValid())
+	{
+		ParentTool->RequestAction(Action);
+	}
+}
+
+void UCubeGridDuringActivityActions::PostAction(ECubeGridToolAction Action)
+{
+	if (ParentTool.IsValid())
+	{
+		ParentTool->RequestAction(Action);
+	}
+}
+
+void UCubeGridTool::RequestAction(ECubeGridToolAction ActionType)
+{
+	if (PendingAction == ECubeGridToolAction::NoAction)
+	{
+		PendingAction = ActionType;
+	}
+}
+
+void UCubeGridTool::ApplyAction(ECubeGridToolAction ActionType)
+{
+	switch (ActionType)
+	{
+	case ECubeGridToolAction::Push:
+		ApplyPushPull(-Settings->BlocksPerStep);
+		break;
+	case ECubeGridToolAction::Pull:
+		ApplyPushPull(Settings->BlocksPerStep);
+		break;
+	case ECubeGridToolAction::Flip:
+		ApplyFlipSelection();
+		break;
+	case ECubeGridToolAction::SlideForward:
+		ApplySlide(-Settings->BlocksPerStep);
+		break;
+	case ECubeGridToolAction::SlideBack:
+		ApplySlide(Settings->BlocksPerStep);
+		break;
+	case ECubeGridToolAction::DecreasePowerOfTwo:
+		// cast is just to be explicit
+		SetPowerOfTwoClamped(static_cast<int32>(Settings->PowerOfTwo) - 1); 
+		break;
+	case ECubeGridToolAction::IncreasePowerOfTwo:
+		SetPowerOfTwoClamped(Settings->PowerOfTwo + 1);
+		break;
+
+	case ECubeGridToolAction::CornerMode:
+		StartCornerMode();
+		break;
+	//case ECubeGridToolAction::FitGrid:
+	//	StartFitGrid();
+	//	break;
+	case ECubeGridToolAction::Done:
+		if (Mode == EMode::Corner)
+		{
+			ApplyCornerMode();
+		}
+		else if (Mode == EMode::FitGrid)
+		{
+			//CancelFitGrid();
+		}
+		break;
+	case ECubeGridToolAction::Cancel:
+		RevertToDefaultMode();
+		break;
+	}
+}
+
+void UCubeGridTool::RegisterActions(FInteractiveToolActionSet& ActionSet)
+{
+	int32 ActionID = (int32)EStandardToolActions::BaseClientDefinedActionID + 1;
+	ActionSet.RegisterAction(this, ActionID++,
+		TEXT("PullBlock"),
+		LOCTEXT("PullAction", "Pull Out Blocks"),
+		LOCTEXT("PullTooltip", ""),
+		EModifierKey::None, EKeys::Q,
+		[this]() { RequestAction(ECubeGridToolAction::Pull); });
+	ActionSet.RegisterAction(this, ActionID++,
+		TEXT("PushBlock"),
+		LOCTEXT("PushAction", "Push In Holes"),
+		LOCTEXT("PushTooltip", ""),
+		EModifierKey::None, EKeys::E,
+		[this]() { RequestAction(ECubeGridToolAction::Push); });
+	ActionSet.RegisterAction(this, ActionID++,
+		TEXT("SlideBack"),
+		LOCTEXT("SlideBackAction", "Slide Selection Back"),
+		LOCTEXT("SlideBackTooltip", ""),
+		EModifierKey::Shift, EKeys::Q,
+		[this]() { RequestAction(ECubeGridToolAction::SlideBack); });
+	ActionSet.RegisterAction(this, ActionID++,
+		TEXT("SlideForward"),
+		LOCTEXT("SlideForwardAction", "Slide Selection Forward"),
+		LOCTEXT("SlideForwardTooltip", ""),
+		EModifierKey::Shift, EKeys::E,
+		[this]() { RequestAction(ECubeGridToolAction::SlideForward); });
+
+	ActionSet.RegisterAction(this, ActionID++,
+		TEXT("DecreasePowerOfTwo"),
+		LOCTEXT("DecreasePowerOfTwoAction", "Decrease Power Of Two"),
+		LOCTEXT("DecreasePowerOfTwoTooltip", ""),
+		EModifierKey::Control, EKeys::Q,
+		[this]() { RequestAction(ECubeGridToolAction::DecreasePowerOfTwo); });
+	ActionSet.RegisterAction(this, ActionID++,
+		TEXT("IncreasePowerOfTwo"),
+		LOCTEXT("IncreasePowerOfTwoAction", "Increase Power Of Two"),
+		LOCTEXT("IncreasePowerOfTwoTooltip", ""),
+		EModifierKey::Control, EKeys::E,
+		[this]() { RequestAction(ECubeGridToolAction::IncreasePowerOfTwo); });
+
+	ActionSet.RegisterAction(this, ActionID++,
+		TEXT("ToggleGizmoVisibility"),
+		LOCTEXT("ToggleGizmoVisibilityAction", "Toggle Gizmo Visibility"),
+		LOCTEXT("ToggleGizmoVisibilityTooltip", ""),
+		EModifierKey::None, EKeys::R,
+		[this]() {
+			if (Mode != EMode::FitGrid)
+			{
+				UpdateGizmoVisibility(!GridGizmo->IsVisible());
+			}
+		});
+
+	ActionSet.RegisterAction(this, ActionID++,
+		TEXT("ToggleCornerMode"),
+		LOCTEXT("ToggleCornerModeAction", "Toggle Corner Mode"),
+		LOCTEXT("ToggleCornerModeTooltip", ""),
+		EModifierKey::None, EKeys::Z,
+		[this]() {
+			if (Mode != EMode::Corner)
+			{
+				StartCornerMode();
+			}
+			else
+			{
+				ApplyCornerMode();
+			}
+		});
+
+	ActionSet.RegisterAction(this, ActionID++,
+		TEXT("ToggleDiagonalMode"),
+		LOCTEXT("ToggleDiagonalModeAction", "Toggle Diagonal Mode"),
+		LOCTEXT("ToggleDiagonalModeTooltip", ""),
+		EModifierKey::None, EKeys::Z,
+		[this]() {
+			if (Mode == EMode::Corner)
+			{
+				Settings->bCrosswiseDiagonal = !Settings->bCrosswiseDiagonal;
+			}
+		});
+
+	ActionSet.RegisterAction(this, ActionID++,
+		TEXT("FlipSelection"),
+		LOCTEXT("FlipSelectionAction", "Flip Selection"),
+		LOCTEXT("FlipSelectionTooltip", ""),
+		EModifierKey::None, EKeys::T,
+		[this]() {
+			ApplyFlipSelection();
+		});
+}
+
+void UCubeGridTool::StartCornerMode()
+{
+	using namespace CubeGridToolLocals;
+
+	if (!bHaveSelection)
+	{
+		// TODO: Write out a message here and clear it at some point
+		return;
+	}
+	if (Mode == EMode::Corner)
+	{
+		return; // Already in mode
+	}
+
+	// Clear/cancel stuff
+	//if (Mode == EMode::FitGrid)
+	//{
+	//	CancelFitGrid();
+	//}
+	CurrentExtrudeAmount = 0;
+	InvalidatePreview();
+
+	// Clear selected corner render
+	for (int i = 0; i < 4; ++i)
+	{
+		CornerSelectedFlags[i] = false;
+		PreDragCornerSelectedFlags[i] = false;
+	}
+
+	UpdateCornerGeometrySet();
+	Mode = EMode::Corner;
+
+	SetToolPropertySourceEnabled(ToolActions, false);
+	SetToolPropertySourceEnabled(DuringActivityActions, true);
+
+	Settings->bInCornerMode = true;
+	NotifyOfPropertyChangeByTool(Settings);
+
+	GetToolManager()->BeginUndoTransaction(ModeChangeTransactionName);
+	GetToolManager()->EmitObjectChange(this, MakeUnique<FCubeGridToolModeChange>(), ModeChangeTransactionName);
+	GetToolManager()->EndUndoTransaction();
+
+	GetToolManager()->DisplayMessage(CornerModeMessage, EToolMessageLevel::UserNotification);
+}
+
+void UCubeGridTool::UpdateCornerGeometrySet()
+{
+	using namespace CubeGridToolLocals;
+
+	FOrientedBox3d FrameSpaceBox = ConvertToOrientedBox(Selection.Box, Selection.Direction);
+	CornersGeometrySet.Reset();
+	for (int i = 0; i < 4; ++i)
+	{
+		CornersGeometrySet.AddPoint(i, CubeGrid->GetFrame().FromFramePoint(FrameSpaceBox.GetCorner(i)));
+	}
+}
+
+void UCubeGridTool::ApplyCornerMode(bool bDontWaitForTick)
+{
+	using namespace CubeGridToolLocals;
+
+	if (CurrentExtrudeAmount != 0 && IsAnyCornerSelected(CornerSelectedFlags))
+	{
+		bWaitingToApplyPreview = true;
+		bBlockUntilPreviewUpdate = true;
+		bAdjustSelectionOnPreviewUpdate = false;
+
+		if (bDontWaitForTick)
+		{
+			ApplyPreview();
+		}
+	}
+
+	CornersGeometrySet.Reset();
+
+	Mode = EMode::PushPull;
+	GetToolManager()->DisplayMessage(PushPullModeMessage, EToolMessageLevel::UserNotification);
+	SetToolPropertySourceEnabled(ToolActions, true);
+	SetToolPropertySourceEnabled(DuringActivityActions, false);
+
+	Settings->bInCornerMode = false;
+	NotifyOfPropertyChangeByTool(Settings);
+
+	UpdateCornerModeLineSet();
+}
+
+void UCubeGridTool::CancelCornerMode()
+{
+	using namespace CubeGridToolLocals;
+
+	CornersGeometrySet.Reset();
+
+	Mode = EMode::PushPull;
+	GetToolManager()->DisplayMessage(PushPullModeMessage, EToolMessageLevel::UserNotification);
+	SetToolPropertySourceEnabled(ToolActions, true);
+	SetToolPropertySourceEnabled(DuringActivityActions, false);
+
+	CurrentExtrudeAmount = 0;
+	InvalidatePreview();
+
+	Settings->bInCornerMode = false;
+	NotifyOfPropertyChangeByTool(Settings);
+
+	UpdateCornerModeLineSet();
+}
+
+void UCubeGridTool::UpdateUsingMeshChange(const FDynamicMeshChange& MeshChange, bool bRevert)
+{
+	MeshChange.Apply(CurrentMesh.Get(), bRevert);
+	MeshSpatial->Build();
+	CurrentExtrudeAmount = 0;
+	bPreviewMayDiffer = true;
+	InvalidatePreview();
+}
+
+bool UCubeGridTool::IsInDefaultMode() const
+{
+	return Mode == EMode::PushPull;
+}
+
+void UCubeGridTool::RevertToDefaultMode()
+{
+	if (Mode == EMode::Corner)
+	{
+		CancelCornerMode();
+	}
+	else if (Mode == EMode::FitGrid)
+	{
+		//CancelFitGrid();
+	}
+}
+
+bool UCubeGridTool::CanCurrentlyNestedCancel()
+{
+	return Mode == EMode::Corner || bHaveSelection;
+}
+
+bool UCubeGridTool::ExecuteNestedCancelCommand()
+{
+	if (!IsInDefaultMode())
+	{
+		RevertToDefaultMode();
+		return true;
+	}
+	else if (bHaveSelection)
+	{
+		ClearSelection(true);
+	}
+	return false;
+}
+
+bool UCubeGridTool::CanCurrentlyNestedAccept()
+{
+	return Mode == EMode::Corner;
+}
+
+bool UCubeGridTool::ExecuteNestedAcceptCommand()
+{
+	if (Mode == EMode::Corner)
+	{
+		ApplyCornerMode();
+		return true;
+	}
+	return false;
+}
+
+#undef LOCTEXT_NAMESPACE
diff --git a/Engine/Plugins/Experimental/MeshModelingToolsetExp/Source/MeshModelingToolsExp/Public/CubeGridTool.h b/Engine/Plugins/Experimental/MeshModelingToolsetExp/Source/MeshModelingToolsExp/Public/CubeGridTool.h
new file mode 100644
--- /dev/null
+++ b/Engine/Plugins/Experimental/MeshModelingToolsetExp/Source/MeshModelingToolsExp/Public/CubeGridTool.h
@@ -0,0 +1,505 @@
+// Copyright Epic Games, Inc. All Rights Reserved.
+
+#pragma once
+
+#include "CoreMinimal.h"
+
+#include "BaseBehaviors/BehaviorTargetInterfaces.h"
+#include "BoxTypes.h"
+#include "DynamicMesh/DynamicMeshAABBTree3.h"
+#include "FrameTypes.h"
+#include "GeometryBase.h"
+#include "InteractiveTool.h"
+#include "InteractiveToolBuilder.h"
+#include "InteractiveToolQueryInterfaces.h" // IInteractiveToolNestedAcceptCancelAPI
+#include "Mechanics/CubeGrid.h"
+#include "ModelingOperators.h"
+#include "OrientedBoxTypes.h"
+#include "Spatial/GeometrySet3.h"
+#include "ToolContextInterfaces.h" // FViewCameraState
+#include "ToolDataVisualizer.h"
+
+#include "CubeGridTool.generated.h"
+
+PREDECLARE_GEOMETRY(class FDynamicMesh3);
+PREDECLARE_GEOMETRY(class FDynamicMeshChange);
+PREDECLARE_GEOMETRY(class FCubeGrid);
+
+class IAssetGenerationAPI;
+class UClickDragInputBehavior;
+class UCubeGridTool;
+class UCreateMeshObjectTypeProperties;
+class UDragAlignmentMechanic;
+class ULocalClickDragInputBehavior;
+class ULocalSingleClickInputBehavior;
+class UMeshOpPreviewWithBackgroundCompute;
+class UMouseHoverBehavior;
+class UNewMeshMaterialProperties;
+class UPreviewGeometry;
+class UTransformGizmo;
+class UTransformProxy;
+class UToolTarget;
+
+
+UCLASS()
+class MESHMODELINGTOOLSEXP_API UCubeGridToolBuilder : public UInteractiveToolBuilder
+{
+	GENERATED_BODY()
+
+public:
+	IAssetGenerationAPI* AssetAPI = nullptr;
+
+	virtual bool CanBuildTool(const FToolBuilderState& SceneState) const override;
+	virtual UInteractiveTool* BuildTool(const FToolBuilderState& SceneState) const override;
+
+protected:
+
+	virtual const FToolTargetTypeRequirements& GetTargetRequirements() const override;
+};
+
+
+UENUM()
+enum class ECubeGridToolFaceSelectionMode
+{
+	OutsideBasedOnNormal,
+	InsideBasedOnNormal,
+	OutsideBasedOnViewRay,
+	InsideBasedOnViewRay
+};
+
+
+UCLASS()
+class MESHMODELINGTOOLSEXP_API UCubeGridToolProperties : public UInteractiveToolPropertySet
+{
+	GENERATED_BODY()
+
+public:
+
+	UPROPERTY(EditAnywhere, Category = Options, meta = (
+		EditCondition = "bAllowedToEditGrid", HideEditConditionToggle))
+	FVector GridFrameOrigin = FVector(0, 0, 0);
+
+	UPROPERTY(EditAnywhere, Category = Options, meta = (
+		EditCondition = "bAllowedToEditGrid", HideEditConditionToggle,
+		UIMin = -180, UIMax = 180, ClampMin = -180000, ClampMax = 18000))
+	FRotator GridFrameOrientation = FRotator(0, 0, 0);
+
+	UPROPERTY(EditAnywhere, Category = Options, meta = (TransientToolProperty))
+	bool bShowGizmo = false;
+
+	/** How many blocks each push/pull invocation will do at a time.*/
+	UPROPERTY(EditAnywhere, Category = Options, meta = (
+		UIMin = "0", ClampMin = "0"))
+	int32 BlocksPerStep = 1;
+
+	/** Determines cube grid scale. Can also be adjusted with Ctrl+E/Q. */
+	UPROPERTY(EditAnywhere, Category = Options, meta = (
+		EditCondition = "bAllowedToEditGrid", HideEditConditionToggle,
+		UIMin = "0", UIMax = "10", ClampMin = "0", ClampMax = "31"))
+	uint8 PowerOfTwo = 5;
+
+	// Must match ClampMax in PowerOfTwo, used to make hotkeys not exceed it.
+	const uint8 MaxPowerOfTwo = 31;
+
+	/** Smallest block size to use in the grid. For instance, 3.125 results in
+	 blocks that are 100 sized at power of 2 since 3.125 * 2^5 = 100. */
+	UPROPERTY(EditAnywhere, Category = Options, meta = (
+		EditCondition = "bAllowedToEditGrid", HideEditConditionToggle,
+		UIMin = "0", UIMax = "10", ClampMin = "0", ClampMax = "1000"))
+	double BlockBaseSize = 3.125;
+
+	/** When pushing/pulling in a way where the diagonal matters, setting this to true
+	 makes the diagonal generally try to lie flat across the face rather than at
+	 an incline. */
+	UPROPERTY(EditAnywhere, Category = Options, meta = (
+		EditCondition = "bInCornerMode", HideEditConditionToggle))
+	bool bCrosswiseDiagonal = false;
+
+	/** When performing selection, the tolerance to use when determining
+	 whether things lie in the same plane as a cube face. */
+	UPROPERTY(EditAnywhere, Category = BlockSelection, AdvancedDisplay, meta = (
+		UIMin = "0", UIMax = "0.5", ClampMin = "0", ClampMax = "10"))
+	double PlaneTolerance = 0.01;
+
+	/** When raycasting to find a selected grid face, this determines whether geometry
+	 in the scene that is not part of the edited mesh is hit. */
+	UPROPERTY(EditAnywhere, Category = BlockSelection)
+	bool bHitUnrelatedGeometry = true;
+
+	/** When the grid ground plane is above some geometry, whether we should hit that
+	 plane or pass through to the other geometry. */
+	UPROPERTY(EditAnywhere, Category = BlockSelection, AdvancedDisplay)
+	bool bHitGridGroundPlaneIfCloser = false;
+	
+	/** How the selected face is determined. */
+	UPROPERTY(EditAnywhere, Category = BlockSelection, AdvancedDisplay)
+	ECubeGridToolFaceSelectionMode FaceSelectionMode = ECubeGridToolFaceSelectionMode::OutsideBasedOnNormal;
+
+	UPROPERTY(VisibleAnywhere, Category = ShortcutInfo)
+	FString ToggleCornerMode = TEXT("Z to start/complete corner mode.");
+
+	UPROPERTY(VisibleAnywhere, Category = ShortcutInfo)
+	FString PushPull = TEXT("E and Q push and pull, or use Ctrl+drag.");
+
+	UPROPERTY(VisibleAnywhere, Category = ShortcutInfo)
+	FString ResizeGrid = TEXT("Ctrl + E/Q resize the grid in powers of two.");
+
+	UPROPERTY(VisibleAnywhere, Category = ShortcutInfo)
+	FString SlideSelection = TEXT("Middle mouse drag to slide selection in plane. "
+		"Shift + E/Q shift selection forward/back.");
+
+	UPROPERTY(VisibleAnywhere, Category = ShortcutInfo)
+	FString FlipSelection = TEXT("T flips the selection.");
+
+	UPROPERTY(VisibleAnywhere, Category = ShortcutInfo)
+	FString GridGizmo = TEXT("R shows/hides grid gizmo. ");
+
+	UPROPERTY(VisibleAnywhere, Category = ShortcutInfo)
+	FString QuickShiftGizmo = TEXT("Ctrl + middle click quick-repositions "
+		"the gizmo while keeping it on grid");
+
+	UPROPERTY(VisibleAnywhere, Category = ShortcutInfo)
+	FString AlignGizmo = TEXT("While dragging gizmo handles, hold Ctrl to align "
+		"to items in scene (constrained to the moved axes).");
+
+	UPROPERTY(meta = (TransientToolProperty))
+	bool bInCornerMode = false;
+
+	//~ Currently unused... Used to disallow it during corner mode, might do so again.
+	UPROPERTY(meta = (TransientToolProperty))
+	bool bAllowedToEditGrid = true;
+};
+
+UENUM()
+enum class ECubeGridToolAction
+{
+	NoAction,
+	Push,
+	Pull,
+	Flip,
+	SlideForward,
+	SlideBack,
+	DecreasePowerOfTwo,
+	IncreasePowerOfTwo,
+	CornerMode,
+	// FitGrid,
+	Done,
+	Cancel,
+};
+
+UCLASS()
+class MESHMODELINGTOOLSEXP_API UCubeGridToolActions : public UInteractiveToolPropertySet
+{
+	GENERATED_BODY()
+
+public:
+	TWeakObjectPtr<UCubeGridTool> ParentTool;
+
+	void Initialize(UCubeGridTool* ParentToolIn) { ParentTool = ParentToolIn; }
+
+	void PostAction(ECubeGridToolAction Action);
+
+	/** Engages a mode where specific corners can be selected to push/pull only
+	 those corners. Press Apply to commit the result afterward. Can also be toggled
+	 with Z. */
+	UFUNCTION(CallInEditor, Category = Actions)
+	void CornerMode() { PostAction(ECubeGridToolAction::CornerMode); }
+	
+	/** Can also be invoked with E */
+	UFUNCTION(CallInEditor, Category = Actions)
+	void Push() { PostAction(ECubeGridToolAction::Push); }
+
+	/** Can aslo be invoked with Q */
+	UFUNCTION(CallInEditor, Category = Actions)
+	void Pull() { PostAction(ECubeGridToolAction::Pull); }
+
+	/** Can also be invoked with T. */
+	UFUNCTION(CallInEditor, Category = Actions)
+	void Flip() { PostAction(ECubeGridToolAction::Flip); }
+
+	/** Can also be invoked with Shift + E. */
+	UFUNCTION(CallInEditor, Category = Actions)
+	void SlideForward() { PostAction(ECubeGridToolAction::SlideForward); }
+
+	/** Can also be invoked with Shift + Q. */
+	UFUNCTION(CallInEditor, Category = Actions)
+	void SlideBack() { PostAction(ECubeGridToolAction::SlideBack); }
+};
+
+UCLASS()
+class MESHMODELINGTOOLSEXP_API UCubeGridDuringActivityActions : public UInteractiveToolPropertySet
+{
+	GENERATED_BODY()
+
+public:
+	TWeakObjectPtr<UCubeGridTool> ParentTool;
+
+	void Initialize(UCubeGridTool* ParentToolIn) { ParentTool = ParentToolIn; }
+
+	void PostAction(ECubeGridToolAction Action);
+
+	/** Accept and complete current action. */
+	UFUNCTION(CallInEditor, Category = Actions)
+	void Done() { PostAction(ECubeGridToolAction::Done); }
+
+	/** Cancel and exit current action */
+	UFUNCTION(CallInEditor, Category = Actions)
+	void Cancel() { PostAction(ECubeGridToolAction::Cancel); }
+};
+
+/** Tool that allows for blocky boolean operations on an orientable power-of-two grid. */
+UCLASS()
+class MESHMODELINGTOOLSEXP_API UCubeGridTool : public UInteractiveTool,
+	public IClickDragBehaviorTarget, public IHoverBehaviorTarget,
+	public UE::Geometry::IDynamicMeshOperatorFactory,
+	public IInteractiveToolNestedAcceptCancelAPI
+{
+	GENERATED_BODY()
+protected:
+	enum class EMouseState
+	{
+		NotDragging,
+		DraggingExtrudeDistance,
+		DraggingCornerSelection,
+		DraggingRegularSelection,
+	};
+
+	enum class EMode
+	{
+		PushPull,
+		Corner,
+
+		// This is currently not supported, but some of the code was written with space
+		// for a "fit grid" mode that would have allowed the dimensions of the grid to
+		// be fit using a sequence of (snapped) mouse clicks. It seems useful to leave 
+		// those code stubs for now in case we add the mode in, so it's easier to track
+		// down the affected portions of code.
+		FitGrid
+	};
+
+	using FCubeGrid = UE::Geometry::FCubeGrid;
+
+public:
+
+	struct FSelection
+	{
+		// Both of these boxes are in the coordinate space of the (unscaled) grid frame.
+		UE::Geometry::FAxisAlignedBox3d Box;
+		UE::Geometry::FAxisAlignedBox3d StartBox; // Box delineating original selected face
+		FCubeGrid::EFaceDirection Direction;
+
+		bool operator==(const FSelection& Other)
+		{
+			return Box == Other.Box
+				&& StartBox == Other.StartBox
+				&& Direction == Other.Direction;
+		}
+
+		bool operator!=(const FSelection& Other)
+		{
+			return !(*this == Other);
+		}
+	};
+
+	virtual void Setup() override;
+	virtual void Shutdown(EToolShutdownType ShutdownType) override;
+
+	virtual void SetTarget(TObjectPtr<UToolTarget> TargetIn) { Target = TargetIn; }
+	virtual void SetWorld(UWorld* World) { TargetWorld = World; }
+
+	virtual void OnTick(float DeltaTime) override;
+	virtual void Render(IToolsContextRenderAPI* RenderAPI) override;
+
+	virtual void OnPropertyModified(UObject* PropertySet, FProperty* Property) override;
+
+	virtual void RegisterActions(FInteractiveToolActionSet& ActionSet) override;
+
+	virtual void RequestAction(ECubeGridToolAction ActionType);
+
+	virtual void SetSelection(const FSelection& Selection, bool bEmitChange);
+	virtual void ClearSelection(bool bEmitChange);
+
+	// Used by undo/redo
+	virtual void UpdateUsingMeshChange(const UE::Geometry::FDynamicMeshChange& MeshChange, bool bRevert);
+	virtual bool IsInDefaultMode() const;
+	virtual void RevertToDefaultMode();
+
+	// IClickDragBehaviorTarget
+	virtual FInputRayHit CanBeginClickDragSequence(const FInputDeviceRay& PressPos) override;
+	virtual void OnClickPress(const FInputDeviceRay& PressPos) override;
+	virtual void OnClickDrag(const FInputDeviceRay& DragPos) override;
+	virtual void OnClickRelease(const FInputDeviceRay& ReleasePos) override;
+	virtual void OnTerminateDragSequence() override;
+
+	// IHoverBehaviorTarget
+	virtual FInputRayHit BeginHoverSequenceHitTest(const FInputDeviceRay& PressPos) override;
+	virtual void OnBeginHover(const FInputDeviceRay& DevicePos) override;
+	virtual bool OnUpdateHover(const FInputDeviceRay& DevicePos) override;
+	virtual void OnEndHover() override;
+	virtual void OnUpdateModifierState(int ModifierID, bool bIsOn) override;
+
+	// IDynamicMeshOperatorFactory API
+	virtual TUniquePtr<UE::Geometry::FDynamicMeshOperator> MakeNewOperator() override;
+
+	// IInteractiveToolNestedAcceptCancelAPI
+	virtual bool SupportsNestedCancelCommand() override { return true; }
+	virtual bool CanCurrentlyNestedCancel() override;
+	virtual bool ExecuteNestedCancelCommand() override;
+	virtual bool SupportsNestedAcceptCommand() override { return true; }
+	virtual bool CanCurrentlyNestedAccept() override;
+	virtual bool ExecuteNestedAcceptCommand() override;
+
+protected:
+
+	UPROPERTY()
+	TObjectPtr<UTransformGizmo> GridGizmo = nullptr;
+
+	UPROPERTY()
+	TObjectPtr<UDragAlignmentMechanic> GridGizmoAlignmentMechanic = nullptr;
+	
+	UPROPERTY()
+	TObjectPtr<UTransformProxy> GridGizmoTransformProxy = nullptr;
+
+	UPROPERTY()
+	TObjectPtr<UPreviewGeometry> LineSets = nullptr;
+
+	UPROPERTY()
+	TObjectPtr<UClickDragInputBehavior> ClickDragBehavior = nullptr;
+	
+	UPROPERTY()
+	TObjectPtr<UMouseHoverBehavior> HoverBehavior = nullptr;
+
+	UPROPERTY()
+	TObjectPtr<ULocalSingleClickInputBehavior> CtrlMiddleClickBehavior = nullptr;
+
+	UPROPERTY()
+	TObjectPtr<ULocalClickDragInputBehavior> MiddleClickDragBehavior = nullptr;
+
+	// Properties, etc
+	UPROPERTY()
+	TObjectPtr<UCubeGridToolProperties> Settings = nullptr;
+
+	UPROPERTY()
+	TObjectPtr<UCubeGridToolActions> ToolActions = nullptr;
+
+	UPROPERTY()
+	TObjectPtr<UCubeGridDuringActivityActions> DuringActivityActions = nullptr;
+
+	UPROPERTY()
+	TObjectPtr<UNewMeshMaterialProperties> MaterialProperties = nullptr;
+
+	UPROPERTY()
+	TObjectPtr<UCreateMeshObjectTypeProperties> OutputTypeProperties = nullptr;
+
+	// Existing asset to modify, if one was selected
+	UPROPERTY()
+	TObjectPtr<UToolTarget> Target = nullptr;
+
+	TSharedPtr<FCubeGrid, ESPMode::ThreadSafe> CubeGrid;
+
+	// Where to make the preview, new mesh, etc
+	TObjectPtr<UWorld> TargetWorld = nullptr;
+
+	// Important state. Could refactor things into tool activities (like PolyEdit)
+	// someday.
+	EMode Mode = EMode::PushPull;
+	EMouseState MouseState = EMouseState::NotDragging;
+
+	bool GetHitGridFace(const FRay& WorldRay, UE::Geometry::FCubeGrid::FCubeFace& FaceOut);
+	bool UpdateHover(const FRay& WorldRay);
+	void UpdateHoverLineSet(bool bHaveHover, const UE::Geometry::FAxisAlignedBox3d& HoveredBox);
+	void UpdateSelectionLineSet();
+	void UpdateGridLineSet();
+	void UpdateCornerModeLineSet();
+	void ClearHover();
+
+	bool bHaveSelection = false;
+	FSelection Selection;
+
+	bool bPreviousHaveSelection = false;
+	FSelection PreviousSelection;
+
+	bool bHaveHoveredSelection = false;
+	UE::Geometry::FAxisAlignedBox3d HoveredSelectionBox;
+
+	void SlideSelection(int32 ExtrudeAmount, bool bEmitChange);
+
+	bool bSlideToggle = false;
+	bool bSelectionToggle = false;
+	bool bChangeSideToggle = false;
+	bool bMouseDragShouldPushPull = false;
+	UE::Geometry::FRay3d DragProjectionAxis;
+	double DragProjectedStartParam;
+
+	void ApplyFlipSelection();
+	void ApplySlide(int32 NumBlocks);
+	void ApplyPushPull(int32 NumBlocks);
+
+	// Parameter is signed on purpose so we can give negatives
+	void SetPowerOfTwoClamped(int32 PowerOfTwo);
+	uint8 PowerOfTwoPrevious;
+
+	UPROPERTY()
+	TObjectPtr<UMeshOpPreviewWithBackgroundCompute> Preview;
+
+	/**
+	 * @param bUpdateCornerLineSet This can be set to false when the invalidation
+	 *  is a result of a grid transform change (which is applied to the line set
+	 *  via the LineSets->SetTransform), or when the corner shape otherwise doesn't
+	 *  change. Usually it can be left to true.
+	 */
+	void InvalidatePreview(bool bUpdateCornerLineSet = true);
+
+	void ApplyPreview();
+
+	int32 CurrentExtrudeAmount = 0;
+	bool bPreviewMayDiffer = false;
+	bool bWaitingToApplyPreview = false;
+	bool bBlockUntilPreviewUpdate = false;
+	bool bAdjustSelectionOnPreviewUpdate = false;
+
+	TSharedPtr<UE::Geometry::FDynamicMesh3, ESPMode::ThreadSafe> CurrentMesh;
+	TSharedPtr<UE::Geometry::FDynamicMeshAABBTree3, ESPMode::ThreadSafe> MeshSpatial;
+	UE::Geometry::FTransform3d CurrentMeshTransform = UE::Geometry::FTransform3d::Identity();
+	TSharedPtr<TArray<int32>, ESPMode::ThreadSafe> LastOpChangedTids;
+
+	static const int32 ShiftModifierID = 1;
+	static const int32 CtrlModifierID = 2;
+
+	ECubeGridToolAction PendingAction = ECubeGridToolAction::NoAction;
+	void ApplyAction(ECubeGridToolAction ActionType);
+
+	int32 GridFrameOriginWatcherIdx;
+	int32 GridFrameOrientationWatcherIdx;
+	void GridGizmoMoved(UTransformProxy* Proxy, FTransform Transform);
+	void UpdateGizmoVisibility(bool bVisible);
+	bool bInGizmoDrag = false;
+
+	FVector3d MiddleClickDragStart;
+	FInputRayHit RayCastSelectionPlane(const UE::Geometry::FRay3d& WorldRay, 
+		FVector3d& HitPointOut);
+	FInputRayHit CanBeginMiddleClickDrag(const FInputDeviceRay& ClickPos);
+	void OnMiddleClickDrag(const FInputDeviceRay& DragPos);
+	void OnCtrlMiddleClick(const FInputDeviceRay& ClickPos);
+	void PrepForSelectionChange();
+	void EndSelectionChange();
+
+	// Used in corner push/pull mode. If you create a flat FOrientedBox3d out of 
+	// the current selection, with the Z axis being along selection normal, the 0-3
+	// indices here correspond to the 0-3 corner indices in the box, which are the
+	// bottom corners along Z axis (see TOrientedBox3::GetCorner())
+	bool CornerSelectedFlags[4] = { false, false, false, false };
+	bool PreDragCornerSelectedFlags[4] = { false, false, false, false };
+
+	FViewCameraState CameraState;
+	FToolDataVisualizer SelectedCornerRenderer;
+	UE::Geometry::FGeometrySet3 CornersGeometrySet;
+	void UpdateCornerGeometrySet();
+	void StartCornerMode();
+	void ApplyCornerMode(bool bDontWaitForTick = false);
+	void CancelCornerMode();
+	void AttemptToSelectCorner(const UE::Geometry::FRay3d& WorldRay);
+
+	// Used to see if we need to update the asset that we've been modifying
+	bool bChangesMade = false;
+};
diff --git a/Engine/Plugins/Experimental/MeshModelingToolsetExp/Source/ModelingComponents/Private/Mechanics/CubeGrid.cpp b/Engine/Plugins/Experimental/MeshModelingToolsetExp/Source/ModelingComponents/Private/Mechanics/CubeGrid.cpp
new file mode 100644
--- /dev/null
+++ b/Engine/Plugins/Experimental/MeshModelingToolsetExp/Source/ModelingComponents/Private/Mechanics/CubeGrid.cpp
@@ -0,0 +1,130 @@
+// Copyright Epic Games, Inc. All Rights Reserved.
+
+#include "Mechanics/CubeGrid.h"
+
+#include "BaseGizmos/GizmoMath.h"
+#include "Math/Ray.h"
+
+#include "ExplicitUseGeometryMathTypes.h"		// using UE::Geometry::(math types)
+using namespace UE::Geometry;
+
+#define LOCTEXT_NAMESPACE "CubeGrid"
+
+namespace CubeGridLocals {
+	FVector3d Floor(const FVector3d& VectorIn)
+	{
+		return FVector3d(
+			FMath::Floor(VectorIn.X),
+			FMath::Floor(VectorIn.Y),
+			FMath::Floor(VectorIn.Z));
+	}
+}
+
+FCubeGrid::FCubeFace::FCubeFace(const FVector3d& PointOnFace, EFaceDirection DirectionIn, uint8 SourceCubeGridPower)
+	: Direction(DirectionIn)
+	, CubeGridPower(SourceCubeGridPower)
+{
+	FVector3d MinCorner = CubeGridLocals::Floor(PointOnFace);
+
+	Center = FVector3d(0.5, 0.5, 0.5);
+	Center[DirToFlatDim(DirectionIn)] = 0;
+	Center += MinCorner;
+}
+
+FVector3d FCubeGrid::FCubeFace::GetMinCorner() const
+{
+	return CubeGridLocals::Floor(Center);
+}
+FVector3d FCubeGrid::FCubeFace::GetMaxCorner() const
+{
+	return FVector3d(
+		FMath::CeilToDouble(Center.X),
+		FMath::CeilToDouble(Center.Y),
+		FMath::CeilToDouble(Center.Z));
+}
+
+
+bool FCubeGrid::GetHitGridFaceBasedOnRay(const FVector3d& WorldHitPoint, const FVector3d& NormalOrTowardCameraRay,
+	FCubeGrid::FCubeFace& FaceOut, bool bPierceToBack, double PlaneTolerance) const
+{
+	// TODO: Scale plane tolerance to grid space
+
+	FVector3d GridSpaceHitPoint = ToGridPoint(WorldHitPoint);
+	FVector3d GridSpaceRayDirection = GridFrame.ToFrameVector(NormalOrTowardCameraRay);
+
+	FVector3d CellMinCorner = CubeGridLocals::Floor(GridSpaceHitPoint);
+
+	// Generally, we start from the point in grid space and hit test the faces of the enclosing cell using the
+	// given ray direction to figure out which face to select. However, there is ambiguity about which cell to
+	// place our point when it lies on a face. In those cases, if the ray direction is inward into the current
+	// cell, we need to choose the adjacent one so that the direction is outward. We also clamp the point to
+	// the face because it makes it easier to treat the case properly based on what we do to t=0 during our
+	// ray cast later, and guarantees that the point is inside (boundary-inclusive) our cell.
+	for (int32 i = 0; i < 3; ++i)
+	{
+		double LowerFace = CellMinCorner[i];
+		double UpperFace = LowerFace + 1;
+
+		if ((UpperFace - GridSpaceHitPoint[i]) * CurrentCellSize <= PlaneTolerance)
+		{
+			GridSpaceHitPoint[i] = UpperFace;
+
+			// Shift the chosen cell up if ray is going in the negative direction
+			CellMinCorner[i] += GridSpaceRayDirection[i] < 0 ? 1 : 0;
+		}
+		else if ((GridSpaceHitPoint[i] - LowerFace) * CurrentCellSize <= PlaneTolerance)
+		{
+			GridSpaceHitPoint[i] = LowerFace;
+
+			// Shift the chosen cell down if ray is going in the positive direction
+			CellMinCorner[i] += GridSpaceRayDirection[i] > 0 ? -1 : 0;
+		}
+	}
+
+	// Now do ray cast:
+	GridSpaceRayDirection *= bPierceToBack ? -1 : 1;
+	double BestT = TNumericLimits<double>::Max();
+
+	// We're going iterate over XYZ dimensions and the min and max face in each dimension.
+	for (int Dimension = 0; Dimension < 3; ++Dimension)
+	{
+		for (int Face = 0; Face < 2; ++Face) // 0 for min, 1 for other side
+		{
+			if (GridSpaceRayDirection[Dimension] == 0)
+			{
+				continue; // ray is parallel
+			}
+
+			// Since all of our planes are axis-aligned, getting the intersection T is just a matter of dividing
+			// the difference of the relevant coordinate by our dt in that direction (i.e. ray direction coordinate)
+			double PlaneCoordInDimension = CellMinCorner[Dimension] + Face;
+			double PlaneT = (PlaneCoordInDimension - GridSpaceHitPoint[Dimension]) / GridSpaceRayDirection[Dimension];
+
+			if (PlaneT == 0 && bPierceToBack)
+			{
+				// If on a face and piercing back, we don't want to select that same face
+				continue;
+			}
+
+			if (PlaneT >= 0 && PlaneT < BestT)
+			{
+				BestT = PlaneT;
+
+				FVector3d ArbitraryPointOnPlane = GridSpaceHitPoint;
+				ArbitraryPointOnPlane[Dimension] = PlaneCoordInDimension;
+
+				EFaceDirection FaceDirection = static_cast<EFaceDirection>(Dimension + 1);
+
+				// Face direction should actually be flipped either if this is the lower face or if it's
+				// a pierce to back, and the two flips can cancel themselves out.
+				FaceDirection = (bPierceToBack ^ (Face == 0)) ? FlipDir(FaceDirection) : FaceDirection;
+
+				FaceOut = FCubeGrid::FCubeFace(ArbitraryPointOnPlane, FaceDirection, GetGridPower());
+			}
+		}
+	}
+
+	return BestT >= 0 && BestT != TNumericLimits<double>::Max();
+}
+
+#undef LOCTEXT_NAMESPACE
diff --git a/Engine/Plugins/Experimental/MeshModelingToolsetExp/Source/ModelingComponents/Private/Mechanics/DragAlignmentMechanic.cpp b/Engine/Plugins/Experimental/MeshModelingToolsetExp/Source/ModelingComponents/Private/Mechanics/DragAlignmentMechanic.cpp
--- a/Engine/Plugins/Experimental/MeshModelingToolsetExp/Source/ModelingComponents/Private/Mechanics/DragAlignmentMechanic.cpp
+++ b/Engine/Plugins/Experimental/MeshModelingToolsetExp/Source/ModelingComponents/Private/Mechanics/DragAlignmentMechanic.cpp
@@ -179,20 +179,23 @@ void UDragAlignmentMechanic::OnGizmoTransformChanged(FTransform NewTransform)
 void UDragAlignmentMechanic::InitializeDeformedMeshRayCast(TFunction<FDynamicMeshAABBTree3* ()> GetSpatialIn, 
 	const UE::Geometry::FTransform3d &TargetTransform, const FGroupTopologyDeformer* LinearDeformer)
 {
-	// Use the deformer to create a way to filter out triangles we don't want to be hitting (the ones we're manipulating)
-	const FDynamicMesh3* Mesh = LinearDeformer->GetMesh();
-	MeshRayCastTriangleFilter = [Mesh, LinearDeformer](int32 Tid)
+	if (LinearDeformer)
 	{
-		FIndex3i Vids = Mesh->GetTriangle(Tid);
-		for (int32 i = 0; i < 3; ++i)
+		// Use the deformer to create a way to filter out triangles we don't want to be hitting (the ones we're manipulating)
+		const FDynamicMesh3* Mesh = LinearDeformer->GetMesh();
+		MeshRayCastTriangleFilter = [Mesh, LinearDeformer](int32 Tid)
 		{
-			if (LinearDeformer->GetModifiedVertices().Contains(Vids[i]))
+			FIndex3i Vids = Mesh->GetTriangle(Tid);
+			for (int32 i = 0; i < 3; ++i)
 			{
-				return false;
+				if (LinearDeformer->GetModifiedVertices().Contains(Vids[i]))
+				{
+					return false;
+				}
 			}
-		}
-		return true;
-	};
+			return true;
+		};
+	}
 
 	// Set up the actual function that casts rays into the mesh.
 	MeshRayCast = [this, GetSpatialIn, TargetTransform, LinearDeformer](const FRay& WorldRay, FHitResult& HitResult, bool bUseFilter)
diff --git a/Engine/Plugins/Experimental/MeshModelingToolsetExp/Source/ModelingComponents/Public/Mechanics/CubeGrid.h b/Engine/Plugins/Experimental/MeshModelingToolsetExp/Source/ModelingComponents/Public/Mechanics/CubeGrid.h
new file mode 100644
--- /dev/null
+++ b/Engine/Plugins/Experimental/MeshModelingToolsetExp/Source/ModelingComponents/Public/Mechanics/CubeGrid.h
@@ -0,0 +1,175 @@
+// Copyright Epic Games, Inc. All Rights Reserved.
+
+#pragma once
+
+#include "CoreMinimal.h"
+
+#include "VectorTypes.h"
+#include "FrameTypes.h"
+#include "InteractionMechanic.h"
+
+class FRay;
+
+namespace UE {
+namespace Geometry {
+
+/**
+ * Class representing an adjustable and resizable power-of-two grid in which faces can be
+ * selected.
+ */
+class MODELINGCOMPONENTS_API FCubeGrid
+{
+public:
+
+	/** 
+	 * Direction that a face is facing. Values are set up like in IndexUtil, such that
+	 * abs(value)-1 is the nonzero normal direction, and sign(value) is the sign of the
+	 * normal.
+	 */
+	enum class EFaceDirection : int8
+	{
+		NegativeX = -1,
+		PositiveX = 1,
+		NegativeY = -2,
+		PositiveY = 2,
+		NegativeZ = -3,
+		PositiveZ = 3
+	};
+
+	/** Convert face direction to its normal. */
+	static int8 DirToFlatDim(EFaceDirection DirectionIn)
+	{
+		return FMath::Abs(static_cast<int8>(DirectionIn)) - 1;
+	}
+
+	/** 
+	 * Convert face direction to the index of the nonzero normal dimension, or equivalently,
+	 * the dimension in which that face is flat.
+	 */
+	static FVector3d DirToNormal(EFaceDirection DirectionIn)
+	{
+		FVector3d Normal = FVector3d::Zero();
+		Normal[DirToFlatDim(DirectionIn)] = FMath::Sign(static_cast<int8>(DirectionIn));
+		return Normal;
+	}
+
+	/** Get flipped version of a face direction. */
+	static EFaceDirection FlipDir(EFaceDirection DirectionIn)
+	{
+		return static_cast<EFaceDirection>(-static_cast<int8>(DirectionIn));
+	}
+
+	/** 
+	 * A face always refers to coordinates in the grid space. It is always
+	 * on grid, and tied to a particular grid power.
+	 */
+	class MODELINGCOMPONENTS_API FCubeFace
+	{
+	public:
+		FCubeFace(const FVector3d& PointOnFace, EFaceDirection Direction, uint8 SourceCubeGridPower);
+
+		FCubeFace()
+			: FCubeFace(FVector3d::Zero(), 
+				EFaceDirection::PositiveZ,
+				0)
+		{}
+
+		const FVector3d& GetCenter() const { return Center; }
+		EFaceDirection GetDirection() const { return Direction; }
+
+		FVector3d GetNormal() const { 
+			return DirToNormal(Direction);
+		}
+
+		FVector3d GetMinCorner() const;
+		FVector3d GetMaxCorner() const;
+		uint8 GetSourceCubeGridPower() const { return CubeGridPower; }
+
+		bool operator==(const FCubeFace& Other) const
+		{
+			return Center == Other.Center && Direction == Other.Direction;
+		}
+		bool operator!=(const FCubeFace& Other) const
+		{
+			return !(*this == Other);
+		}
+	private:
+		// One of the components will always be an integer, and the other two will always be offset
+		// from an integer by 0.5, since this is the center of a cube grid face.
+		FVector3d Center;
+
+		EFaceDirection Direction;
+		uint8 CubeGridPower;
+	};
+
+	FCubeGrid()
+		: BaseGridCellSize(10)
+		, CurrentGridPower(3)
+	{
+		UpdateCellSize();
+	}
+
+	virtual bool GetHitGridFaceBasedOnRay(const FVector3d& WorldHitPoint, 
+		const FVector3d& NormalOrTowardCameraRay,
+		FCubeGrid::FCubeFace& FaceOut, bool bPierceToBack, double PlaneTolerance) const;
+
+	int8 GetGridPower() const { return CurrentGridPower; }
+
+	void SetGridPower(uint8 PowerIn)
+	{
+		ensureMsgf(PowerIn <= GetMaxGridPower(), TEXT("FCubeGrid: Grid power set above MaxGridPower."));
+		CurrentGridPower = PowerIn;
+		UpdateCellSize();
+	}
+	int8 GetMaxGridPower() const
+	{
+		return sizeof(uint32) * 8 - 1;
+	}
+
+	double GetBaseGridCellSize() const { return BaseGridCellSize; }
+	void SetBaseGridCellSize(double SizeIn)
+	{
+		BaseGridCellSize = SizeIn;
+		UpdateCellSize();
+	}
+
+	double GetCurrentGridCellSize() const { return CurrentCellSize; }
+
+	double GetCellSize(int8 GridPower) const { return BaseGridCellSize * ((uint32)1 << GridPower); }
+
+	FVector3d ToWorldPoint(FVector3d GridPoint) const
+	{
+		return GridFrame.FromFramePoint(GridPoint * CurrentCellSize);
+	}
+
+	FVector3d ToWorldPoint(FVector3d GridPoint, int8 GridPower) const
+	{
+		return GridFrame.FromFramePoint(GridPoint * GetCellSize(GridPower));
+	}
+
+	FVector3d ToGridPoint(FVector3d WorldPoint) const
+	{
+		return GridFrame.ToFramePoint(WorldPoint) / CurrentCellSize;
+	}
+
+	FTransform GetFTransform() const
+	{
+		FTransform ToReturn = GridFrame.ToFTransform();
+		ToReturn.SetScale3D(FVector(GetCurrentGridCellSize()));
+		return ToReturn;
+	}
+
+	const FFrame3d& GetFrame() const { return GridFrame; }
+	void SetGridFrame(const FFrame3d& FrameIn) { GridFrame = FrameIn; }
+protected:
+
+	void UpdateCellSize() { CurrentCellSize = GetCellSize(CurrentGridPower); }
+
+	FFrame3d GridFrame;
+	double BaseGridCellSize;
+	int8 CurrentGridPower;
+	// This should always be BaseGridCellSize * 2^CurrentGridPower
+	double CurrentCellSize;
+};
+
+}}//end namespace UE::Geometry
diff --git a/Engine/Plugins/Experimental/MeshModelingToolsetExp/Source/ModelingOperators/Private/CompositionOps/CubeGridBooleanOp.cpp b/Engine/Plugins/Experimental/MeshModelingToolsetExp/Source/ModelingOperators/Private/CompositionOps/CubeGridBooleanOp.cpp
new file mode 100644
--- /dev/null
+++ b/Engine/Plugins/Experimental/MeshModelingToolsetExp/Source/ModelingOperators/Private/CompositionOps/CubeGridBooleanOp.cpp
@@ -0,0 +1,472 @@
+// Copyright Epic Games, Inc. All Rights Reserved.
+
+#include "CompositionOps/CubeGridBooleanOp.h"
+
+#include "DynamicMesh/MeshTransforms.h"
+#include "DynamicMeshEditor.h"
+#include "Generators/GridBoxMeshGenerator.h"
+#include "Generators/MeshShapeGenerator.h"
+#include "MeshSimplification.h"
+#include "Operations/MeshBoolean.h"
+
+#include "ExplicitUseGeometryMathTypes.h"		// using UE::Geometry::(math types)
+using namespace UE::Geometry;
+
+namespace CubeGridBooleanOpLocals
+{
+
+/**
+ * Generator for creating ramps/pyramids/pyramid cutaways in the bounds of a box. The "welded"
+ * refers to welding vertices along the Z axis to create those shapes. See the comment in
+ * FCubeGridBooleanOp::FCornerInfo for more information.
+ */
+class FWeldedMinimalBoxMeshGenerator : public FMeshShapeGenerator
+{
+public:
+	FOrientedBox3d Box;
+	FCubeGridBooleanOp::FCornerInfo CornerInfo;
+	bool bCrosswiseDiagonal = false;
+		
+	virtual FMeshShapeGenerator& Generate() override
+	{
+		// There is a particular odd case that we have to guard against. If one set of diagonal
+		// corners is welded and the second set is unwelded, and if the chosen triangulation
+		// diagonal connects the welded vertices, then the edge between these verts would normally
+		// have four connected triangles, since both top and bottom use it. The two solutions to
+		// this are (a) triangulate the bottom differently, which is fine since it's flat, or
+		// (b) duplicate the welded verts, essentially building two pyramids with triangular bases.
+		// While option b is a little more hassle, it seems like the one that may result in
+		// cleaner geometry, so that is what we do. We do it by temporarily welding the raised
+		// verts to create two independent concatenated pyramids.
+
+		int FaceIndexWeldingToModify = -1;
+		if (bCrosswiseDiagonal)
+		{
+			if (CornerInfo.WeldedAtBase[0] && CornerInfo.WeldedAtBase[2] 
+				&& !CornerInfo.WeldedAtBase[1] && !CornerInfo.WeldedAtBase[3])
+			{
+				FaceIndexWeldingToModify = 1;
+			}
+			else if (CornerInfo.WeldedAtBase[1] && CornerInfo.WeldedAtBase[3]
+				&& !CornerInfo.WeldedAtBase[0] && !CornerInfo.WeldedAtBase[2])
+			{
+				FaceIndexWeldingToModify = 0;
+			}
+		}
+
+		if (FaceIndexWeldingToModify >= 0)
+		{
+			CornerInfo.WeldedAtBase[FaceIndexWeldingToModify] = true;
+			AppendUsingCurrentSettings();
+			Swap(CornerInfo.WeldedAtBase[FaceIndexWeldingToModify], CornerInfo.WeldedAtBase[FaceIndexWeldingToModify + 2]);
+			AppendUsingCurrentSettings();
+			CornerInfo.WeldedAtBase[FaceIndexWeldingToModify + 2] = false;
+		}
+		else
+		{
+			AppendUsingCurrentSettings();
+		}
+
+		return *this;
+	}
+
+protected:
+
+	// This is basically the full generation function, but importantly it does all of the
+	// tri/element allocations using append, which lets us use it more than once (see above)
+	void AppendUsingCurrentSettings()
+	{
+		// Mapping from corner indices (like in TOrientedBox3::GetCorner())
+		int CornerToVert[8] = { -1,-1,-1,-1,-1,-1,-1,-1 };
+
+		int GroupID = 0;
+
+		// These actually get reinitialized at during each face traversal, and are out here as globals
+		// to the lambdas we use (to avoid having to pass tons of args). They are indexed by 0-3 
+		// counterclockwise from the bottom left corner, like IndexUtil::BoxFacesUV
+		int UVIndices[4];
+		int NormalIndices[4];
+		int FaceCornerToVert[4];
+		double FaceWidth;
+		double FaceHeight;
+
+		auto MakeUVsAndNormals = [this, &UVIndices, &NormalIndices, &FaceCornerToVert, &FaceWidth, &FaceHeight]
+		(int FaceCornerIndex, const FVector3f& Normal, int ParentVert) {
+			UVIndices[FaceCornerIndex] = AppendUV(FVector2f(
+				FaceWidth * IndexUtil::BoxFacesUV[FaceCornerIndex].X,
+				FaceHeight * IndexUtil::BoxFacesUV[FaceCornerIndex].Y), ParentVert);
+
+			NormalIndices[FaceCornerIndex] = AppendNormal(Normal, ParentVert);
+			FaceCornerToVert[FaceCornerIndex] = ParentVert;
+		};
+
+		auto MakeTriangle = [this, &CornerToVert, &UVIndices, &NormalIndices, &FaceCornerToVert, &GroupID]
+		(int FaceCorner1, int FaceCorner2, int FaceCorner3)
+		{
+			int TriIndex = AppendTriangle(FaceCornerToVert[FaceCorner1], FaceCornerToVert[FaceCorner2], FaceCornerToVert[FaceCorner3]);
+			SetTrianglePolygon(TriIndex, GroupID);
+			SetTriangleUVs(TriIndex, UVIndices[FaceCorner1], UVIndices[FaceCorner2], UVIndices[FaceCorner3]);
+			SetTriangleNormals(TriIndex, NormalIndices[FaceCorner1], NormalIndices[FaceCorner2], NormalIndices[FaceCorner3]);
+			++TriIndex;
+		};
+
+		// Ordered in the way we iterate across the sides of the box
+		FVector3f SideFaceBoxSpaceNormals[4] = {
+			-Box.AxisX(),
+			-Box.AxisY(),
+			Box.AxisX(),
+			Box.AxisY(),
+		};
+
+		// Set up the sides
+		for (int LeftCorner = 0; LeftCorner < 4; ++LeftCorner)
+		{
+			// We're iterating over the bottom corners of the oriented box (see TOrientedBox3::GetCorner()),
+			// but due to the way these are numbered, if we are looking at the side, the bottom two verts
+			// actually decrease in corner index (wrapping around).
+			int RightCorner = (LeftCorner + 3) % 4;
+			if (CornerInfo.WeldedAtBase[LeftCorner] && CornerInfo.WeldedAtBase[RightCorner])
+			{
+				// This side is fully welded (i.e. doesn't exist)
+				continue;
+			}
+
+			// Indices of the corners up the Z axis
+			int LeftUp = LeftCorner + 4;
+			int RightUp = RightCorner + 4;
+
+			// Create the necessary verts if they don't exist yet
+			if (CornerToVert[LeftCorner] < 0)
+			{
+				CornerToVert[LeftCorner] = AppendVertex(Box.GetCorner(LeftCorner));
+			}
+			if (CornerToVert[RightCorner] < 0)
+			{
+				CornerToVert[RightCorner] = AppendVertex(Box.GetCorner(RightCorner));
+			}
+			if (!CornerInfo.WeldedAtBase[LeftCorner] && CornerToVert[LeftUp] < 0)
+			{
+				CornerToVert[LeftUp] = AppendVertex(Box.GetCorner(LeftUp));
+			}
+			if (!CornerInfo.WeldedAtBase[RightCorner] && CornerToVert[RightUp] < 0)
+			{
+				CornerToVert[RightUp] = AppendVertex(Box.GetCorner(RightUp));
+			}
+
+			FaceWidth = Box.GetCorner(LeftCorner).Distance(Box.GetCorner(RightCorner));
+			FaceHeight = Box.GetCorner(LeftCorner).Distance(Box.GetCorner(LeftUp));
+
+			// The bottom two vertices will definitely need uv/normal elements
+			MakeUVsAndNormals(0, SideFaceBoxSpaceNormals[LeftCorner], CornerToVert[LeftCorner]);
+			MakeUVsAndNormals(1, SideFaceBoxSpaceNormals[LeftCorner], CornerToVert[RightCorner]);
+
+			if (!CornerInfo.WeldedAtBase[LeftCorner])
+			{
+				MakeUVsAndNormals(3, SideFaceBoxSpaceNormals[LeftCorner], CornerToVert[LeftUp]);
+				MakeTriangle(0, 1, 3);
+				if (!CornerInfo.WeldedAtBase[RightCorner])
+				{
+					MakeUVsAndNormals(2, SideFaceBoxSpaceNormals[LeftCorner], CornerToVert[RightUp]);
+					MakeTriangle(3, 1, 2);
+				}
+			}
+			else
+			{
+				MakeUVsAndNormals(2, SideFaceBoxSpaceNormals[LeftCorner], CornerToVert[RightUp]);
+				MakeTriangle(0, 1, 2);
+			}
+
+			++GroupID;
+		}// end making side faces
+
+		// Figure out which direction we're going to triangulate the top and bottom. 
+		// In the normal case, we place the diagonal such that a single raised/lowered
+		// corner results in a pyramid with a square (rather than triangular) base, and
+		// a checkerboard raised/lowered state results in a raised ridge. In the crosswise
+		// case we flip the diagonal.
+		int DiagFaceIdx1 = (CornerInfo.WeldedAtBase[0] && CornerInfo.WeldedAtBase[2]) ? 1 : 0;
+		DiagFaceIdx1 = bCrosswiseDiagonal ? 1 - DiagFaceIdx1 : DiagFaceIdx1;
+
+		int DiagFaceIdx2 = DiagFaceIdx1 + 2;
+		int CCWFromDiag1 = (DiagFaceIdx1 + 1) % 4;
+		int CCWFromDiag2 = (DiagFaceIdx2 + 1) % 4;
+
+		// We've iterated across the sides but it's still actually possible to end up not
+		// allocating one of the vertices if we're making a rectangle-base pyramid.
+		if (CornerToVert[DiagFaceIdx1] < 0)
+		{
+			CornerToVert[DiagFaceIdx1] = AppendVertex(Box.GetCorner(DiagFaceIdx1));
+		}
+		if (CornerToVert[DiagFaceIdx2] < 0)
+		{
+			CornerToVert[DiagFaceIdx2] = AppendVertex(Box.GetCorner(DiagFaceIdx2));
+		}
+
+		// Triangulate the bottom:
+		// Prep for UVs
+		FaceWidth = Box.GetCorner(0).Distance(Box.GetCorner(1));
+		FaceHeight = Box.GetCorner(0).Distance(Box.GetCorner(3));
+
+		MakeUVsAndNormals(DiagFaceIdx1, -Box.Frame.Z(), CornerToVert[DiagFaceIdx1]);
+		MakeUVsAndNormals(DiagFaceIdx2, -Box.Frame.Z(), CornerToVert[DiagFaceIdx2]);
+
+		bool bDiagonalWelded = CornerInfo.WeldedAtBase[DiagFaceIdx1] && CornerInfo.WeldedAtBase[DiagFaceIdx2];
+		if (!(bDiagonalWelded && CornerInfo.WeldedAtBase[CCWFromDiag1]))
+		{
+			MakeUVsAndNormals(CCWFromDiag1, -Box.Frame.Z(), CornerToVert[CCWFromDiag1]);
+			MakeTriangle(DiagFaceIdx1, CCWFromDiag1, DiagFaceIdx2);
+		}
+		if (!(bDiagonalWelded && CornerInfo.WeldedAtBase[CCWFromDiag2]))
+		{
+			MakeUVsAndNormals(CCWFromDiag2, -Box.Frame.Z(), CornerToVert[CCWFromDiag2]);
+			MakeTriangle(DiagFaceIdx1, DiagFaceIdx2, CCWFromDiag2);
+		}
+		++GroupID;
+
+		// Triangulate the top. This is slightly different from everything else because
+		// this face is not oriented to the box axes
+		
+		// Our ccw order of top corners
+		int FaceIdxToOrigIdx[4] = { 5, 4, 7, 6 };
+
+		// Given 0-3 index in ccw face ordering of top, is that vert welded?
+		bool TopFaceIdxWelded[4] = { false, false, false, false };
+
+		// Given 0-3 index in ccw face ordering of top, what is the position in space
+		FVector3d CornerPositions[4];
+
+		// Fill out the above and other structures by looking at welding of bottom verts
+		for (int i = 0; i < 4; ++i)
+		{
+			int CornerIdx = FaceIdxToOrigIdx[i];
+			if (CornerInfo.WeldedAtBase[CornerIdx - 4]) // downward index in box
+			{
+				CornerToVert[CornerIdx] = CornerToVert[CornerIdx - 4];
+				CornerPositions[i] = Box.GetCorner(CornerIdx - 4);
+				TopFaceIdxWelded[i] = true;
+			}
+			else
+			{
+				CornerPositions[i] = Box.GetCorner(CornerIdx);
+			}
+			FaceCornerToVert[i] = CornerToVert[CornerIdx];
+		}
+
+		// Prep for setting up UV's. We'll still get stretching when we're not planar, 
+		// but that's probably not worth dealing with.
+		FaceWidth = FMath::Max(
+			FVector3d::Distance(CornerPositions[0], CornerPositions[1]), 
+			FVector3d::Distance(CornerPositions[2], CornerPositions[3]));
+		FaceHeight = FMath::Max(
+			FVector3d::Distance(CornerPositions[0], CornerPositions[3]),
+			FVector3d::Distance(CornerPositions[1], CornerPositions[2]));
+
+		// We're looking at things from the opposite direction, so our diagonal needs to
+		// connect the opposite face indices.
+		DiagFaceIdx1 = 1 - DiagFaceIdx1;
+		DiagFaceIdx2 = DiagFaceIdx1 + 2;
+		CCWFromDiag1 = (DiagFaceIdx1 + 1) % 4;
+		CCWFromDiag2 = (DiagFaceIdx2 + 1) % 4;
+
+		bDiagonalWelded = TopFaceIdxWelded[DiagFaceIdx1] && TopFaceIdxWelded[DiagFaceIdx2];
+
+		bool bIsPlanar = (TopFaceIdxWelded[0] == TopFaceIdxWelded[1]
+			&& TopFaceIdxWelded[2] == TopFaceIdxWelded[3])
+			|| (TopFaceIdxWelded[1] == TopFaceIdxWelded[2]
+				&& TopFaceIdxWelded[3] == TopFaceIdxWelded[0]);
+
+		bool bFirstTriExists = !(bDiagonalWelded && TopFaceIdxWelded[CCWFromDiag1]);
+		if (bFirstTriExists)
+		{
+			FVector3f AToC(CornerPositions[DiagFaceIdx2] - CornerPositions[DiagFaceIdx1]);
+			FVector3f AToB(CornerPositions[CCWFromDiag1] - CornerPositions[DiagFaceIdx1]);
+			FVector3f Normal = AToC.Cross(AToB);
+			Normal.Normalize();
+				
+			MakeUVsAndNormals(DiagFaceIdx1, Normal, FaceCornerToVert[DiagFaceIdx1]);
+			MakeUVsAndNormals(CCWFromDiag1, Normal, FaceCornerToVert[CCWFromDiag1]);
+			MakeUVsAndNormals(DiagFaceIdx2, Normal, FaceCornerToVert[DiagFaceIdx2]);
+
+			MakeTriangle(DiagFaceIdx1, CCWFromDiag1, DiagFaceIdx2);
+
+			if (!bIsPlanar)
+			{
+				++GroupID;
+			}
+		}
+
+		if (!(bDiagonalWelded && TopFaceIdxWelded[CCWFromDiag2]))
+		{
+			if (bIsPlanar)
+			{
+				// Have to make this copy because we're not allowed to pass a reference to an element
+				// in the container we're modifying.
+				MakeUVsAndNormals(CCWFromDiag2, FVector3f(Normals[NormalIndices[DiagFaceIdx1]]), FaceCornerToVert[CCWFromDiag2]);
+			}
+			else
+			{
+				FVector3d AToC(CornerPositions[CCWFromDiag2] - CornerPositions[DiagFaceIdx1]);
+				FVector3d AToB(CornerPositions[DiagFaceIdx2] - CornerPositions[DiagFaceIdx1]);
+				FVector3d Normal = AToC.Cross(AToB);
+				Normal.Normalize();
+
+				if (bFirstTriExists)
+				{
+					NormalIndices[DiagFaceIdx1] = AppendNormal(Normal, FaceCornerToVert[DiagFaceIdx1]);
+					NormalIndices[DiagFaceIdx2] = AppendNormal(Normal, FaceCornerToVert[DiagFaceIdx2]);
+				}
+				else
+				{
+					MakeUVsAndNormals(DiagFaceIdx1, Normal, FaceCornerToVert[DiagFaceIdx1]);
+					MakeUVsAndNormals(DiagFaceIdx2, Normal, FaceCornerToVert[DiagFaceIdx2]);
+				}
+				MakeUVsAndNormals(CCWFromDiag2, Normal, FaceCornerToVert[CCWFromDiag2]);
+			}
+			MakeTriangle(DiagFaceIdx1, DiagFaceIdx2, CCWFromDiag2);
+		}// end triangulating second tri
+	}
+};
+
+void GetTrisWithChangedPositionsAndConnectivity(
+	const FDynamicMesh3& StartMeshIn, const FDynamicMesh3& EndMeshIn, TArray<int32>& TidsOut)
+{
+	// The boolean operations don't currently keep track of the triangles they modify. So instead,
+	// we do our own pass and check vertex positions and triangles to see which changed.
+	TSet<int32> ChangedVids;
+	for (int32 Vid : StartMeshIn.VertexIndicesItr())
+	{
+		if (!EndMeshIn.IsVertex(Vid) || StartMeshIn.GetVertex(Vid) != EndMeshIn.GetVertex(Vid))
+		{
+			ChangedVids.Add(Vid);
+		}
+	}
+	for (int32 Tid : StartMeshIn.TriangleIndicesItr())
+	{
+		if (!EndMeshIn.IsTriangle(Tid))
+		{
+			TidsOut.Add(Tid);
+			continue;
+		}
+
+		FIndex3i TriVids = StartMeshIn.GetTriangle(Tid);
+		if (TriVids != EndMeshIn.GetTriangle(Tid)
+			|| ChangedVids.Contains(TriVids.A)
+			|| ChangedVids.Contains(TriVids.B)
+			|| ChangedVids.Contains(TriVids.C))
+		{
+			TidsOut.Add(Tid);
+		}
+	}
+}//end FWeldedMinimalBoxMeshGenerator
+}//end CubeGridBooleanOpLocals
+
+void FCubeGridBooleanOp::CalculateResult(FProgressCancel* Progress)
+{
+	using namespace CubeGridBooleanOpLocals;
+
+	if (Progress && Progress->Cancelled())
+	{
+		return;
+	}
+
+	ResultMesh->Copy(*InputMesh);
+	if (!ResultMesh->HasTriangleGroups())
+	{
+		ResultMesh->EnableTriangleGroups();
+	}
+
+	if (Progress && Progress->Cancelled())
+	{
+		return;
+	}
+
+	TUniquePtr<FDynamicMesh3> OpMesh;
+	if (CornerInfo)
+	{
+		if (bSubtract)
+		{
+			// Recall that for subtracting a welded box mesh, we pick a different 
+			// portion of the box volume as our operator.
+			// However, to be able to use the same code, we can mirror the box in
+			// such a way that the 0-3 indices are at the selection plane, and the
+			// z axis points in the direction of the push (because pushing a corner
+			// down is equivalent to raising it up and then mirroring it down).
+			// We mirror twice to keep the frame handedness, then adjust the
+			// welded flags to point to the correct verts.
+			WorldBox.Frame = FFrame3d(WorldBox.Center(),
+				-WorldBox.AxisX(),
+				WorldBox.AxisY(),
+				-WorldBox.AxisZ());
+			Swap(CornerInfo->WeldedAtBase[0], CornerInfo->WeldedAtBase[1]);
+			Swap(CornerInfo->WeldedAtBase[2], CornerInfo->WeldedAtBase[3]);
+		}
+
+		FWeldedMinimalBoxMeshGenerator Generator;
+		Generator.Box = WorldBox;
+		Generator.CornerInfo = *CornerInfo;
+		Generator.bCrosswiseDiagonal = bCrosswiseDiagonal;
+
+		OpMesh = MakeUnique<FDynamicMesh3>(&Generator.Generate());
+	}
+	else
+	{
+		// Create the box that will be added/subtracted
+		FGridBoxMeshGenerator BoxMeshGenerator;
+		BoxMeshGenerator.Box = WorldBox;
+		BoxMeshGenerator.EdgeVertices = FIndex3i(2, 2, 2);
+		BoxMeshGenerator.bPolygroupPerQuad = true;
+
+		OpMesh = MakeUnique<FDynamicMesh3>(&BoxMeshGenerator.Generate());
+	}
+
+	// Rescale UV's to world space
+	FDynamicMeshEditor Editor(OpMesh.Get());
+	float WorldUnitsInMetersFactor = .01;
+	Editor.RescaleAttributeUVs(WorldUnitsInMetersFactor, true);
+
+	if (Progress && Progress->Cancelled())
+	{
+		return;
+	}
+	
+	// Perform the boolean operation.
+	FMeshBoolean MeshBoolean(ResultMesh.Get(), InputTransform,
+		OpMesh.Get(), FTransform3d::Identity(), ResultMesh.Get(),
+		bSubtract ? FMeshBoolean::EBooleanOp::Difference : FMeshBoolean::EBooleanOp::Union);
+
+	MeshBoolean.bPutResultInInputSpace = true;
+	MeshBoolean.bSimplifyAlongNewEdges = true;
+	MeshBoolean.Compute();
+
+	if (Progress && Progress->Cancelled())
+	{
+		return;
+	}
+
+	// Adjust the transform
+	FTransform3d NewTransform;
+	if (bKeepInputTransform)
+	{
+		ResultTransform = InputTransform;
+	}
+	else
+	{
+		// Set the transform to be in the center.
+		FVector3d Center = ResultMesh->GetBounds().Center();
+		ResultTransform = FTransform3d(Center);
+	}
+	MeshTransforms::ApplyTransformInverse(*ResultMesh, InputTransform);
+
+	// TODO: Is it worth trying to fix holes like BooleanMeshesOp tries to?
+
+	if (Progress && Progress->Cancelled())
+	{
+		return;
+	}
+
+	if (bTrackChangedTids)
+	{
+		ChangedTids = MakeShared<TArray<int32>, ESPMode::ThreadSafe>();
+		GetTrisWithChangedPositionsAndConnectivity(*InputMesh, *ResultMesh, *ChangedTids);
+	}
+}
diff --git a/Engine/Plugins/Experimental/MeshModelingToolsetExp/Source/ModelingOperators/Public/CompositionOps/CubeGridBooleanOp.h b/Engine/Plugins/Experimental/MeshModelingToolsetExp/Source/ModelingOperators/Public/CompositionOps/CubeGridBooleanOp.h
new file mode 100644
--- /dev/null
+++ b/Engine/Plugins/Experimental/MeshModelingToolsetExp/Source/ModelingOperators/Public/CompositionOps/CubeGridBooleanOp.h
@@ -0,0 +1,74 @@
+// Copyright Epic Games, Inc. All Rights Reserved.
+
+#pragma once
+
+#include "CoreMinimal.h"
+
+#include "ModelingOperators.h"
+#include "OrientedBoxTypes.h"
+
+class FProgressCancel;
+
+namespace UE {
+namespace Geometry {
+
+/**
+ * Operator meant to be used with UCubeGridTool that adds or subtracts a box to/from
+ * an input mesh, with options to weld the box's corners to make ramps/corners
+ */
+class MODELINGOPERATORS_API FCubeGridBooleanOp : public FDynamicMeshOperator
+{
+
+public:
+
+	virtual ~FCubeGridBooleanOp() {}
+
+	// Inputs:
+
+	// The mesh to add to or subtract from
+	TSharedPtr<const FDynamicMesh3, ESPMode::ThreadSafe> InputMesh;
+	FTransform3d InputTransform = FTransform3d::Identity();
+
+	// Box from which to generate the second mesh
+	FOrientedBox3d WorldBox;
+
+	// Optional: information to weld corners along the Z axis to create ramps/pyramids/pyramid cutaways
+	struct FCornerInfo
+	{
+		// "Base" refers to the 0-3 indexed corners in FOrientedBox3d, and the flags here are
+		// true if the vertex is welded to its Z axis neighbor. Note that we have a choice of
+		// which of the two resulting sides of the box volume form our operator mesh. In the
+		// case of addition, we choose the lower side on the Z axis (i.e., non-welded verts
+		// stick up), while in subtraction, we choose the higher side on the Z axis (i.e.
+		// welded verts stick down and therefore cut away the input mesh). This lines up with
+		// how a user uses this operation- they select some corners, and either pull them out
+		// or push them in.
+		bool WeldedAtBase[4] = { false, false, false, false };
+	};
+	TSharedPtr<FCornerInfo> CornerInfo;
+
+	// Determines whether the box is added or subtracted from the mesh. Also has an effect on
+	// how we treat operators with welded corners, see above.
+	bool bSubtract = false;
+
+	// When true, mesh is constructed such that ResultTransform is InputTransform. When false,
+	// ResultTransform is based off of the centroid of the result.
+	bool bKeepInputTransform = false;
+
+	// Only relevant when CornerInfo is used. When true, diagonal will generally prefer to 
+	// lie flat across the non-planar top. This determines, for instance, whether a pulled
+	// corner results in a pyramid with three faces (if true) or four (if false).
+	bool bCrosswiseDiagonal = false;
+
+	// Outputs:
+
+	// Will be reset to a container storing tids whose positions or connectivity changed. Only filled
+	// if bTrackChangedTids is true.
+	TSharedPtr<TArray<int32>, ESPMode::ThreadSafe> ChangedTids;
+	bool bTrackChangedTids = false;
+
+	virtual void CalculateResult(FProgressCancel* Progress) override;
+};
+
+
+}} // end namespace UE::Geometry
diff --git a/Engine/Source/Runtime/GeometryCore/Private/Spatial/GeometrySet3.cpp b/Engine/Source/Runtime/GeometryCore/Private/Spatial/GeometrySet3.cpp
--- a/Engine/Source/Runtime/GeometryCore/Private/Spatial/GeometrySet3.cpp
+++ b/Engine/Source/Runtime/GeometryCore/Private/Spatial/GeometrySet3.cpp
@@ -138,15 +138,15 @@ bool FGeometrySet3::CollectPointsNearRay(const FRay3d& Ray, TArray<FNearest>& Re
 	// This is used later to see if we added any new results.
 	int OriginalResultsLength = ResultsOut.Num();
 
-	// To avoid locking in the ParallelFor below, we let all curves have space for their
+	// To avoid locking in the ParallelFor below, we let all points have space for their
 	// own output even if their point turns out not to be within tolerance.
+	int NumPoints = Points.Num();
 	TArray<FNearest> NonCompactResults;
-	NonCompactResults.AddUninitialized(Curves.Num());
+	NonCompactResults.AddUninitialized(NumPoints);
 	TArray<bool> WithinTolerance;
-	WithinTolerance.AddZeroed(Curves.Num());
+	WithinTolerance.AddZeroed(NumPoints);
 
 	// Check distances to points in parallel
-	int NumPoints = Points.Num();
 	ParallelFor(NumPoints, [&](int pi)
 	{
 		const FPoint& Point = Points[pi];
diff --git a/Engine/Source/Runtime/GeometryCore/Public/Generators/MeshShapeGenerator.h b/Engine/Source/Runtime/GeometryCore/Public/Generators/MeshShapeGenerator.h
--- a/Engine/Source/Runtime/GeometryCore/Public/Generators/MeshShapeGenerator.h
+++ b/Engine/Source/Runtime/GeometryCore/Public/Generators/MeshShapeGenerator.h
@@ -191,6 +191,9 @@ public:
 	inline int AppendTriangle(int A, int B, int C)
 	{
 		Triangles.Add( bReverseOrientation ? FIndex3i(C, B, A) : FIndex3i(A, B, C) );
+		TriangleUVs.SetNum(Triangles.Num());
+		TriangleNormals.SetNum(Triangles.Num());
+		TrianglePolygonIDs.SetNum(Triangles.Num());
 		return Triangles.Num() - 1;
 	}
 
diff --git a/Engine/Source/Runtime/InteractiveToolsFramework/Public/BaseBehaviors/SingleClickBehavior.h b/Engine/Source/Runtime/InteractiveToolsFramework/Public/BaseBehaviors/SingleClickBehavior.h
--- a/Engine/Source/Runtime/InteractiveToolsFramework/Public/BaseBehaviors/SingleClickBehavior.h
+++ b/Engine/Source/Runtime/InteractiveToolsFramework/Public/BaseBehaviors/SingleClickBehavior.h
@@ -71,3 +71,41 @@ protected:
 	virtual void Clicked(const FInputDeviceState& Input, const FInputCaptureData& Data);
 };
 
+/**
+ * An implementation of USingleClickInputBehavior that also implements IClickBehaviorTarget directly, via a set 
+ * of local lambda functions. To use/customize this class, the client replaces the lambda functions with their own.
+ * This avoids having to create a separate IClickBehaviorTarget implementation for trivial use-cases.
+ */
+UCLASS()
+class INTERACTIVETOOLSFRAMEWORK_API ULocalSingleClickInputBehavior : public USingleClickInputBehavior, public IClickBehaviorTarget
+{
+	GENERATED_BODY()
+protected:
+	using USingleClickInputBehavior::Initialize;
+
+public:
+	/** Call this to initialize the class */
+	virtual void Initialize()
+	{
+		this->Initialize(this);
+	}
+
+	/** lambda implementation of IsHitByClick */
+	TUniqueFunction<FInputRayHit(const FInputDeviceRay&)> IsHitByClickFunc = [](const FInputDeviceRay& ClickPos) { return FInputRayHit(); };
+
+	/** lambda implementation of OnClicked */
+	TUniqueFunction<void(const FInputDeviceRay&)> OnClickedFunc = [](const FInputDeviceRay& ClickPos) {};
+
+public:
+	// IClickBehaviorTarget implementation
+
+	virtual FInputRayHit IsHitByClick(const FInputDeviceRay& ClickPos) override
+	{
+		return IsHitByClickFunc(ClickPos);
+	}
+
+	virtual void OnClicked(const FInputDeviceRay& ClickPos) override
+	{
+		return OnClickedFunc(ClickPos);
+	}
+};
