MeshConversion: FMeshDescriptionBuilder::AppendPolygonGroup() now initializes Material Slot Names attribute
ModelingComponents: handle situations where MeshDescription has more Sections than Material Slots. See code comments for more details.
#rnx

diff --git a/Engine/Plugins/Experimental/MeshModelingToolsetExp/Source/ModelingComponentsEditorOnly/Private/AssetUtils/CreateStaticMeshUtil.cpp b/Engine/Plugins/Experimental/MeshModelingToolsetExp/Source/ModelingComponentsEditorOnly/Private/AssetUtils/CreateStaticMeshUtil.cpp
--- a/Engine/Plugins/Experimental/MeshModelingToolsetExp/Source/ModelingComponentsEditorOnly/Private/AssetUtils/CreateStaticMeshUtil.cpp
+++ b/Engine/Plugins/Experimental/MeshModelingToolsetExp/Source/ModelingComponentsEditorOnly/Private/AssetUtils/CreateStaticMeshUtil.cpp
@@ -18,7 +18,6 @@ using namespace UE::AssetUtils;
 using namespace UE::Geometry;
 
 
-
 UE::AssetUtils::ECreateStaticMeshResult UE::AssetUtils::CreateStaticMeshAsset(
 	FStaticMeshAssetOptions& Options,
 	FStaticMeshResults& ResultsOut)
@@ -89,6 +88,9 @@ UE::AssetUtils::ECreateStaticMeshResult UE::AssetUtils::CreateStaticMeshAsset(
 		}
 	}
 
+	// determine maximum number of sections across all mesh LODs
+	int32 MaxNumSections = 0;
+
 	// if options included SourceModel meshes, copy them over
 	if (Options.SourceMeshes.MoveMeshDescriptions.Num() > 0)
 	{
@@ -98,6 +100,7 @@ UE::AssetUtils::ECreateStaticMeshResult UE::AssetUtils::CreateStaticMeshAsset(
 			{
 				FMeshDescription* Mesh = NewStaticMesh->GetMeshDescription(k);
 				*Mesh = MoveTemp(*Options.SourceMeshes.MoveMeshDescriptions[k]);
+				MaxNumSections = FMath::Max(MaxNumSections, Mesh->PolygonGroups().Num());
 				NewStaticMesh->CommitMeshDescription(k);
 			}
 		}
@@ -110,6 +113,7 @@ UE::AssetUtils::ECreateStaticMeshResult UE::AssetUtils::CreateStaticMeshAsset(
 			{
 				FMeshDescription* Mesh = NewStaticMesh->GetMeshDescription(k);
 				*Mesh = *Options.SourceMeshes.MeshDescriptions[k];
+				MaxNumSections = FMath::Max(MaxNumSections, Mesh->PolygonGroups().Num());
 				NewStaticMesh->CommitMeshDescription(k);
 			}
 		}
@@ -123,11 +127,29 @@ UE::AssetUtils::ECreateStaticMeshResult UE::AssetUtils::CreateStaticMeshAsset(
 				FMeshDescription* Mesh = NewStaticMesh->GetMeshDescription(k);
 				FDynamicMeshToMeshDescription Converter;
 				Converter.Convert(Options.SourceMeshes.DynamicMeshes[k], *Mesh, !Options.bEnableRecomputeTangents);
+				MaxNumSections = FMath::Max(MaxNumSections, Mesh->PolygonGroups().Num());
 				NewStaticMesh->CommitMeshDescription(k);
 			}
 		}
 	}
 
+	// Ensure that we have a Material Slot for each Section that exists on the mesh.
+	// This is not technically correct as the Sections on the MeshDescription (ie PolygonGroups)
+	// may reference any MaterialSlot (via the PolygonGroup::ImportedMaterialSlotName attribute),
+	// so it is valid for there to be fewer Materials than Sections. *However* if a Section does
+	// not have an ImportedMaterialSlotName, then the Section Index is used as Material Slot Index,
+	// and if that Material Slot Index does not exist, we are in a bit of an undefined state, for
+	// example the Section will not appear in the Static Mesh Editor, and various other code that
+	// does not properly handle Sections/Slots will be broken. 
+	// So, some of these extra Slots may be unused, but this is not a catastrophe. Note that the
+	// extra Slots will end up with no Material assigned.
+	int32 CurNumValidSections = UseNumMaterialSlots;
+	while (CurNumValidSections < MaxNumSections)
+	{
+		NewStaticMesh->GetStaticMaterials().Add(FStaticMaterial());
+		CurNumValidSections++;
+	}
+
 	// Nanite options
 	NewStaticMesh->NaniteSettings.bEnabled = Options.bGenerateNaniteEnabledMesh;
 	if (Options.bGenerateNaniteEnabledMesh)
diff --git a/Engine/Source/Runtime/MeshConversion/Private/MeshDescriptionBuilder.cpp b/Engine/Source/Runtime/MeshConversion/Private/MeshDescriptionBuilder.cpp
--- a/Engine/Source/Runtime/MeshConversion/Private/MeshDescriptionBuilder.cpp
+++ b/Engine/Source/Runtime/MeshConversion/Private/MeshDescriptionBuilder.cpp
@@ -39,6 +39,8 @@ void FMeshDescriptionBuilder::SetMeshDescription(FMeshDescription* Description)
 	this->InstanceTangents = Attributes.GetVertexInstanceTangents();
 	this->InstanceBiTangentSign = Attributes.GetVertexInstanceBinormalSigns();
 	this->InstanceColors = Attributes.GetVertexInstanceColors();
+	this->GroupMaterialSlotNames = Attributes.GetPolygonGroupMaterialSlotNames();
+
 }
 
 
@@ -75,9 +77,11 @@ FVertexID FMeshDescriptionBuilder::AppendVertex(const FVector& Position)
 }
 
 
-FPolygonGroupID FMeshDescriptionBuilder::AppendPolygonGroup()
+FPolygonGroupID FMeshDescriptionBuilder::AppendPolygonGroup(FName MaterialSlotName)
 {
-	return MeshDescription->CreatePolygonGroup();
+	FPolygonGroupID NewPolygonGroupID = MeshDescription->CreatePolygonGroup();
+	GroupMaterialSlotNames.Set(NewPolygonGroupID, MaterialSlotName);
+	return NewPolygonGroupID;
 }
 
 
diff --git a/Engine/Source/Runtime/MeshConversion/Public/MeshDescriptionBuilder.h b/Engine/Source/Runtime/MeshConversion/Public/MeshDescriptionBuilder.h
--- a/Engine/Source/Runtime/MeshConversion/Public/MeshDescriptionBuilder.h
+++ b/Engine/Source/Runtime/MeshConversion/Public/MeshDescriptionBuilder.h
@@ -73,9 +73,6 @@ public:
 	/** Enable per-triangle integer attribute named PolyTriGroups */
 	void EnablePolyGroups();
 
-	/** Create a new polygon group and return it's ID */
-	FPolygonGroupID AppendPolygonGroup();
-
 	/** Set the PolyTriGroups attribute value to a specific GroupID for a Triangle */
 	void SetPolyGroupID(const FTriangleID& TriangleID, int GroupID);
 
@@ -101,6 +98,15 @@ public:
 	 */
 	FPolygonID AppendPolygon(const TArray<FVertexID>& Vertices, const FPolygonGroupID& PolygonGroup);
 
+
+	/** 
+	 * Create a new MeshDescription PolygonGroup and return it's ID.
+	 * PolygonGroups are not the same as Polygroups, they essentially represent
+	 * Mesh Sections, which then reference Material Slots, etc
+	 */
+	FPolygonGroupID AppendPolygonGroup(FName MaterialSlotName = NAME_None);
+
+
 	/** Set MeshAttribute::Edge::IsHard to true for all edges */
 	void SetAllEdgesHardness(bool bHard);
 
@@ -133,4 +139,6 @@ protected:
 	TArray<FUVID> TempUVBuffer;
 
 	TPolygonAttributesRef<int> PolyGroups;
+
+	TPolygonGroupAttributesRef<FName> GroupMaterialSlotNames;
 };
\ No newline at end of file
