MeshModelingTools: Fix CIS error in cube grid tool.

#rnx

diff --git a/Engine/Plugins/Experimental/MeshModelingToolsetExp/Source/MeshModelingToolsExp/Private/CubeGridTool.cpp b/Engine/Plugins/Experimental/MeshModelingToolsetExp/Source/MeshModelingToolsExp/Private/CubeGridTool.cpp
--- a/Engine/Plugins/Experimental/MeshModelingToolsetExp/Source/MeshModelingToolsExp/Private/CubeGridTool.cpp
+++ b/Engine/Plugins/Experimental/MeshModelingToolsetExp/Source/MeshModelingToolsExp/Private/CubeGridTool.cpp
@@ -1095,11 +1095,13 @@ FInputRayHit UCubeGridTool::RayCastSelectionPlane(const UE::Geometry::FRay3d& Wo
 		FCubeGrid::DirToNormal(Selection.Direction));
 
 	FInputRayHit HitResult;
+	FVector HitPoint;
 	GizmoMath::RayPlaneIntersectionPoint(CubeGrid->GetFrame().FromFramePoint(Selection.Box.Min), Normal,
-		WorldRay.Origin, WorldRay.Direction, HitResult.bHit, HitPointOut);
+		WorldRay.Origin, WorldRay.Direction, HitResult.bHit, HitPoint);
 	
 	if (HitResult.bHit)
 	{
+		HitPointOut = HitPoint;
 		HitResult = FInputRayHit(WorldRay.Project(HitPointOut));
 	}
 	return HitResult;
