MeshModelingTools: Another CIS fix in cube grid tool in non-LWC build.

#rnx

diff --git a/Engine/Plugins/Experimental/MeshModelingToolsetExp/Source/MeshModelingToolsExp/Private/CubeGridTool.cpp b/Engine/Plugins/Experimental/MeshModelingToolsetExp/Source/MeshModelingToolsExp/Private/CubeGridTool.cpp
--- a/Engine/Plugins/Experimental/MeshModelingToolsetExp/Source/MeshModelingToolsExp/Private/CubeGridTool.cpp
+++ b/Engine/Plugins/Experimental/MeshModelingToolsetExp/Source/MeshModelingToolsExp/Private/CubeGridTool.cpp
@@ -1040,7 +1040,7 @@ void UCubeGridTool::OnCtrlMiddleClick(const FInputDeviceRay& ClickPos)
 
 	// Transform the ray to grid space and see which of the corners is closest.
 	FVector3d GridSpaceRayOrigin = CubeGrid->ToGridPoint(ClickPos.WorldRay.Origin);
-	FVector3d GridSpaceRayDirection = CubeGrid->GetFrame().ToFrameVector(ClickPos.WorldRay.Direction);
+	FVector3d GridSpaceRayDirection = CubeGrid->GetFrame().ToFrameVector(FVector3d(ClickPos.WorldRay.Direction));
 	GridSpaceRayDirection.Normalize();
 	FRay3d GizmoSpaceRay(GridSpaceRayOrigin, GridSpaceRayDirection);
 
