ModelingTools: Fix crash in BakeVertex tool.
- Fix curvature settings in BakeVertex in per-channel mode.
- Fix erroneous ensure when in Identity correspondence mode.
- Remove remnants of previous Detail mesh system.
- Minor FMeshMapBaker interface updates.

#rnx

diff --git a/Engine/Plugins/Experimental/MeshModelingToolsetExp/Source/MeshModelingToolsExp/Private/BakeMeshAttributeMapsTool.cpp b/Engine/Plugins/Experimental/MeshModelingToolsetExp/Source/MeshModelingToolsExp/Private/BakeMeshAttributeMapsTool.cpp
--- a/Engine/Plugins/Experimental/MeshModelingToolsetExp/Source/MeshModelingToolsExp/Private/BakeMeshAttributeMapsTool.cpp
+++ b/Engine/Plugins/Experimental/MeshModelingToolsetExp/Source/MeshModelingToolsExp/Private/BakeMeshAttributeMapsTool.cpp
@@ -14,8 +14,6 @@
 #include "Sampling/MeshPropertyMapEvaluator.h"
 #include "Sampling/MeshResampleImageEvaluator.h"
 
-#include "Materials/Material.h"
-#include "Materials/MaterialInstanceDynamic.h"
 #include "ImageUtils.h"
 
 #include "AssetUtils/Texture2DBuilder.h"
@@ -138,7 +136,7 @@ public:
 			{
 				TSharedPtr<FMeshNormalMapEvaluator, ESPMode::ThreadSafe> NormalEval = MakeShared<FMeshNormalMapEvaluator, ESPMode::ThreadSafe>();
 				DetailSampler.SetNormalMap(DetailMesh.Get(), IMeshBakerDetailSampler::FBakeDetailTexture(DetailMeshNormalMap.Get(), DetailMeshNormalUVLayer));
-				Baker->AddBaker(NormalEval);
+				Baker->AddEvaluator(NormalEval);
 				break;
 			}
 			case EBakeMapType::AmbientOcclusion:
@@ -179,69 +177,69 @@ public:
 					OcclusionEval->NormalSpace = FMeshOcclusionMapEvaluator::ESpace::Object;
 					break;
 				}
-				Baker->AddBaker(OcclusionEval);
+				Baker->AddEvaluator(OcclusionEval);
 				break;
 			}
 			case EBakeMapType::Curvature:
 			{
-				TSharedPtr<FMeshCurvatureMapEvaluator, ESPMode::ThreadSafe> CurvatureBaker = MakeShared<FMeshCurvatureMapEvaluator, ESPMode::ThreadSafe>();
-				CurvatureBaker->RangeScale = FMathd::Clamp(CurvatureSettings.RangeMultiplier, 0.0001, 1000.0);
-				CurvatureBaker->MinRangeScale = FMathd::Clamp(CurvatureSettings.MinRangeMultiplier, 0.0, 1.0);
-				CurvatureBaker->UseCurvatureType = (FMeshCurvatureMapEvaluator::ECurvatureType)CurvatureSettings.CurvatureType;
-				CurvatureBaker->UseColorMode = (FMeshCurvatureMapEvaluator::EColorMode)CurvatureSettings.ColorMode;
-				CurvatureBaker->UseClampMode = (FMeshCurvatureMapEvaluator::EClampMode)CurvatureSettings.ClampMode;
-				Baker->AddBaker(CurvatureBaker);
+				TSharedPtr<FMeshCurvatureMapEvaluator, ESPMode::ThreadSafe> CurvatureEval = MakeShared<FMeshCurvatureMapEvaluator, ESPMode::ThreadSafe>();
+				CurvatureEval->RangeScale = FMathd::Clamp(CurvatureSettings.RangeMultiplier, 0.0001, 1000.0);
+				CurvatureEval->MinRangeScale = FMathd::Clamp(CurvatureSettings.MinRangeMultiplier, 0.0, 1.0);
+				CurvatureEval->UseCurvatureType = (FMeshCurvatureMapEvaluator::ECurvatureType)CurvatureSettings.CurvatureType;
+				CurvatureEval->UseColorMode = (FMeshCurvatureMapEvaluator::EColorMode)CurvatureSettings.ColorMode;
+				CurvatureEval->UseClampMode = (FMeshCurvatureMapEvaluator::EClampMode)CurvatureSettings.ClampMode;
+				Baker->AddEvaluator(CurvatureEval);
 				break;
 			}
 			case EBakeMapType::NormalImage:
 			{
-				TSharedPtr<FMeshPropertyMapEvaluator, ESPMode::ThreadSafe> PropertyBaker = MakeShared<FMeshPropertyMapEvaluator, ESPMode::ThreadSafe>();
-				PropertyBaker->Property = EMeshPropertyMapType::Normal;
+				TSharedPtr<FMeshPropertyMapEvaluator, ESPMode::ThreadSafe> PropertyEval = MakeShared<FMeshPropertyMapEvaluator, ESPMode::ThreadSafe>();
+				PropertyEval->Property = EMeshPropertyMapType::Normal;
 				DetailSampler.SetNormalMap(DetailMesh.Get(), IMeshBakerDetailSampler::FBakeDetailTexture(DetailMeshNormalMap.Get(), DetailMeshNormalUVLayer));
-				Baker->AddBaker(PropertyBaker);
+				Baker->AddEvaluator(PropertyEval);
 				break;
 			}
 			case EBakeMapType::FaceNormalImage:
 			{
-				TSharedPtr<FMeshPropertyMapEvaluator, ESPMode::ThreadSafe> PropertyBaker = MakeShared<FMeshPropertyMapEvaluator, ESPMode::ThreadSafe>();
-				PropertyBaker->Property = EMeshPropertyMapType::FacetNormal;
-				Baker->AddBaker(PropertyBaker);
+				TSharedPtr<FMeshPropertyMapEvaluator, ESPMode::ThreadSafe> PropertyEval = MakeShared<FMeshPropertyMapEvaluator, ESPMode::ThreadSafe>();
+				PropertyEval->Property = EMeshPropertyMapType::FacetNormal;
+				Baker->AddEvaluator(PropertyEval);
 				break;
 			}
 			case EBakeMapType::PositionImage:
 			{
-				TSharedPtr<FMeshPropertyMapEvaluator, ESPMode::ThreadSafe> PropertyBaker = MakeShared<FMeshPropertyMapEvaluator, ESPMode::ThreadSafe>();
-				PropertyBaker->Property = EMeshPropertyMapType::Position;
-				Baker->AddBaker(PropertyBaker);
+				TSharedPtr<FMeshPropertyMapEvaluator, ESPMode::ThreadSafe> PropertyEval = MakeShared<FMeshPropertyMapEvaluator, ESPMode::ThreadSafe>();
+				PropertyEval->Property = EMeshPropertyMapType::Position;
+				Baker->AddEvaluator(PropertyEval);
 				break;
 			}
 			case EBakeMapType::MaterialID:
 			{
-				TSharedPtr<FMeshPropertyMapEvaluator, ESPMode::ThreadSafe> PropertyBaker = MakeShared<FMeshPropertyMapEvaluator, ESPMode::ThreadSafe>();
-				PropertyBaker->Property = EMeshPropertyMapType::MaterialID;
-				Baker->AddBaker(PropertyBaker);
+				TSharedPtr<FMeshPropertyMapEvaluator, ESPMode::ThreadSafe> PropertyEval = MakeShared<FMeshPropertyMapEvaluator, ESPMode::ThreadSafe>();
+				PropertyEval->Property = EMeshPropertyMapType::MaterialID;
+				Baker->AddEvaluator(PropertyEval);
 				break;
 			}
 			case EBakeMapType::VertexColorImage:
 			{
-				TSharedPtr<FMeshPropertyMapEvaluator, ESPMode::ThreadSafe> PropertyBaker = MakeShared<FMeshPropertyMapEvaluator, ESPMode::ThreadSafe>();
-				PropertyBaker->Property = EMeshPropertyMapType::VertexColor;
-				Baker->AddBaker(PropertyBaker);
+				TSharedPtr<FMeshPropertyMapEvaluator, ESPMode::ThreadSafe> PropertyEval = MakeShared<FMeshPropertyMapEvaluator, ESPMode::ThreadSafe>();
+				PropertyEval->Property = EMeshPropertyMapType::VertexColor;
+				Baker->AddEvaluator(PropertyEval);
 				break;
 			}
 			case EBakeMapType::Texture2DImage:
 			{
-				TSharedPtr<FMeshResampleImageEvaluator, ESPMode::ThreadSafe> ResampleBaker = MakeShared<FMeshResampleImageEvaluator, ESPMode::ThreadSafe>();
+				TSharedPtr<FMeshResampleImageEvaluator, ESPMode::ThreadSafe> TextureEval = MakeShared<FMeshResampleImageEvaluator, ESPMode::ThreadSafe>();
 				DetailSampler.SetColorMap(DetailMesh.Get(), IMeshBakerDetailSampler::FBakeDetailTexture(TextureImage.Get(), TextureSettings.UVLayer));
-				Baker->AddBaker(ResampleBaker);
+				Baker->AddEvaluator(TextureEval);
 				break;
 			}
 			case EBakeMapType::MultiTexture:
 			{
-				TSharedPtr<FMeshMultiResampleImageEvaluator, ESPMode::ThreadSafe> TextureBaker = MakeShared<FMeshMultiResampleImageEvaluator, ESPMode::ThreadSafe>();
-				TextureBaker->DetailUVLayer = TextureSettings.UVLayer;
-				TextureBaker->MultiTextures = MaterialToTextureImageMap;
-				Baker->AddBaker(TextureBaker);
+				TSharedPtr<FMeshMultiResampleImageEvaluator, ESPMode::ThreadSafe> TextureEval = MakeShared<FMeshMultiResampleImageEvaluator, ESPMode::ThreadSafe>();
+				TextureEval->DetailUVLayer = TextureSettings.UVLayer;
+				TextureEval->MultiTextures = MaterialToTextureImageMap;
+				Baker->AddEvaluator(TextureEval);
 				break;
 			}
 			default:
diff --git a/Engine/Plugins/Experimental/MeshModelingToolsetExp/Source/MeshModelingToolsExp/Private/BakeMeshAttributeMapsToolBase.cpp b/Engine/Plugins/Experimental/MeshModelingToolsetExp/Source/MeshModelingToolsExp/Private/BakeMeshAttributeMapsToolBase.cpp
--- a/Engine/Plugins/Experimental/MeshModelingToolsetExp/Source/MeshModelingToolsExp/Private/BakeMeshAttributeMapsToolBase.cpp
+++ b/Engine/Plugins/Experimental/MeshModelingToolsetExp/Source/MeshModelingToolsExp/Private/BakeMeshAttributeMapsToolBase.cpp
@@ -224,16 +224,16 @@ void UBakeMeshAttributeMapsToolBase::UpdatePreview(const int PreviewIdx)
 void UBakeMeshAttributeMapsToolBase::OnMapsUpdated(const TUniquePtr<UE::Geometry::FMeshMapBaker>& NewResult)
 {
 	FImageDimensions BakeDimensions = NewResult->GetDimensions();
-	const int32 NumBakers = NewResult->NumBakers();
-	for (int32 BakerIdx = 0; BakerIdx < NumBakers; ++BakerIdx)
+	const int32 NumEval = NewResult->NumEvaluators();
+	for (int32 EvalIdx = 0; EvalIdx < NumEval; ++EvalIdx)
 	{
-		FMeshMapEvaluator* Baker = NewResult->GetBaker(BakerIdx);
+		FMeshMapEvaluator* Eval = NewResult->GetEvaluator(EvalIdx);
 
-		auto UpdateCachedMap = [this, &NewResult, &BakerIdx, &BakeDimensions](const EBakeMapType BakeMapType, const FTexture2DBuilder::ETextureType TexType, const int32 ResultIdx) -> void
+		auto UpdateCachedMap = [this, &NewResult, &EvalIdx, &BakeDimensions](const EBakeMapType BakeMapType, const FTexture2DBuilder::ETextureType TexType, const int32 ResultIdx) -> void
 		{
 			FTexture2DBuilder TextureBuilder;
 			TextureBuilder.Initialize(TexType, BakeDimensions);
-			TextureBuilder.Copy(*NewResult->GetBakeResults(BakerIdx)[ResultIdx]);
+			TextureBuilder.Copy(*NewResult->GetBakeResults(EvalIdx)[ResultIdx]);
 			TextureBuilder.Commit(false);
 
 			// The CachedMap & CachedMapIndices can be thrown out of sync if updated during
@@ -244,7 +244,7 @@ void UBakeMeshAttributeMapsToolBase::OnMapsUpdated(const TUniquePtr<UE::Geometry
 			}
 		};
 
-		switch (Baker->Type())
+		switch (Eval->Type())
 		{
 		case EMeshMapEvaluatorType::Normal:
 		{
@@ -254,13 +254,13 @@ void UBakeMeshAttributeMapsToolBase::OnMapsUpdated(const TUniquePtr<UE::Geometry
 		case EMeshMapEvaluatorType::Occlusion:
 		{
 			// Occlusion Evaluator always outputs AmbientOcclusion then BentNormal.
-			FMeshOcclusionMapEvaluator* OcclusionBaker = static_cast<FMeshOcclusionMapEvaluator*>(Baker);
+			const FMeshOcclusionMapEvaluator* OcclusionEval = static_cast<FMeshOcclusionMapEvaluator*>(Eval);
 			int32 OcclusionIdx = 0;
-			if ((bool)(OcclusionBaker->OcclusionType & EMeshOcclusionMapType::AmbientOcclusion))
+			if ((bool)(OcclusionEval->OcclusionType & EMeshOcclusionMapType::AmbientOcclusion))
 			{
 				UpdateCachedMap(EBakeMapType::AmbientOcclusion, FTexture2DBuilder::ETextureType::AmbientOcclusion, OcclusionIdx++);
 			}
-			if ((bool)(OcclusionBaker->OcclusionType & EMeshOcclusionMapType::BentNormal))
+			if ((bool)(OcclusionEval->OcclusionType & EMeshOcclusionMapType::BentNormal))
 			{
 				UpdateCachedMap(EBakeMapType::BentNormal, FTexture2DBuilder::ETextureType::NormalMap, OcclusionIdx++);
 			}
@@ -273,9 +273,9 @@ void UBakeMeshAttributeMapsToolBase::OnMapsUpdated(const TUniquePtr<UE::Geometry
 		}
 		case EMeshMapEvaluatorType::Property:
 		{
-			FMeshPropertyMapEvaluator* PropertyBaker = static_cast<FMeshPropertyMapEvaluator*>(Baker);
+			const FMeshPropertyMapEvaluator* PropertyEval = static_cast<FMeshPropertyMapEvaluator*>(Eval);
 			EBakeMapType MapType = EBakeMapType::None;
-			switch (PropertyBaker->Property)
+			switch (PropertyEval->Property)
 			{
 			case EMeshPropertyMapType::Normal:
 				MapType = EBakeMapType::NormalImage;
@@ -349,7 +349,7 @@ TArray<EBakeMapType> UBakeMeshAttributeMapsToolBase::GetMapTypesArray(const int3
 }
 
 
-UE::Geometry::FTexture2DBuilder::ETextureType UBakeMeshAttributeMapsToolBase::GetTextureType(EBakeMapType MapType)
+FTexture2DBuilder::ETextureType UBakeMeshAttributeMapsToolBase::GetTextureType(EBakeMapType MapType)
 {
 	FTexture2DBuilder::ETextureType TexType = FTexture2DBuilder::ETextureType::ColorLinear;
 	switch (MapType)
diff --git a/Engine/Plugins/Experimental/MeshModelingToolsetExp/Source/MeshModelingToolsExp/Private/BakeMeshAttributeVertexTool.cpp b/Engine/Plugins/Experimental/MeshModelingToolsetExp/Source/MeshModelingToolsExp/Private/BakeMeshAttributeVertexTool.cpp
--- a/Engine/Plugins/Experimental/MeshModelingToolsetExp/Source/MeshModelingToolsExp/Private/BakeMeshAttributeVertexTool.cpp
+++ b/Engine/Plugins/Experimental/MeshModelingToolsetExp/Source/MeshModelingToolsExp/Private/BakeMeshAttributeVertexTool.cpp
@@ -69,11 +69,11 @@ class FMeshVertexBakerOp : public TGenericDataOperator<FMeshVertexBaker>
 {
 public:
 	// General bake settings
-	TSharedPtr<UE::Geometry::FDynamicMesh3, ESPMode::ThreadSafe> DetailMesh;
-	TSharedPtr<UE::Geometry::FDynamicMeshAABBTree3, ESPMode::ThreadSafe> DetailSpatial;
-	UE::Geometry::FDynamicMesh3* BaseMesh;
-	TSharedPtr<UE::Geometry::TMeshTangents<double>, ESPMode::ThreadSafe> BaseMeshTangents;
-	TUniquePtr<UE::Geometry::FMeshVertexBaker> Baker;
+	TSharedPtr<FDynamicMesh3, ESPMode::ThreadSafe> DetailMesh;
+	TSharedPtr<FDynamicMeshAABBTree3, ESPMode::ThreadSafe> DetailSpatial;
+	FDynamicMesh3* BaseMesh;
+	TSharedPtr<TMeshTangents<double>, ESPMode::ThreadSafe> BaseMeshTangents;
+	TUniquePtr<FMeshVertexBaker> Baker;
 
 	UBakeMeshAttributeVertexTool::FBakeSettings BakeSettings;
 	UBakeMeshAttributeVertexTool::FBakeColorSettings ColorSettings;
@@ -133,6 +133,15 @@ public:
 			}
 		};
 
+		auto InitCurvatureEvaluator = [this] (FMeshCurvatureMapEvaluator* CurvatureEval)
+		{
+			CurvatureEval->RangeScale = FMathd::Clamp(CurvatureSettings.RangeMultiplier, 0.0001, 1000.0);
+			CurvatureEval->MinRangeScale = FMathd::Clamp(CurvatureSettings.MinRangeMultiplier, 0.0, 1.0);
+			CurvatureEval->UseCurvatureType = static_cast<FMeshCurvatureMapEvaluator::ECurvatureType>(CurvatureSettings.CurvatureType);
+			CurvatureEval->UseColorMode = static_cast<FMeshCurvatureMapEvaluator::EColorMode>(CurvatureSettings.ColorMode);
+			CurvatureEval->UseClampMode = static_cast<FMeshCurvatureMapEvaluator::EClampMode>(CurvatureSettings.ClampMode);
+		};
+
 		if (BakeSettings.VertexMode == EBakeVertexMode::PerChannel)
 		{
 			for(int ChannelIdx = 0; ChannelIdx < 4; ++ChannelIdx)
@@ -149,6 +158,7 @@ public:
 				case EBakeVertexTypeChannel::Curvature:
 				{
 					TSharedPtr<FMeshCurvatureMapEvaluator, ESPMode::ThreadSafe> CurvatureEval = MakeShared<FMeshCurvatureMapEvaluator, ESPMode::ThreadSafe>();
+					InitCurvatureEvaluator(CurvatureEval.Get());
 					Baker->ChannelEvaluators[ChannelIdx] = CurvatureEval;
 					break;
 				}
@@ -186,56 +196,52 @@ public:
 			}
 			case EBakeVertexTypeColor::Curvature:
 			{
-				TSharedPtr<FMeshCurvatureMapEvaluator, ESPMode::ThreadSafe> CurvatureBaker = MakeShared<FMeshCurvatureMapEvaluator, ESPMode::ThreadSafe>();
-				CurvatureBaker->RangeScale = FMathd::Clamp(CurvatureSettings.RangeMultiplier, 0.0001, 1000.0);
-				CurvatureBaker->MinRangeScale = FMathd::Clamp(CurvatureSettings.MinRangeMultiplier, 0.0, 1.0);
-				CurvatureBaker->UseCurvatureType = (FMeshCurvatureMapEvaluator::ECurvatureType)CurvatureSettings.CurvatureType;
-				CurvatureBaker->UseColorMode = (FMeshCurvatureMapEvaluator::EColorMode)CurvatureSettings.ColorMode;
-				CurvatureBaker->UseClampMode = (FMeshCurvatureMapEvaluator::EClampMode)CurvatureSettings.ClampMode;
-				Baker->ColorEvaluator = CurvatureBaker;
+				TSharedPtr<FMeshCurvatureMapEvaluator, ESPMode::ThreadSafe> CurvatureEval = MakeShared<FMeshCurvatureMapEvaluator, ESPMode::ThreadSafe>();
+				InitCurvatureEvaluator(CurvatureEval.Get());
+				Baker->ColorEvaluator = CurvatureEval;
 				break;
 			}
 			case EBakeVertexTypeColor::PositionImage:
 			{
-				TSharedPtr<FMeshPropertyMapEvaluator, ESPMode::ThreadSafe> PropertyBaker = MakeShared<FMeshPropertyMapEvaluator, ESPMode::ThreadSafe>();
-				PropertyBaker->Property = EMeshPropertyMapType::Position;
-				Baker->ColorEvaluator = PropertyBaker;
+				TSharedPtr<FMeshPropertyMapEvaluator, ESPMode::ThreadSafe> PropertyEval = MakeShared<FMeshPropertyMapEvaluator, ESPMode::ThreadSafe>();
+				PropertyEval->Property = EMeshPropertyMapType::Position;
+				Baker->ColorEvaluator = PropertyEval;
 				break;
 			}
 			case EBakeVertexTypeColor::NormalImage:
 			{
-				TSharedPtr<FMeshPropertyMapEvaluator, ESPMode::ThreadSafe> PropertyBaker = MakeShared<FMeshPropertyMapEvaluator, ESPMode::ThreadSafe>();
-				PropertyBaker->Property = EMeshPropertyMapType::Normal;
-				Baker->ColorEvaluator = PropertyBaker;
+				TSharedPtr<FMeshPropertyMapEvaluator, ESPMode::ThreadSafe> PropertyEval = MakeShared<FMeshPropertyMapEvaluator, ESPMode::ThreadSafe>();
+				PropertyEval->Property = EMeshPropertyMapType::Normal;
+				Baker->ColorEvaluator = PropertyEval;
 				break;
 			}
 			case EBakeVertexTypeColor::FaceNormalImage:
 			{
-				TSharedPtr<FMeshPropertyMapEvaluator, ESPMode::ThreadSafe> PropertyBaker = MakeShared<FMeshPropertyMapEvaluator, ESPMode::ThreadSafe>();
-				PropertyBaker->Property = EMeshPropertyMapType::FacetNormal;
-				Baker->ColorEvaluator = PropertyBaker;
+				TSharedPtr<FMeshPropertyMapEvaluator, ESPMode::ThreadSafe> PropertyEval = MakeShared<FMeshPropertyMapEvaluator, ESPMode::ThreadSafe>();
+				PropertyEval->Property = EMeshPropertyMapType::FacetNormal;
+				Baker->ColorEvaluator = PropertyEval;
 				break;
 			}
 			case EBakeVertexTypeColor::MaterialID:
 			{
-				TSharedPtr<FMeshPropertyMapEvaluator, ESPMode::ThreadSafe> PropertyBaker = MakeShared<FMeshPropertyMapEvaluator, ESPMode::ThreadSafe>();
-				PropertyBaker->Property = EMeshPropertyMapType::MaterialID;
-				Baker->ColorEvaluator = PropertyBaker;
+				TSharedPtr<FMeshPropertyMapEvaluator, ESPMode::ThreadSafe> PropertyEval = MakeShared<FMeshPropertyMapEvaluator, ESPMode::ThreadSafe>();
+				PropertyEval->Property = EMeshPropertyMapType::MaterialID;
+				Baker->ColorEvaluator = PropertyEval;
 				break;
 			}
 			case EBakeVertexTypeColor::Texture2DImage:
 			{
-				TSharedPtr<FMeshResampleImageEvaluator, ESPMode::ThreadSafe> TextureBaker = MakeShared<FMeshResampleImageEvaluator, ESPMode::ThreadSafe>();
+				TSharedPtr<FMeshResampleImageEvaluator, ESPMode::ThreadSafe> TextureEval = MakeShared<FMeshResampleImageEvaluator, ESPMode::ThreadSafe>();
 				DetailSampler.SetColorMap(DetailMesh.Get(), IMeshBakerDetailSampler::FBakeDetailTexture(TextureImage.Get(), TextureSettings.UVLayer));
-				Baker->ColorEvaluator = TextureBaker;
+				Baker->ColorEvaluator = TextureEval;
 				break;
 			}
 			case EBakeVertexTypeColor::MultiTexture:
 			{
-				TSharedPtr<FMeshMultiResampleImageEvaluator, ESPMode::ThreadSafe> TextureBaker = MakeShared<FMeshMultiResampleImageEvaluator, ESPMode::ThreadSafe>();
-				TextureBaker->DetailUVLayer = TextureSettings.UVLayer;
-				TextureBaker->MultiTextures = MaterialToTextureImageMap;
-				Baker->ColorEvaluator = TextureBaker;
+				TSharedPtr<FMeshMultiResampleImageEvaluator, ESPMode::ThreadSafe> TextureEval = MakeShared<FMeshMultiResampleImageEvaluator, ESPMode::ThreadSafe>();
+				TextureEval->DetailUVLayer = TextureSettings.UVLayer;
+				TextureEval->MultiTextures = MaterialToTextureImageMap;
+				Baker->ColorEvaluator = TextureEval;
 				break;
 			}
 			}
diff --git a/Engine/Plugins/Experimental/MeshModelingToolsetExp/Source/MeshModelingToolsExp/Private/BakeMultiMeshAttributeMapsTool.cpp b/Engine/Plugins/Experimental/MeshModelingToolsetExp/Source/MeshModelingToolsExp/Private/BakeMultiMeshAttributeMapsTool.cpp
--- a/Engine/Plugins/Experimental/MeshModelingToolsetExp/Source/MeshModelingToolsExp/Private/BakeMultiMeshAttributeMapsTool.cpp
+++ b/Engine/Plugins/Experimental/MeshModelingToolsetExp/Source/MeshModelingToolsExp/Private/BakeMultiMeshAttributeMapsTool.cpp
@@ -353,15 +353,15 @@ public:
 			{
 			case EBakeMapType::TangentSpaceNormalMap:
 			{
-				TSharedPtr<FMeshNormalMapEvaluator, ESPMode::ThreadSafe> NormalBaker = MakeShared<FMeshNormalMapEvaluator, ESPMode::ThreadSafe>();
-				Baker->AddBaker(NormalBaker);
+				TSharedPtr<FMeshNormalMapEvaluator, ESPMode::ThreadSafe> NormalEval = MakeShared<FMeshNormalMapEvaluator, ESPMode::ThreadSafe>();
+				Baker->AddEvaluator(NormalEval);
 				break;	
 			}
 			case EBakeMapType::Texture2DImage:
 			{
-				TSharedPtr<FMeshResampleImageEvaluator, ESPMode::ThreadSafe> ResampleBaker = MakeShared<FMeshResampleImageEvaluator, ESPMode::ThreadSafe>();
+				TSharedPtr<FMeshResampleImageEvaluator, ESPMode::ThreadSafe> TextureEval = MakeShared<FMeshResampleImageEvaluator, ESPMode::ThreadSafe>();
 				DetailSampler.SetColorMaps(CachedMeshToColorImageMap);
-				Baker->AddBaker(ResampleBaker);
+				Baker->AddEvaluator(TextureEval);
 				break;
 			}
 			default:
diff --git a/Engine/Plugins/Experimental/MeshModelingToolsetExp/Source/MeshModelingToolsExp/Public/BakeMeshAttributeMapsToolBase.h b/Engine/Plugins/Experimental/MeshModelingToolsetExp/Source/MeshModelingToolsExp/Public/BakeMeshAttributeMapsToolBase.h
--- a/Engine/Plugins/Experimental/MeshModelingToolsetExp/Source/MeshModelingToolsExp/Public/BakeMeshAttributeMapsToolBase.h
+++ b/Engine/Plugins/Experimental/MeshModelingToolsetExp/Source/MeshModelingToolsExp/Public/BakeMeshAttributeMapsToolBase.h
@@ -45,8 +45,8 @@ ENUM_CLASS_FLAGS(EBakeMapType);
 
 
 // Only include the Occlusion bitmask rather than its components
-// (AmbientOcclusion | BentNormal). Since the Occlusion baker can
-// bake both types in a single pass, only iterating over the Occlusion
+// (AmbientOcclusion | BentNormal). Since the Occlusion evaluator can
+// evaluate both types in a single pass, only iterating over the Occlusion
 // bitmask gives direct access to both types without the need to
 // externally track if we've handled the Occlusion evaluator in a prior
 // iteration loop.
diff --git a/Engine/Plugins/Runtime/GeometryProcessing/Source/DynamicMesh/Private/Sampling/MeshMapBaker.cpp b/Engine/Plugins/Runtime/GeometryProcessing/Source/DynamicMesh/Private/Sampling/MeshMapBaker.cpp
--- a/Engine/Plugins/Runtime/GeometryProcessing/Source/DynamicMesh/Private/Sampling/MeshMapBaker.cpp
+++ b/Engine/Plugins/Runtime/GeometryProcessing/Source/DynamicMesh/Private/Sampling/MeshMapBaker.cpp
@@ -142,7 +142,15 @@ void FMeshMapBaker::Bake()
 	}
 
 	ECorrespondenceStrategy UseStrategy = this->CorrespondenceStrategy;
-	if (UseStrategy == ECorrespondenceStrategy::Identity && ensure(DetailMesh == Mesh) == false)
+	bool bIsIdentity = true;
+	int NumDetailMeshes = 0;
+	auto CheckIdentity = [Mesh, &bIsIdentity, &NumDetailMeshes](const void* DetailMesh)
+	{
+		bIsIdentity = bIsIdentity && (DetailMesh == Mesh);
+		++NumDetailMeshes;
+	};
+	DetailSampler->ProcessMeshes(CheckIdentity);
+	if (UseStrategy == ECorrespondenceStrategy::Identity && !ensure(bIsIdentity && (NumDetailMeshes == 1)))
 	{
 		// Identity strategy requires mesh to be the same. Could potentially have two copies, in which
 		// case this ensure is too conservative, but for now we will assume this
@@ -277,7 +285,7 @@ void FMeshMapBaker::Bake()
 		}
 	};
 	
-	ParallelFor(NumTiles, [this, &Tiles, &ImageTileBuffer, &GutterTexelsPerTile, &OutputQueue, &WriteToOutputBuffer, &OverwriteFn, &AddFn, &NoopFn, &WriteQueuedOutput](int32 TileIdx)
+	ParallelFor(NumTiles, [this, &Tiles, &GutterTexelsPerTile, &OutputQueue, &WriteToOutputBuffer, &OverwriteFn, &NoopFn, &WriteQueuedOutput](int32 TileIdx)
 	{
 		TRACE_CPUPROFILER_EVENT_SCOPE("FMeshMapBaker::BakeTile");
 		// Generate unpadded and padded tiles.
@@ -385,7 +393,7 @@ void FMeshMapBaker::Bake()
 
 				auto WriteToPixel = [this, &PixelBuffer, &ImageLinearIdx](TArray<int32>& BakeIds, float OneOverWeight)
 				{
-					for (int32 Idx : BakeIds)
+					for (const int32 Idx : BakeIds)
 					{
 						const FMeshMapEvaluator::FEvaluationContext& Context = BakeContexts[Idx];
 						const int32 NumData = Context.DataLayout.Num();
@@ -432,8 +440,6 @@ void FMeshMapBaker::Bake()
 				int64 GutterPixelTo;
 				int64 GutterPixelFrom;
 				Tie(GutterPixelTo, GutterPixelFrom) = GutterTexelsPerTile[TileIdx][GutterIdx];
-				FVector2i FromCoords = Dimensions.GetCoords(GutterPixelFrom);
-				FVector2i ToCoords = Dimensions.GetCoords(GutterPixelTo);
 				for (int32 Idx = 0; Idx < NumResults; Idx++)
 				{
 					BakeResults[Idx]->CopyPixel(GutterPixelFrom, GutterPixelTo);
@@ -487,7 +493,7 @@ void FMeshMapBaker::BakeSample(
 
 			const float FilterWeight = TextureFilterEval(TexelDistance);
 			PixelWeight += FilterWeight;
-			for (int32 BakeIdx : BakeIds)
+			for (const int32 BakeIdx : BakeIds)
 			{
 				const FMeshMapEvaluator::FEvaluationContext& Context = BakeContexts[BakeIdx];
 				const int32 NumData = Context.DataLayout.Num();
@@ -511,7 +517,7 @@ void FMeshMapBaker::BakeSample(
 		const int64 BufferTilePixelLinearIdx = Tile.GetIndexFromSourceCoords(ImageCoords); 
 		float* PixelBuffer = TileBuffer.GetPixel(BufferTilePixelLinearIdx);
 		
-		for (int32 Idx : BakeIds)
+		for (const int32 Idx : BakeIds)
 		{
 			const FMeshMapEvaluator::FEvaluationContext& Context = BakeContexts[Idx];
 			const int32 NumData = Context.DataLayout.Num();
@@ -533,14 +539,14 @@ void FMeshMapBaker::BakeSample(
 	OverwriteFn(BakeAccumulateLists[(int32)FMeshMapEvaluator::EAccumulateMode::Overwrite]);
 }
 
-int32 FMeshMapBaker::AddBaker(TSharedPtr<FMeshMapEvaluator> Sampler)
+int32 FMeshMapBaker::AddEvaluator(const TSharedPtr<FMeshMapEvaluator, ESPMode::ThreadSafe>& Eval)
 {
-	return Bakers.Add(Sampler);
+	return Bakers.Add(Eval);
 }
 
-FMeshMapEvaluator* FMeshMapBaker::GetBaker(const int32 BakerIdx)
+FMeshMapEvaluator* FMeshMapBaker::GetEvaluator(const int32 EvalIdx)
 {
-	return Bakers[BakerIdx].Get();
+	return Bakers[EvalIdx].Get();
 }
 
 void FMeshMapBaker::Reset()
@@ -549,15 +555,15 @@ void FMeshMapBaker::Reset()
 	BakeResults.Empty();
 }
 
-int32 FMeshMapBaker::NumBakers() const
+int32 FMeshMapBaker::NumEvaluators() const
 {
 	return Bakers.Num();
 }
 
-const TArrayView<TUniquePtr<TImageBuilder<FVector4f>>> FMeshMapBaker::GetBakeResults(const int32 BakerIdx)
+const TArrayView<TUniquePtr<TImageBuilder<FVector4f>>> FMeshMapBaker::GetBakeResults(const int32 EvalIdx)
 {
-	const int32 ResultIdx = BakeOffsets[BakerIdx];
-	const int32 NumResults = BakeOffsets[BakerIdx + 1] - ResultIdx;
+	const int32 ResultIdx = BakeOffsets[EvalIdx];
+	const int32 NumResults = BakeOffsets[EvalIdx + 1] - ResultIdx;
 	return TArrayView<TUniquePtr<TImageBuilder<FVector4f>>>(&BakeResults[ResultIdx], NumResults);
 }
 
diff --git a/Engine/Plugins/Runtime/GeometryProcessing/Source/DynamicMesh/Private/Sampling/MeshVertexBaker.cpp b/Engine/Plugins/Runtime/GeometryProcessing/Source/DynamicMesh/Private/Sampling/MeshVertexBaker.cpp
--- a/Engine/Plugins/Runtime/GeometryProcessing/Source/DynamicMesh/Private/Sampling/MeshVertexBaker.cpp
+++ b/Engine/Plugins/Runtime/GeometryProcessing/Source/DynamicMesh/Private/Sampling/MeshVertexBaker.cpp
@@ -11,7 +11,8 @@ FMeshConstantMapEvaluator FMeshVertexBaker::OneEvaluator(1.0f);
 
 void FMeshVertexBaker::Bake()
 {
-	if (!ensure(TargetMesh && TargetMesh->HasAttributes() && TargetMesh->Attributes()->HasPrimaryColors()))
+	if (!ensure(TargetMesh && TargetMesh->HasAttributes() && TargetMesh->Attributes()->HasPrimaryColors()) ||
+		!ensure(DetailSampler))
 	{
 		return;
 	}
@@ -92,7 +93,15 @@ void FMeshVertexBaker::BakeImpl(void* Data)
 	FMeshVertexBaker* Baker = static_cast<FMeshVertexBaker*>(Data);
 	
 	ECorrespondenceStrategy UseStrategy = Baker->CorrespondenceStrategy;
-	if (UseStrategy == ECorrespondenceStrategy::Identity && ensure(Baker->DetailMesh == Baker->TargetMesh) == false)
+	bool bIsIdentity = true;
+	int NumDetailMeshes = 0;
+	auto CheckIdentity = [&Baker, &bIsIdentity, &NumDetailMeshes](const void* DetailMesh)
+	{
+		bIsIdentity = bIsIdentity && (DetailMesh == Baker->TargetMesh);
+		++NumDetailMeshes;
+	};
+	Baker->DetailSampler->ProcessMeshes(CheckIdentity);
+	if (UseStrategy == ECorrespondenceStrategy::Identity && !ensure(bIsIdentity && (NumDetailMeshes == 1)))
 	{
 		// Identity strategy requires mesh to be the same. Could potentially have two copies, in which
 		// case this ensure is too conservative, but for now we will assume this
@@ -179,6 +188,11 @@ void FMeshVertexBaker::BakeImpl(void* Data)
 			FMeshMapEvaluator::FCorrespondenceSample Sample;
 			SampleSurface(ElemIdx, Sample);
 
+			if (!Sample.DetailMesh || !Baker->DetailSampler->IsTriangle(Sample.DetailMesh, Sample.DetailTriID))
+			{
+				continue;
+			}
+
 			FVector4f& Pixel = Baker->BakeResult->GetPixel(ElemIdx);
 			float* BufferPtr = &Pixel[0];
 			for (int32 BakerIdx = 0; BakerIdx < NumBakers; ++BakerIdx)
diff --git a/Engine/Plugins/Runtime/GeometryProcessing/Source/DynamicMesh/Public/Sampling/MeshBakerCommon.h b/Engine/Plugins/Runtime/GeometryProcessing/Source/DynamicMesh/Public/Sampling/MeshBakerCommon.h
--- a/Engine/Plugins/Runtime/GeometryProcessing/Source/DynamicMesh/Public/Sampling/MeshBakerCommon.h
+++ b/Engine/Plugins/Runtime/GeometryProcessing/Source/DynamicMesh/Public/Sampling/MeshBakerCommon.h
@@ -84,7 +84,6 @@ private:
 };
 
 
-// TODO: Replace virtuals with templates. Dynamic dispatch isn't required here.	
 /**
  * Base bake detail sampler class.
  *
diff --git a/Engine/Plugins/Runtime/GeometryProcessing/Source/DynamicMesh/Public/Sampling/MeshBaseBaker.h b/Engine/Plugins/Runtime/GeometryProcessing/Source/DynamicMesh/Public/Sampling/MeshBaseBaker.h
--- a/Engine/Plugins/Runtime/GeometryProcessing/Source/DynamicMesh/Public/Sampling/MeshBaseBaker.h
+++ b/Engine/Plugins/Runtime/GeometryProcessing/Source/DynamicMesh/Public/Sampling/MeshBaseBaker.h
@@ -104,10 +104,6 @@ public:
 	}
 
 protected:
-	const FDynamicMesh3* DetailMesh = nullptr;
-	const FDynamicMeshAABBTree3* DetailSpatial = nullptr;
-	TSharedPtr<FMeshTangentsd, ESPMode::ThreadSafe> DetailMeshTangents;
-	
 	const FDynamicMesh3* TargetMesh = nullptr;
 	TSharedPtr<FMeshTangentsd, ESPMode::ThreadSafe> TargetMeshTangents;
 
diff --git a/Engine/Plugins/Runtime/GeometryProcessing/Source/DynamicMesh/Public/Sampling/MeshMapBaker.h b/Engine/Plugins/Runtime/GeometryProcessing/Source/DynamicMesh/Public/Sampling/MeshMapBaker.h
--- a/Engine/Plugins/Runtime/GeometryProcessing/Source/DynamicMesh/Public/Sampling/MeshMapBaker.h
+++ b/Engine/Plugins/Runtime/GeometryProcessing/Source/DynamicMesh/Public/Sampling/MeshMapBaker.h
@@ -31,19 +31,19 @@ public:
 	void Bake();
 
 	/** Add a baker to be processed. */
-	int32 AddBaker(TSharedPtr<FMeshMapEvaluator, ESPMode::ThreadSafe> Target);
+	int32 AddEvaluator(const TSharedPtr<FMeshMapEvaluator, ESPMode::ThreadSafe>& Eval);
 
 	/** @return the evaluator at the given index. */
-	FMeshMapEvaluator* GetBaker(int32 BakerIdx);
+	FMeshMapEvaluator* GetEvaluator(int32 EvalIdx);
 
 	/** @return the number of bake evaluators on this baker. */
-	int32 NumBakers() const;
+	int32 NumEvaluators() const;
 
 	/** Reset the list of bakers. */
 	void Reset();
 
 	/** @return the bake result image for a given baker index. */
-	const TArrayView<TUniquePtr<TImageBuilder<FVector4f>>> GetBakeResults(int32 BakerIdx);
+	const TArrayView<TUniquePtr<TImageBuilder<FVector4f>>> GetBakeResults(int32 EvalIdx);
 
 	/** if this function returns true, we should abort calculation */
 	TFunction<bool(void)> CancelF = []() { return false; };
diff --git a/Engine/Plugins/Runtime/GeometryProcessing/Source/DynamicMesh/Public/Sampling/MeshVertexBaker.h b/Engine/Plugins/Runtime/GeometryProcessing/Source/DynamicMesh/Public/Sampling/MeshVertexBaker.h
--- a/Engine/Plugins/Runtime/GeometryProcessing/Source/DynamicMesh/Public/Sampling/MeshVertexBaker.h
+++ b/Engine/Plugins/Runtime/GeometryProcessing/Source/DynamicMesh/Public/Sampling/MeshVertexBaker.h
@@ -19,12 +19,15 @@ public:
 	enum class EBakeMode
 	{
 		Channel,	// Bake scalar data into color channels
-		Color		// Bake color data to RGB
+		Color		// Bake color data to RGBA
 	};
 
 	EBakeMode BakeMode = EBakeMode::Color;
 
+	/** Evaluator to use for full color bakes */
 	TSharedPtr<FMeshMapEvaluator, ESPMode::ThreadSafe> ColorEvaluator;
+
+	/** Evaluators to use for per channel (RGBA) bakes */
 	TSharedPtr<FMeshMapEvaluator, ESPMode::ThreadSafe> ChannelEvaluators[4];
 
 public:	
@@ -46,6 +49,7 @@ public:
 	//
 
 protected:
+	/** Template bake implementation. */
 	template<EBakeMode ComputeMode>
 	static void BakeImpl(void* Data);
 
@@ -70,7 +74,7 @@ protected:
 	/** The total size of the temporary float buffer for BakeSample. */
 	int32 BakeSampleBufferSize = 0;
 
-	/** */
+	/** Function pointer to internal bake implementation. */
 	using BakeFn = void(*)(void* Data);
 	BakeFn BakeInternal = nullptr;
 
