ModelingTools: add support for initial-groups constraint in GenGroups and AutoUV Tools. When enabled, groups / islands are constrained to existing groups, ie the existing groups are subdivided but not merged.
#rnx

diff --git a/Engine/Plugins/Experimental/MeshModelingToolsetExp/Source/MeshModelingToolsEditorOnly/Private/ParameterizeMeshTool.cpp b/Engine/Plugins/Experimental/MeshModelingToolsetExp/Source/MeshModelingToolsEditorOnly/Private/ParameterizeMeshTool.cpp
--- a/Engine/Plugins/Experimental/MeshModelingToolsetExp/Source/MeshModelingToolsEditorOnly/Private/ParameterizeMeshTool.cpp
+++ b/Engine/Plugins/Experimental/MeshModelingToolsetExp/Source/MeshModelingToolsEditorOnly/Private/ParameterizeMeshTool.cpp
@@ -86,6 +86,12 @@ void UParameterizeMeshTool::Setup()
 	AddToolPropertySource(PatchBuilderProperties);
 	SetToolPropertySourceEnabled(PatchBuilderProperties, false);
 
+	PolygroupLayerProperties = NewObject<UParameterizeMeshToolPatchBuilderGroupLayerProperties>(this);
+	PolygroupLayerProperties->RestoreProperties(this, TEXT("ModelingUVTools") );
+	PolygroupLayerProperties->InitializeGroupLayers(InputMesh.Get());
+	AddToolPropertySource(PolygroupLayerProperties);
+	SetToolPropertySourceEnabled(PolygroupLayerProperties, false);
+
 
 	MaterialSettings = NewObject<UExistingMeshMaterialProperties>(this);
 	MaterialSettings->MaterialMode = ESetMeshMaterialMode::Checkerboard;
@@ -121,7 +127,8 @@ void UParameterizeMeshTool::OnPropertyModified(UObject* PropertySet, FProperty*
 		|| PropertySet == Settings 
 		|| PropertySet == UVAtlasProperties  
 		|| PropertySet == XAtlasProperties  
-		|| PropertySet == PatchBuilderProperties )
+		|| PropertySet == PatchBuilderProperties
+		|| PropertySet == PolygroupLayerProperties)
 	{
 		Preview->InvalidateResult();
 	}
@@ -136,6 +143,7 @@ void UParameterizeMeshTool::OnMethodTypeChanged()
 	SetToolPropertySourceEnabled(UVAtlasProperties, Settings->Method == EParameterizeMeshUVMethod::UVAtlas);
 	SetToolPropertySourceEnabled(XAtlasProperties, Settings->Method == EParameterizeMeshUVMethod::XAtlas);
 	SetToolPropertySourceEnabled(PatchBuilderProperties, Settings->Method == EParameterizeMeshUVMethod::PatchBuilder);
+	SetToolPropertySourceEnabled(PolygroupLayerProperties, Settings->Method == EParameterizeMeshUVMethod::PatchBuilder);
 
 	Preview->InvalidateResult();
 }
@@ -155,6 +163,7 @@ void UParameterizeMeshTool::Shutdown(EToolShutdownType ShutdownType)
 	UVAtlasProperties->SaveProperties(this);
 	XAtlasProperties->SaveProperties(this);
 	PatchBuilderProperties->SaveProperties(this);
+	PolygroupLayerProperties->SaveProperties(this, TEXT("ModelingUVTools"));
 
 	FDynamicMeshOpResult Result = Preview->Shutdown();
 	
@@ -215,6 +224,11 @@ TUniquePtr<FDynamicMeshOperator> UParameterizeMeshTool::MakeNewOperator()
 
 	// patchbuilder options
 	ParameterizeMeshOp->InitialPatchCount = PatchBuilderProperties->InitialPatches;
+	ParameterizeMeshOp->bRespectInputGroups = PolygroupLayerProperties->bConstrainToPolygroups;
+	if (PolygroupLayerProperties->bConstrainToPolygroups)
+	{
+		ParameterizeMeshOp->InputGroupLayer = PolygroupLayerProperties->GetSelectedLayer(*InputMesh);
+	}
 	ParameterizeMeshOp->PatchCurvatureAlignmentWeight = PatchBuilderProperties->CurvatureAlignment;
 	ParameterizeMeshOp->PatchMergingMetricThresh = PatchBuilderProperties->MergingThreshold;
 	ParameterizeMeshOp->PatchMergingAngleThresh = PatchBuilderProperties->MaxAngleDeviation;
diff --git a/Engine/Plugins/Experimental/MeshModelingToolsetExp/Source/MeshModelingToolsEditorOnly/Public/ParameterizeMeshTool.h b/Engine/Plugins/Experimental/MeshModelingToolsetExp/Source/MeshModelingToolsEditorOnly/Public/ParameterizeMeshTool.h
--- a/Engine/Plugins/Experimental/MeshModelingToolsetExp/Source/MeshModelingToolsEditorOnly/Public/ParameterizeMeshTool.h
+++ b/Engine/Plugins/Experimental/MeshModelingToolsetExp/Source/MeshModelingToolsEditorOnly/Public/ParameterizeMeshTool.h
@@ -8,6 +8,7 @@
 #include "MeshOpPreviewHelpers.h"
 #include "Properties/MeshMaterialProperties.h"
 #include "Properties/MeshUVChannelProperties.h"
+#include "PropertySets/PolygroupLayersProperties.h"
 #include "Drawing/UVLayoutPreview.h"
 
 #include "ParameterizeMeshTool.generated.h"
@@ -21,6 +22,7 @@ class UParameterizeMeshToolUVAtlasProperties;
 class UParameterizeMeshToolXAtlasProperties;
 class UParameterizeMeshToolPatchBuilderProperties;
 
+
 UCLASS()
 class MESHMODELINGTOOLSEDITORONLY_API UParameterizeMeshToolBuilder : public USingleSelectionMeshEditingToolBuilder
 {
@@ -30,6 +32,19 @@ public:
 };
 
 
+UCLASS()
+class MESHMODELINGTOOLSEDITORONLY_API UParameterizeMeshToolPatchBuilderGroupLayerProperties : public UPolygroupLayersProperties
+{
+	GENERATED_BODY()
+public:
+
+	UPROPERTY(EditAnywhere, Category = PolygroupLayers)
+	bool bConstrainToPolygroups = false;
+
+};
+
+
+
 /**
  * UParameterizeMeshTool automatically decomposes the input mesh into charts, solves for UVs,
  * and then packs the resulting charts
@@ -71,6 +86,8 @@ protected:
 	UPROPERTY()
 	TObjectPtr<UParameterizeMeshToolPatchBuilderProperties> PatchBuilderProperties = nullptr;
 
+	UPROPERTY()
+	TObjectPtr<UParameterizeMeshToolPatchBuilderGroupLayerProperties> PolygroupLayerProperties = nullptr;
 
 	UPROPERTY()
 	TObjectPtr<UExistingMeshMaterialProperties> MaterialSettings = nullptr;
diff --git a/Engine/Plugins/Experimental/MeshModelingToolsetExp/Source/MeshModelingToolsExp/Private/ConvertToPolygonsTool.cpp b/Engine/Plugins/Experimental/MeshModelingToolsetExp/Source/MeshModelingToolsExp/Private/ConvertToPolygonsTool.cpp
--- a/Engine/Plugins/Experimental/MeshModelingToolsetExp/Source/MeshModelingToolsExp/Private/ConvertToPolygonsTool.cpp
+++ b/Engine/Plugins/Experimental/MeshModelingToolsetExp/Source/MeshModelingToolsExp/Private/ConvertToPolygonsTool.cpp
@@ -42,6 +42,7 @@ public:
 	EConvertToPolygonsMode ConversionMode = EConvertToPolygonsMode::FaceNormalDeviation;
 	double AngleTolerance = 0.1f;
 	int32 NumPoints = 10;
+	bool bSubdivideExisting = false;
 	FPolygroupsGenerator::EWeightingType WeightingType = FPolygroupsGenerator::EWeightingType::None;
 	FVector3d WeightingCoeffs = FVector3d::One();
 	int32 MinGroupSize = 2;
@@ -77,6 +78,11 @@ public:
 			Generator.FindPolygroupsFromUVIslands();
 			break;
 		}
+		case EConvertToPolygonsMode::FromNormalSeams:
+		{
+			Generator.FindPolygroupsFromHardNormalSeams();
+			break;
+		}
 		case EConvertToPolygonsMode::FromConnectedTris:
 		{
 			Generator.FindPolygroupsFromConnectedTris();
@@ -90,7 +96,15 @@ public:
 		}
 		case EConvertToPolygonsMode::FromFurthestPointSampling:
 		{
-			Generator.FindPolygroupsFromFurthestPointSampling(NumPoints, WeightingType, WeightingCoeffs);
+			if (bSubdivideExisting)
+			{
+				FPolygroupSet InputGroups(OriginalMesh.Get());
+				Generator.FindPolygroupsFromFurthestPointSampling(NumPoints, WeightingType, WeightingCoeffs, &InputGroups);
+			}
+			else
+			{
+				Generator.FindPolygroupsFromFurthestPointSampling(NumPoints, WeightingType, WeightingCoeffs, nullptr);
+			}
 			break;
 		}
 		default:
@@ -216,6 +230,7 @@ void UConvertToPolygonsTool::Setup()
 	Settings->WatchProperty(Settings->bShowGroupColors, [this](bool) { UpdateVisualization(); });
 	Settings->WatchProperty(Settings->AngleTolerance, [this](float) { OnSettingsModified(); });
 	Settings->WatchProperty(Settings->NumPoints, [this](int32) { OnSettingsModified(); });
+	Settings->WatchProperty(Settings->bSplitExisting, [this](bool) { OnSettingsModified(); });
 	Settings->WatchProperty(Settings->bNormalWeighted, [this](bool) { OnSettingsModified(); });
 	Settings->WatchProperty(Settings->NormalWeighting, [this](float) { OnSettingsModified(); });
 	Settings->WatchProperty(Settings->MinGroupSize, [this](int32) { OnSettingsModified(); });
@@ -232,6 +247,7 @@ void UConvertToPolygonsTool::UpdateOpParameters(FConvertToPolygonsOp& ConvertToP
 	ConvertToPolygonsOp.ConversionMode    = Settings->ConversionMode;
 	ConvertToPolygonsOp.AngleTolerance    = Settings->AngleTolerance;
 	ConvertToPolygonsOp.NumPoints = Settings->NumPoints;
+	ConvertToPolygonsOp.bSubdivideExisting = Settings->bSplitExisting;
 	ConvertToPolygonsOp.WeightingType = (Settings->bNormalWeighted) ? FPolygroupsGenerator::EWeightingType::NormalDeviation : FPolygroupsGenerator::EWeightingType::None;
 	ConvertToPolygonsOp.WeightingCoeffs = FVector3d(Settings->NormalWeighting, 1.0, 1.0);
 	ConvertToPolygonsOp.MinGroupSize = Settings->MinGroupSize;
diff --git a/Engine/Plugins/Experimental/MeshModelingToolsetExp/Source/MeshModelingToolsExp/Public/ConvertToPolygonsTool.h b/Engine/Plugins/Experimental/MeshModelingToolsetExp/Source/MeshModelingToolsExp/Public/ConvertToPolygonsTool.h
--- a/Engine/Plugins/Experimental/MeshModelingToolsetExp/Source/MeshModelingToolsExp/Public/ConvertToPolygonsTool.h
+++ b/Engine/Plugins/Experimental/MeshModelingToolsetExp/Source/MeshModelingToolsExp/Public/ConvertToPolygonsTool.h
@@ -36,6 +36,8 @@ enum class EConvertToPolygonsMode
 	FaceNormalDeviation UMETA(DisplayName = "Face Normal Deviation"),
 	/** Create PolyGroups based on UV Islands */
 	FromUVIslands  UMETA(DisplayName = "From UV Islands"),
+	/** Create PolyGroups based on Hard Normal Seams */
+	FromNormalSeams  UMETA(DisplayName = "From Hard Normal Seams"),
 	/** Create Polygroups based on Connected Triangles */
 	FromConnectedTris UMETA(DisplayName = "From Connected Tris"),
 	/** Create Polygroups centered on well-spaced sample points, approximating a surface Voronoi diagram */
@@ -55,13 +57,17 @@ public:
 	EConvertToPolygonsMode ConversionMode = EConvertToPolygonsMode::FaceNormalDeviation;
 
 	/** Tolerance for planarity */
-	UPROPERTY(EditAnywhere, Category = PolyGroups, meta = (UIMin = "0.001", UIMax = "20.0", ClampMin = "0.0", ClampMax = "90.0", EditCondition = "ConversionMode == EConvertToPolygonsMode::FaceNormalDeviation", EditConditionHides))
+	UPROPERTY(EditAnywhere, Category = PolyGroups, meta = (UIMin = "0.001", UIMax = "60.0", ClampMin = "0.0", ClampMax = "90.0", EditCondition = "ConversionMode == EConvertToPolygonsMode::FaceNormalDeviation", EditConditionHides))
 	float AngleTolerance = 0.1f;
 
 	/** Furthest-Point Sample count, approximately this number of polygroups will be generated */
 	UPROPERTY(EditAnywhere, Category = PolyGroups, meta = (UIMin = "1", UIMax = "100", ClampMin = "1", ClampMax = "10000", EditCondition = "ConversionMode == EConvertToPolygonsMode::FromFurthestPointSampling", EditConditionHides))
 	int32 NumPoints = 100;
 
+	/** If enabled, then furthest-point sampling happens with respect to existing Polygroups, ie the existing groups are further subdivided */
+	UPROPERTY(EditAnywhere, Category = PolyGroups, meta = (EditCondition = "ConversionMode == EConvertToPolygonsMode::FromFurthestPointSampling", EditConditionHides))
+	bool bSplitExisting = false;
+
 	/** If true, region-growing in Sampling modes will be controlled by face normals, resulting in regions with borders that are more-aligned with curvature ridges */
 	UPROPERTY(EditAnywhere, Category = PolyGroups, meta = (EditCondition = "ConversionMode == EConvertToPolygonsMode::FromFurthestPointSampling", EditConditionHides))
 	bool bNormalWeighted = true;
diff --git a/Engine/Plugins/Experimental/MeshModelingToolsetExp/Source/ModelingComponents/Private/PropertySets/PolygroupLayersProperties.cpp b/Engine/Plugins/Experimental/MeshModelingToolsetExp/Source/ModelingComponents/Private/PropertySets/PolygroupLayersProperties.cpp
--- a/Engine/Plugins/Experimental/MeshModelingToolsetExp/Source/ModelingComponents/Private/PropertySets/PolygroupLayersProperties.cpp
+++ b/Engine/Plugins/Experimental/MeshModelingToolsetExp/Source/ModelingComponents/Private/PropertySets/PolygroupLayersProperties.cpp
@@ -3,7 +3,9 @@
 #include "PropertySets/PolygroupLayersProperties.h"
 #include "DynamicMesh/DynamicMesh3.h"
 #include "DynamicMesh/DynamicMeshAttributeSet.h"
+#include "Polygroups/PolygroupUtil.h"
 
+using namespace UE::Geometry;
 
 void UPolygroupLayersProperties::InitializeGroupLayers(const FDynamicMesh3* Mesh)
 {
@@ -56,4 +58,18 @@ void UPolygroupLayersProperties::SetSelectedFromPolygroupIndex(int32 Index)
 	{
 		ActiveGroupLayer = FName(GroupLayersList[Index+1]);
 	}
+}
+
+
+UE::Geometry::FPolygroupLayer UPolygroupLayersProperties::GetSelectedLayer(const FDynamicMesh3& FromMesh)
+{
+	if (HasSelectedPolygroup() == false)
+	{
+		return FPolygroupLayer::Default();
+	}
+	else
+	{
+		int32 Index = UE::Geometry::FindPolygroupLayerIndexByName(FromMesh, ActiveGroupLayer);
+		return (Index >= 0) ? FPolygroupLayer::Layer(Index) : FPolygroupLayer::Default();
+	}
 }
\ No newline at end of file
diff --git a/Engine/Plugins/Experimental/MeshModelingToolsetExp/Source/ModelingComponents/Public/PropertySets/PolygroupLayersProperties.h b/Engine/Plugins/Experimental/MeshModelingToolsetExp/Source/ModelingComponents/Public/PropertySets/PolygroupLayersProperties.h
--- a/Engine/Plugins/Experimental/MeshModelingToolsetExp/Source/ModelingComponents/Public/PropertySets/PolygroupLayersProperties.h
+++ b/Engine/Plugins/Experimental/MeshModelingToolsetExp/Source/ModelingComponents/Public/PropertySets/PolygroupLayersProperties.h
@@ -5,6 +5,7 @@
 #include "CoreMinimal.h"
 #include "InteractiveTool.h"
 #include "GeometryBase.h"
+#include "Polygroups/PolygroupSet.h"
 #include "PolygroupLayersProperties.generated.h"
 
 PREDECLARE_USE_GEOMETRY_CLASS(FDynamicMesh3);
@@ -38,5 +39,7 @@ public:
 	bool HasSelectedPolygroup() const;
 
 	void SetSelectedFromPolygroupIndex(int32 Index);
+
+	UE::Geometry::FPolygroupLayer GetSelectedLayer(const FDynamicMesh3& FromMesh);
 };
 
diff --git a/Engine/Plugins/Experimental/MeshModelingToolsetExp/Source/ModelingOperatorsEditorOnly/Private/ParameterizationOps/ParameterizeMeshOp.cpp b/Engine/Plugins/Experimental/MeshModelingToolsetExp/Source/ModelingOperatorsEditorOnly/Private/ParameterizationOps/ParameterizeMeshOp.cpp
--- a/Engine/Plugins/Experimental/MeshModelingToolsetExp/Source/ModelingOperatorsEditorOnly/Private/ParameterizationOps/ParameterizeMeshOp.cpp
+++ b/Engine/Plugins/Experimental/MeshModelingToolsetExp/Source/ModelingOperatorsEditorOnly/Private/ParameterizationOps/ParameterizeMeshOp.cpp
@@ -362,6 +362,14 @@ bool FParameterizeMeshOp::ComputeUVs_XAtlas(FDynamicMesh3& Mesh,  TFunction<bool
 bool FParameterizeMeshOp::ComputeUVs_PatchBuilder(FDynamicMesh3& InOutMesh, FProgressCancel* ProgressCancel)
 {
 	FPatchBasedMeshUVGenerator UVGenerator;
+
+	TUniquePtr<FPolygroupSet> PolygroupConstraint;
+	if (bRespectInputGroups && InputGroupLayer.CheckExists(&InOutMesh))
+	{
+		PolygroupConstraint = MakeUnique<FPolygroupSet>(&InOutMesh, InputGroupLayer);
+		UVGenerator.GroupConstraint = PolygroupConstraint.Get();
+	}
+
 	UVGenerator.TargetPatchCount = FMath::Max(1,InitialPatchCount);
 	UVGenerator.bNormalWeightedPatches = true;
 	UVGenerator.PatchNormalWeight = FMath::Clamp(PatchCurvatureAlignmentWeight, 0.0, 999999.0);
diff --git a/Engine/Plugins/Experimental/MeshModelingToolsetExp/Source/ModelingOperatorsEditorOnly/Public/ParameterizationOps/ParameterizeMeshOp.h b/Engine/Plugins/Experimental/MeshModelingToolsetExp/Source/ModelingOperatorsEditorOnly/Public/ParameterizationOps/ParameterizeMeshOp.h
--- a/Engine/Plugins/Experimental/MeshModelingToolsetExp/Source/ModelingOperatorsEditorOnly/Public/ParameterizationOps/ParameterizeMeshOp.h
+++ b/Engine/Plugins/Experimental/MeshModelingToolsetExp/Source/ModelingOperatorsEditorOnly/Public/ParameterizationOps/ParameterizeMeshOp.h
@@ -5,6 +5,7 @@
 #include "CoreMinimal.h"
 #include "Util/ProgressCancel.h"
 #include "ModelingOperators.h"
+#include "Polygroups/PolygroupSet.h"
 #include "DynamicMesh/DynamicMeshAttributeSet.h"
 
 #include "ParameterizeMeshOp.generated.h"
@@ -47,6 +48,8 @@ public:
 
 	// PatchBuilder generation parameters
 	int32 InitialPatchCount = 100;
+	bool bRespectInputGroups = false;
+	FPolygroupLayer InputGroupLayer;
 	double PatchCurvatureAlignmentWeight = 1.0;
 	double PatchMergingMetricThresh = 1.5;
 	double PatchMergingAngleThresh = 45.0;
diff --git a/Engine/Plugins/Experimental/MeshModelingToolsetExp/Source/ModelingOperatorsEditorOnly/Public/Properties/ParameterizeMeshProperties.h b/Engine/Plugins/Experimental/MeshModelingToolsetExp/Source/ModelingOperatorsEditorOnly/Public/Properties/ParameterizeMeshProperties.h
--- a/Engine/Plugins/Experimental/MeshModelingToolsetExp/Source/ModelingOperatorsEditorOnly/Public/Properties/ParameterizeMeshProperties.h
+++ b/Engine/Plugins/Experimental/MeshModelingToolsetExp/Source/ModelingOperatorsEditorOnly/Public/Properties/ParameterizeMeshProperties.h
@@ -3,7 +3,6 @@
 #pragma once
 
 #include "CoreMinimal.h"
-
 #include "InteractiveTool.h"
 
 #include "ParameterizeMeshProperties.generated.h"
@@ -75,7 +74,7 @@ public:
 	float CurvatureAlignment = 1.0f;
 
 	/** Distortion/Stretching Threshold for island merging - larger values increase the allowable UV stretching */
-	UPROPERTY(EditAnywhere, Category = "PatchBuilder", meta = (UIMin = "1.0", UIMax = "5.0", ClampMin = "1.0", EditCondition = "bIslandMerging == true"))
+	UPROPERTY(EditAnywhere, Category = "PatchBuilder", meta = (UIMin = "1.0", UIMax = "5.0", ClampMin = "1.0"))
 	float MergingThreshold = 1.5f;
 
 	/** UV islands will not be merged if their average face normals deviate by larger than this amount */
@@ -98,4 +97,4 @@ public:
 	/** Target texture resolution used for UV packing, which determines gutter size */
 	UPROPERTY(EditAnywhere, Category = "PatchBuilder", meta = (UIMin = "64", UIMax = "2048", ClampMin = "2", ClampMax = "4096", EditCondition = bAutoPack))
 	int TextureResolution = 1024;
-};
\ No newline at end of file
+};
