MeshModelingTools: Minor tweak to avoid static analysis warning in cube grid tool.

#rnx

diff --git a/Engine/Plugins/Experimental/MeshModelingToolsetExp/Source/ModelingComponents/Public/Mechanics/CubeGrid.h b/Engine/Plugins/Experimental/MeshModelingToolsetExp/Source/ModelingComponents/Public/Mechanics/CubeGrid.h
--- a/Engine/Plugins/Experimental/MeshModelingToolsetExp/Source/ModelingComponents/Public/Mechanics/CubeGrid.h
+++ b/Engine/Plugins/Experimental/MeshModelingToolsetExp/Source/ModelingComponents/Public/Mechanics/CubeGrid.h
@@ -135,7 +135,11 @@ public:
 
 	double GetCurrentGridCellSize() const { return CurrentCellSize; }
 
-	double GetCellSize(int8 GridPower) const { return BaseGridCellSize * ((uint32)1 << GridPower); }
+	double GetCellSize(int8 GridPower) const { 
+		// This is split up into two statements to avoid a static analysis warning about
+		// shifting a 32 bit value and casting to a 64 bit value.
+		uint32 Multiplier = (uint32)1 << GridPower; 
+		return BaseGridCellSize * Multiplier; }
 
 	FVector3d ToWorldPoint(FVector3d GridPoint) const
 	{
