MeshModelingTools: Trivial cubegrid hotkey change discovered through DDC build warning.

#rnx

diff --git a/Engine/Plugins/Experimental/MeshModelingToolsetExp/Source/MeshModelingToolsExp/Private/CubeGridTool.cpp b/Engine/Plugins/Experimental/MeshModelingToolsetExp/Source/MeshModelingToolsExp/Private/CubeGridTool.cpp
--- a/Engine/Plugins/Experimental/MeshModelingToolsetExp/Source/MeshModelingToolsExp/Private/CubeGridTool.cpp
+++ b/Engine/Plugins/Experimental/MeshModelingToolsetExp/Source/MeshModelingToolsExp/Private/CubeGridTool.cpp
@@ -1883,7 +1883,7 @@ void UCubeGridTool::RegisterActions(FInteractiveToolActionSet& ActionSet)
 		TEXT("ToggleDiagonalMode"),
 		LOCTEXT("ToggleDiagonalModeAction", "Toggle Diagonal Mode"),
 		LOCTEXT("ToggleDiagonalModeTooltip", ""),
-		EModifierKey::None, EKeys::Z,
+		EModifierKey::None, EKeys::X,
 		[this]() {
 			if (Mode == EMode::Corner)
 			{
