Modeling Operators - FlareMeshOp, this adds linear flare option,  trig-simplifiied some existing math, and  replaces bSmoothEnds UI selection with a pull down that allows the user to select between the three different profile options.

diff --git a/Engine/Plugins/Experimental/GeometryScripting/Source/GeometryScriptingCore/Private/MeshDeformFunctions.cpp b/Engine/Plugins/Experimental/GeometryScripting/Source/GeometryScriptingCore/Private/MeshDeformFunctions.cpp
--- a/Engine/Plugins/Experimental/GeometryScripting/Source/GeometryScriptingCore/Private/MeshDeformFunctions.cpp
+++ b/Engine/Plugins/Experimental/GeometryScripting/Source/GeometryScriptingCore/Private/MeshDeformFunctions.cpp
@@ -135,7 +135,18 @@ UDynamicMesh* UGeometryScriptLibrary_MeshDeformFunctions::ApplyFlareWarpToMesh(
 		FlareOp.UpperBoundsInterval = FlareExtent;
 		FlareOp.FlarePercentX = FlarePercentX;
 		FlareOp.FlarePercentY = FlarePercentY;
-		FlareOp.bSmoothEnds = Options.bSmoothEnds;
+		if (Options.FlareType == EGeometryScriptFlareType::SinMode)
+		{
+			FlareOp.FlareType = FFlareMeshOp::EFlareType::SinFlare;
+		}
+		else if (Options.FlareType == EGeometryScriptFlareType::SinSquaredMode)
+		{
+			FlareOp.FlareType = FFlareMeshOp::EFlareType::SinSqrFlare;
+		}
+		else
+		{
+			FlareOp.FlareType = FFlareMeshOp::EFlareType::LinearFlare;
+		}
 
 		FlareOp.CalculateResult(nullptr);
 
diff --git a/Engine/Plugins/Experimental/GeometryScripting/Source/GeometryScriptingCore/Public/GeometryScript/MeshDeformFunctions.h b/Engine/Plugins/Experimental/GeometryScripting/Source/GeometryScriptingCore/Public/GeometryScript/MeshDeformFunctions.h
--- a/Engine/Plugins/Experimental/GeometryScripting/Source/GeometryScriptingCore/Public/GeometryScript/MeshDeformFunctions.h
+++ b/Engine/Plugins/Experimental/GeometryScripting/Source/GeometryScriptingCore/Public/GeometryScript/MeshDeformFunctions.h
@@ -48,7 +48,18 @@ public:
 	bool bBidirectional = true;
 };
 
+UENUM(BlueprintType)
+enum class EGeometryScriptFlareType : uint8
+{
+	//Displaced by sin(pi x) with x in 0 to 1
+	SinMode = 0,
+
+	//Displaced by sin(pi x)*sin(pi x) with x in 0 to 1. This provides a smooth normal transition.
+	SinSquaredMode = 1,
 
+	// Displaced by piecewise-linear trianglular mode
+	TriangleMode = 2
+};
 
 USTRUCT(BlueprintType)
 struct GEOMETRYSCRIPTINGCORE_API FGeometryScriptFlareWarpOptions
@@ -63,9 +74,9 @@ public:
 	UPROPERTY(BlueprintReadWrite, Category = Options, meta = (EditCondition = "bSymmetricExtents == false"))
 	float LowerExtent = 10;
 
-	/** */
+	/** Determines the profile used as a displacement */
 	UPROPERTY(BlueprintReadWrite, Category = Options)
-	bool bSmoothEnds = true;
+	EGeometryScriptFlareType FlareType = EGeometryScriptFlareType::SinMode;
 };
 
 
diff --git a/Engine/Plugins/Experimental/MeshModelingToolsetExp/Source/MeshModelingToolsExp/Private/MeshSpaceDeformerTool.cpp b/Engine/Plugins/Experimental/MeshModelingToolsetExp/Source/MeshModelingToolsExp/Private/MeshSpaceDeformerTool.cpp
--- a/Engine/Plugins/Experimental/MeshModelingToolsetExp/Source/MeshModelingToolsExp/Private/MeshSpaceDeformerTool.cpp
+++ b/Engine/Plugins/Experimental/MeshModelingToolsetExp/Source/MeshModelingToolsExp/Private/MeshSpaceDeformerTool.cpp
@@ -70,7 +70,18 @@ TUniquePtr<FDynamicMeshOperator> USpaceDeformerOperatorFactory::MakeNewOperator(
 			static_cast<FFlareMeshOp*>(DeformerOp.Get())->FlarePercentY = SpaceDeformerTool->Settings->FlarePercentY;
 			static_cast<FFlareMeshOp*>(DeformerOp.Get())->FlarePercentX = SpaceDeformerTool->Settings->bLockXAndYFlaring ? 
 				SpaceDeformerTool->Settings->FlarePercentY : SpaceDeformerTool->Settings->FlarePercentX;
-			static_cast<FFlareMeshOp*>(DeformerOp.Get())->bSmoothEnds = SpaceDeformerTool->Settings->bSmoothEnds;
+			if (SpaceDeformerTool->Settings->FlareProfileType == EFlareProfileType::SinMode)
+			{
+				static_cast<FFlareMeshOp*>(DeformerOp.Get())->FlareType = FFlareMeshOp::EFlareType::SinFlare;
+			}
+			else if (SpaceDeformerTool->Settings->FlareProfileType == EFlareProfileType::SinSquaredMode)
+			{
+				static_cast<FFlareMeshOp*>(DeformerOp.Get())->FlareType = FFlareMeshOp::EFlareType::SinSqrFlare;
+			}
+			else
+			{
+				static_cast<FFlareMeshOp*>(DeformerOp.Get())->FlareType = FFlareMeshOp::EFlareType::LinearFlare;
+			}
 			break;
 		}
 		case ENonlinearOperationType::Twist:
diff --git a/Engine/Plugins/Experimental/MeshModelingToolsetExp/Source/MeshModelingToolsExp/Public/MeshSpaceDeformerTool.h b/Engine/Plugins/Experimental/MeshModelingToolsetExp/Source/MeshModelingToolsExp/Public/MeshSpaceDeformerTool.h
--- a/Engine/Plugins/Experimental/MeshModelingToolsetExp/Source/MeshModelingToolsExp/Public/MeshSpaceDeformerTool.h
+++ b/Engine/Plugins/Experimental/MeshModelingToolsetExp/Source/MeshModelingToolsExp/Public/MeshSpaceDeformerTool.h
@@ -61,6 +61,19 @@ enum class  ENonlinearOperationType : int8
 	Twist
 };
 
+UENUM()
+enum class EFlareProfileType : int8
+{
+	//Displaced by sin(pi x) with x in 0 to 1
+	SinMode, 
+
+	//Displaced by sin(pi x)*sin(pi x) with x in 0 to 1. This provides a smooth normal transition.
+	SinSquaredMode,
+	
+	// Displaced by piecewise-linear trianglular mode
+	TriangleMode 
+};
+
 UCLASS()
 class MESHMODELINGTOOLSEXP_API UMeshSpaceDeformerToolProperties : public UInteractiveToolPropertySet
 {
@@ -94,6 +107,12 @@ public:
 		EditCondition = "SelectedOperationType == ENonlinearOperationType::Twist", EditConditionHides))
 	float TwistDegrees = 180;
 
+	/**
+	* Determines the profile used as a displacement
+	*/ 
+	UPROPERTY(EditAnywhere, Category = Options, meta = (EditCondition = "SelectedOperationType == ENonlinearOperationType::Flare", EditConditionHides))
+	EFlareProfileType FlareProfileType = EFlareProfileType::SinMode;
+
 	/**
 	 * Determines how much to flare perpendicular to the Z axis. When set to 100%, points are moved double the distance
 	 * away from the gizmo Z axis at the most extreme flare point. 0% does not flare at all, whereas -100% pinches all
@@ -125,12 +144,6 @@ public:
 		EditCondition = "SelectedOperationType == ENonlinearOperationType::Bend || SelectedOperationType == ENonlinearOperationType::Twist", EditConditionHides))
 	bool bLockBottom = false;
 
-	/**
-	 * When true, changes the flare/pinch function to have C1 continuity at the upper and lower bounds.
-	 */
-	UPROPERTY(EditAnywhere, Category = Options, meta = (EditCondition = "SelectedOperationType == ENonlinearOperationType::Flare", EditConditionHides))
-	bool bSmoothEnds = true;
-
 	UPROPERTY(EditAnywhere, Category = Options)
 	bool bShowOriginalMesh = true;
 
diff --git a/Engine/Plugins/Experimental/MeshModelingToolsetExp/Source/ModelingOperators/Private/MeshSpaceDeformerOps/FlareMeshOp.cpp b/Engine/Plugins/Experimental/MeshModelingToolsetExp/Source/ModelingOperators/Private/MeshSpaceDeformerOps/FlareMeshOp.cpp
--- a/Engine/Plugins/Experimental/MeshModelingToolsetExp/Source/ModelingOperators/Private/MeshSpaceDeformerOps/FlareMeshOp.cpp
+++ b/Engine/Plugins/Experimental/MeshModelingToolsetExp/Source/ModelingOperators/Private/MeshSpaceDeformerOps/FlareMeshOp.cpp
@@ -55,15 +55,19 @@ void FFlareMeshOp::CalculateResult(FProgressCancel* Progress)
 
 	FMatrix GizmoToObject = ObjectToGizmo.Inverse();
 
+
 	const double ZMin = LowerBoundsInterval;
 	const double ZMax =  UpperBoundsInterval;
 
+	const double FlareRatioX = FlarePercentX / 100.;
+	const double FlareRatioY = FlarePercentY / 100.;
+
 	if (ResultMesh->HasAttributes())
 	{
 		// Fix the normals first if they exist.
 
 		FDynamicMeshNormalOverlay* Normals = ResultMesh->Attributes()->PrimaryNormals();
-		ParallelFor(Normals->MaxElementID(), [this, Normals, &GizmoToObject, ZMin, ZMax](int32 ElID)
+		ParallelFor(Normals->MaxElementID(), [this, Normals, &GizmoToObject, ZMin, ZMax, FlareRatioX, FlareRatioY](int32 ElID)
 		{
 			if (!Normals->IsElement(ElID))
 			{
@@ -102,14 +106,6 @@ void FFlareMeshOp::CalculateResult(FProgressCancel* Progress)
 					}
 				}
 			}
-
-
-			const double T = FMath::Clamp((GizmoPos4[2] - ZMin) / (ZMax - ZMin), 0.0, 1.0);
-
-			double Rx = bSmoothEnds ? 1. + (FMath::Cos(2 * PI * T - PI) + 1) * (FlarePercentX / 200.f) // Shift cos curve up 1, right PI, scale down by 2, go from 0 to 2PI
-				: 1. + FMath::Sin(PI * T) * (FlarePercentX / 100.f);
-			double Ry = bSmoothEnds ? 1. + (FMath::Cos(2 * PI * T - PI) + 1) * (FlarePercentY / 200.f) // Shift cos curve up 1, right PI, scale down by 2, go from 0 to 2PI
-				: 1. + FMath::Sin(PI * T) * (FlarePercentY / 100.f);
 			
 			// transform normal.  Do this before changing GizmoPos
 			// To get the normal, note that the positions are transformed like this:
@@ -123,34 +119,100 @@ void FFlareMeshOp::CalculateResult(FProgressCancel* Progress)
 			// Where DRx is dRx/dz and DRy is dRy/dz
 			// Then take the transpose of the inverse of the Jacobian and multiply by the determinant (or don't, since we don't
 			// care about the length, but it's cleaner if you do).
+			double Rx = 0., Ry = 0., DRx = 0., DRy = 0.;
+			const double T = FMath::Clamp((GizmoPos4[2] - ZMin) / (ZMax - ZMin), 0.0, 1.0);
+			
+			switch (FlareType) // evaluate the requested curve and its slope 
 			{
-				double DRx = bSmoothEnds ? -FMath::Sin(2 * PI * T - PI) * (2 * PI / (ZMax - ZMin)) * (FlarePercentX / 200.f)
-					: FMath::Cos(PI * T) * (PI / (ZMax - ZMin)) * (FlarePercentX / 100.f);
-				double DRy = bSmoothEnds ? -FMath::Sin(2 * PI * T - PI) * (2 * PI / (ZMax - ZMin)) * (FlarePercentY / 200.f)
-					: FMath::Cos(PI * T) * (PI / (ZMax - ZMin)) *  (FlarePercentY / 100.f);
+				case EFlareType::LinearFlare :
+				{
+					const double LinearDisplacement = (T < 0.5) ? 2. * T : 2. * (1. - T);
+					Rx = 1. + FlareRatioX * LinearDisplacement;
+					Ry = 1. + FlareRatioY * LinearDisplacement;
+
+					double DLinearDz = (T < 0.5) ? 2. / (ZMax - ZMin) : -2. / (ZMax - ZMin);
+
+					// hack: make normal transitions between regions of discontinuous surface slope (critical points) less noticable 
+					{
+						const double HalfSmoothDist = 0.05;
+						if (FMath::Abs(T - 0.5) < HalfSmoothDist)
+						{
+							double Alpha = (T - 0.5 + HalfSmoothDist) / (2. * HalfSmoothDist);
+							DLinearDz = FMath::Lerp(2. / (ZMax - ZMin), -2. / (ZMax - ZMin), Alpha);
+						}
+						if (T < HalfSmoothDist)
+						{
+							double Alpha = (T + HalfSmoothDist) / (2. * HalfSmoothDist);
+							DLinearDz = FMath::Lerp(0. , 2. / (ZMax - ZMin), Alpha);
+						}
+						if (1. - T < HalfSmoothDist)
+						{
+							double Alpha = (T - 1. + HalfSmoothDist) / (2. * HalfSmoothDist);
+							DLinearDz = FMath::Lerp(-2. / (ZMax - ZMin), 0., Alpha);
+						}
+					}
+
+					DRx = FlareRatioX * DLinearDz;
+					DRy = FlareRatioY * DLinearDz;
+				}
+				break;
 
-				if (GizmoPos4[2] > ZMax || GizmoPos4[2] < ZMin)
+				case EFlareType::SinFlare :
 				{
-					DRx = DRy = 0.f;
+					const double SinPiT = FMath::Sin(PI * T);
+					const double CosPiT = FMath::Cos(PI * T);
+					
+					Rx = 1. + FlareRatioX * SinPiT;
+					Ry = 1. + FlareRatioY * SinPiT;
+
+					DRx = FlareRatioX * CosPiT * (PI / (ZMax - ZMin));
+					DRy = FlareRatioY * CosPiT * (PI / (ZMax - ZMin));
+
 				}
-			
-				FVector3d DstNormal(0., 0., 0.);
-				DstNormal[0] = Ry * RotatedNormal[0];
-				DstNormal[1] = Rx * RotatedNormal[1];
-				DstNormal[2] = -Ry * DRx * GizmoPos4[0] * RotatedNormal[0] - Rx * DRy * GizmoPos4[1] * RotatedNormal[1] + Rx * Ry * RotatedNormal[2];
+				break;
 
-				// rotate back to mesh space.
-				RotatedNormal = FVector3d(0, 0, 0);
-				for (int i = 0; i < 3; ++i)
+				case EFlareType::SinSqrFlare :
 				{
-					for (int j = 0; j < 3; ++j)
-					{
-						RotatedNormal[i] += ObjectToGizmo.M[j][i] * DstNormal[j];
-					}
+					const double SinPiT = FMath::Sin(PI * T);
+					const double Sin2PiT = FMath::Sin(2 * PI * T);
+
+					Rx = 1. + FlareRatioX * SinPiT * SinPiT;
+					Ry = 1. + FlareRatioY * SinPiT * SinPiT;
+
+					DRx = FlareRatioX * Sin2PiT * (PI / (ZMax - ZMin));
+					DRy = FlareRatioY * Sin2PiT * (PI / (ZMax - ZMin));
 				}
+				break;
 
+				default :
+				{
+					// should never reach this
+					check(0);
+				}
+			}
+			
+			if (GizmoPos4[2] > ZMax || GizmoPos4[2] < ZMin)
+			{
+				DRx = DRy = 0.f;
+			}
+			
+			FVector3d DstNormal(0., 0., 0.);
+			DstNormal[0] = Ry * RotatedNormal[0];
+			DstNormal[1] = Rx * RotatedNormal[1];
+			DstNormal[2] = -Ry * DRx * GizmoPos4[0] * RotatedNormal[0] - Rx * DRy * GizmoPos4[1] * RotatedNormal[1] + Rx * Ry * RotatedNormal[2];
+
+			// rotate back to mesh space.
+			RotatedNormal = FVector3d(0, 0, 0);
+			for (int i = 0; i < 3; ++i)
+			{
+				for (int j = 0; j < 3; ++j)
+				{
+					RotatedNormal[i] += ObjectToGizmo.M[j][i] * DstNormal[j];
+				}
 			}
 
+			
+
 			Normals->SetElement(ElID, FVector3f(Normalized(RotatedNormal)));
 		});
 	}
@@ -160,7 +222,7 @@ void FFlareMeshOp::CalculateResult(FProgressCancel* Progress)
 		return;
 	}
 
-	ParallelFor(ResultMesh->MaxVertexID(), [this, &GizmoToObject, ZMin, ZMax](int32 VertexID)
+	ParallelFor(ResultMesh->MaxVertexID(), [this, &GizmoToObject, ZMin, ZMax, FlareRatioX, FlareRatioY](int32 VertexID)
 	{
 		if (!ResultMesh->IsVertex(VertexID))
 		{
@@ -184,12 +246,40 @@ void FFlareMeshOp::CalculateResult(FProgressCancel* Progress)
 		
 		// Parameterize curve between ZMin and ZMax to go between 0, 1
 		double T = FMath::Clamp( (GizmoPos4[2]- ZMin) / (ZMax - ZMin), 0.0, 1.0);
+		double Rx = 0., Ry = 0.;
+		switch (FlareType)
+		{
+			case EFlareType::LinearFlare :
+			{
+				const double LinearDisplacement = (T < 0.5) ? 2. * T : 2. * (1. - T);
+				Rx = 1. + FlareRatioX * LinearDisplacement;
+				Ry = 1. + FlareRatioY * LinearDisplacement;
+			}
+			break;
+
+			case EFlareType::SinFlare :
+			{
+				const double SinPiT = FMath::Sin(PI * T);
+				Rx = 1. + FlareRatioX * SinPiT;
+				Ry = 1. + FlareRatioY * SinPiT;
+			}
+			break;
 
-		double Rx = bSmoothEnds ? 1. + (FMath::Cos(2 * PI * T - PI) + 1) * (FlarePercentX / 200.f) // Shift cos curve up 1, right PI, scale down by 2, go from 0 to 2PI
-			: 1. + FMath::Sin(PI * T) * (FlarePercentX / 100.f);
-		double Ry = bSmoothEnds ? 1. + (FMath::Cos(2 * PI * T - PI) + 1) * (FlarePercentY / 200.f) // Shift cos curve up 1, right PI, scale down by 2, go from 0 to 2PI
-			: 1. + FMath::Sin(PI * T) * (FlarePercentY / 100.f);
+			case EFlareType::SinSqrFlare :
+			{
+				const double SinPiT = FMath::Sin(PI * T);
+				Rx = 1. + FlareRatioX * SinPiT * SinPiT;
+				Ry = 1. + FlareRatioY * SinPiT * SinPiT;
+			}
+			break;
 
+			default :
+			{
+				// should never..
+				check(0);
+			}
+		}
+		
 		// 2d scale x,y values.
 		GizmoPos4[0] *= Rx;
 		GizmoPos4[1] *= Ry;
diff --git a/Engine/Plugins/Experimental/MeshModelingToolsetExp/Source/ModelingOperators/Public/SpaceDeformerOps/FlareMeshOp.h b/Engine/Plugins/Experimental/MeshModelingToolsetExp/Source/ModelingOperators/Public/SpaceDeformerOps/FlareMeshOp.h
--- a/Engine/Plugins/Experimental/MeshModelingToolsetExp/Source/ModelingOperators/Public/SpaceDeformerOps/FlareMeshOp.h
+++ b/Engine/Plugins/Experimental/MeshModelingToolsetExp/Source/ModelingOperators/Public/SpaceDeformerOps/FlareMeshOp.h
@@ -19,13 +19,15 @@ public:
 	double FlarePercentX = 100;
 	double FlarePercentY = 100;
 
-	/** 
-	 * Changes the type of flaring. When false, the flaring is the curve of sin(y) from 0 to Pi,
-	 * which makes the ends of the flaring sharp because the derivative is not 0. When true, the 
-	 * flaring is the curve of cos(y) + 1 from -Pi to Pi, which makes the ends smooth out back
-	 * into the shape.
-	 */
-	bool bSmoothEnds = false;
+	enum class EFlareType
+	{
+		SinFlare = 0,     // Flaring is the curve sin(Pi y)  y in [0, 1]
+		SinSqrFlare = 1,  // Flaring is the curve sin(Pi y)**2  y in  [0, 1]  which makes the ends smooth out back into shape
+		LinearFlare = 2   // Flaring is  2*y  y in [0, 1/2] and 2(1-y) y in (1/2, 1] 
+	};
+
+	EFlareType FlareType = EFlareType::SinFlare;
+
 protected:
 
 };
