GeometryProcessing:
- add FDynamicMeshUVEditor::CopyUVLayer()
- Move Polygroup topology checks from SubdividePolyTool to FSubdividePoly::ValidateTopology()
#rnx

diff --git a/Engine/Plugins/Experimental/MeshModelingToolsetExp/Source/MeshModelingToolsEditorOnly/Private/SubdividePolyTool.cpp b/Engine/Plugins/Experimental/MeshModelingToolsetExp/Source/MeshModelingToolsEditorOnly/Private/SubdividePolyTool.cpp
--- a/Engine/Plugins/Experimental/MeshModelingToolsetExp/Source/MeshModelingToolsEditorOnly/Private/SubdividePolyTool.cpp
+++ b/Engine/Plugins/Experimental/MeshModelingToolsetExp/Source/MeshModelingToolsEditorOnly/Private/SubdividePolyTool.cpp
@@ -68,49 +68,42 @@ USingleSelectionMeshEditingTool* USubdividePolyToolBuilder::CreateNewTool(const
 bool USubdividePolyTool::CheckGroupTopology(FText& Message)
 {
 	FGroupTopology Topo(OriginalMesh.Get(), true);
+	FSubdividePoly TempSubD(Topo, *OriginalMesh, 1);
+	FSubdividePoly::ETopologyCheckResult CheckResult = TempSubD.ValidateTopology();
 
-	if (Topo.Groups.Num() == 0)
+	if (CheckResult == FSubdividePoly::ETopologyCheckResult::NoGroups)
 	{
 		Message = LOCTEXT("NoGroupsWarning",
 						  "This object has no PolyGroups.\nUse the PolyGroups or Select Tool to assign PolyGroups.\nTool will be limited to Loop subdivision scheme.");
 		return false;
 	}
 
-	if (Topo.Groups.Num() < 2)
+	if (CheckResult == FSubdividePoly::ETopologyCheckResult::InsufficientGroups)
 	{
-		// TODO: for an open surface, use the surface boundary as a group boundary?
 		Message = LOCTEXT("SingleGroupsWarning",
 						  "This object has only one PolyGroup.\nUse the PolyGroups or Select Tool to assign PolyGroups.\nTool will be limited to Loop subdivision scheme.");
 		return false;
 	}
 
-	for (const FGroupTopology::FGroup& Group : Topo.Groups)
+	if (CheckResult == FSubdividePoly::ETopologyCheckResult::UnboundedPolygroup)
 	{
-		if (Group.Boundaries.Num() == 0)
-		{
-			// Error: Group has no boundaries, e.g. closed surface component with only one group
-			Message = LOCTEXT("NoGroupBoundaryWarning",
-							  "Found a PolyGroup with no boundaries.\nUse the PolyGroups or Select Tool to assign PolyGroups.\nTool will be limited to Loop subdivision scheme.");
-			return false;
-		}
+		Message = LOCTEXT("NoGroupBoundaryWarning",
+			"Found a PolyGroup with no boundaries.\nUse the PolyGroups or Select Tool to assign PolyGroups.\nTool will be limited to Loop subdivision scheme.");
+		return false;
+	}
 
-		if (Group.Boundaries.Num() > 1)
-		{
-			// Error: Group has multiple boundaries, e.g. nested polygon
-			Message = LOCTEXT("MultipleGroupBoundaryWarning",
-							  "Found a PolyGroup with multiple boundaries, which is not supported.\nUse the PolyGroups or Select Tool to assign PolyGroups.\nTool will be limited to Loop subdivision scheme.");
-			return false;
-		}
+	if (CheckResult == FSubdividePoly::ETopologyCheckResult::MultiBoundaryPolygroup)
+	{
+		Message = LOCTEXT("MultipleGroupBoundaryWarning",
+			"Found a PolyGroup with multiple boundaries, which is not supported.\nUse the PolyGroups or Select Tool to assign PolyGroups.\nTool will be limited to Loop subdivision scheme.");
+		return false;
+	}
 
-		for (const FGroupTopology::FGroupBoundary& Boundary : Group.Boundaries)
-		{
-			if (Boundary.GroupEdges.Num() < 3)
-			{
-				Message = LOCTEXT("DegenerateGroupPolygon",
-								  "One PolyGroup has fewer than three boundary edges.\nUse the PolyGroups or Select Tool to assign/fix PolyGroups.\nTool will be limited to Loop subdivision scheme.");
-				return false;
-			}
-		}
+	if (CheckResult == FSubdividePoly::ETopologyCheckResult::DegeneratePolygroup)
+	{
+		Message = LOCTEXT("DegenerateGroupPolygon",
+			"One PolyGroup has fewer than three boundary edges.\nUse the PolyGroups or Select Tool to assign/fix PolyGroups.\nTool will be limited to Loop subdivision scheme.");
+		return false;
 	}
 
 	return true;
diff --git a/Engine/Plugins/Experimental/MeshModelingToolsetExp/Source/ModelingComponentsEditorOnly/Private/Operations/SubdividePoly.cpp b/Engine/Plugins/Experimental/MeshModelingToolsetExp/Source/ModelingComponentsEditorOnly/Private/Operations/SubdividePoly.cpp
--- a/Engine/Plugins/Experimental/MeshModelingToolsetExp/Source/ModelingComponentsEditorOnly/Private/Operations/SubdividePoly.cpp
+++ b/Engine/Plugins/Experimental/MeshModelingToolsetExp/Source/ModelingComponentsEditorOnly/Private/Operations/SubdividePoly.cpp
@@ -686,3 +686,41 @@ bool FSubdividePoly::ComputeSubdividedMesh(FDynamicMesh3& OutMesh)
 	return true;
 }
 
+
+
+FSubdividePoly::ETopologyCheckResult FSubdividePoly::ValidateTopology()
+{
+	if (GroupTopology.Groups.Num() == 0)
+	{
+		return ETopologyCheckResult::NoGroups;
+	}
+
+	if (GroupTopology.Groups.Num() < 2)
+	{
+		return ETopologyCheckResult::InsufficientGroups;
+	}
+
+	for (const FGroupTopology::FGroup& Group : GroupTopology.Groups)
+	{
+		if (Group.Boundaries.Num() == 0)
+		{
+			return ETopologyCheckResult::UnboundedPolygroup;
+		}
+
+		if (Group.Boundaries.Num() > 1)
+		{
+			return ETopologyCheckResult::MultiBoundaryPolygroup;
+		}
+
+		for (const FGroupTopology::FGroupBoundary& Boundary : Group.Boundaries)
+		{
+			if (Boundary.GroupEdges.Num() < 3)
+			{
+				return ETopologyCheckResult::DegeneratePolygroup;
+			}
+		}
+	}
+
+	return ETopologyCheckResult::Ok;
+}
+
diff --git a/Engine/Plugins/Experimental/MeshModelingToolsetExp/Source/ModelingComponentsEditorOnly/Public/Operations/SubdividePoly.h b/Engine/Plugins/Experimental/MeshModelingToolsetExp/Source/ModelingComponentsEditorOnly/Public/Operations/SubdividePoly.h
--- a/Engine/Plugins/Experimental/MeshModelingToolsetExp/Source/ModelingComponentsEditorOnly/Public/Operations/SubdividePoly.h
+++ b/Engine/Plugins/Experimental/MeshModelingToolsetExp/Source/ModelingComponentsEditorOnly/Public/Operations/SubdividePoly.h
@@ -68,6 +68,17 @@ public:
 
 	~FSubdividePoly();
 
+	enum class ETopologyCheckResult
+	{
+		Ok = 0,
+		NoGroups = 1,
+		InsufficientGroups = 2,
+		UnboundedPolygroup = 3,
+		MultiBoundaryPolygroup = 4,
+		DegeneratePolygroup = 5
+	};
+	ETopologyCheckResult ValidateTopology();
+
 	bool ComputeTopologySubdivision();
 
 	bool ComputeSubdividedMesh(FDynamicMesh3& OutMesh);
diff --git a/Engine/Plugins/Runtime/GeometryProcessing/Source/DynamicMesh/Private/Parameterization/DynamicMeshUVEditor.cpp b/Engine/Plugins/Runtime/GeometryProcessing/Source/DynamicMesh/Private/Parameterization/DynamicMeshUVEditor.cpp
--- a/Engine/Plugins/Runtime/GeometryProcessing/Source/DynamicMesh/Private/Parameterization/DynamicMeshUVEditor.cpp
+++ b/Engine/Plugins/Runtime/GeometryProcessing/Source/DynamicMesh/Private/Parameterization/DynamicMeshUVEditor.cpp
@@ -80,6 +80,40 @@ void FDynamicMeshUVEditor::ResetUVs(const TArray<int32>& Triangles)
 	}
 }
 
+
+bool FDynamicMeshUVEditor::CopyUVLayer(FDynamicMeshUVOverlay* FromUVOverlay)
+{
+	if (FromUVOverlay == nullptr || FromUVOverlay == UVOverlay)
+	{
+		return false;
+	}
+
+	TArray<int32> ElementIDMap;
+	ElementIDMap.SetNum(FromUVOverlay->MaxElementID());
+
+	UVOverlay->ClearElements();
+	for (int32 ElementID : FromUVOverlay->ElementIndicesItr())
+	{
+		FVector2f UV = FromUVOverlay->GetElement(ElementID);
+		int32 NewID = UVOverlay->AppendElement(UV);
+		ElementIDMap[ElementID] = NewID;
+	}
+
+	for (int32 TriangleID : Mesh->TriangleIndicesItr())
+	{
+		if (FromUVOverlay->IsSetTriangle(TriangleID))
+		{
+			FIndex3i UVTriangle = FromUVOverlay->GetTriangle(TriangleID);
+			UVTriangle.A = ElementIDMap[UVTriangle.A];
+			UVTriangle.B = ElementIDMap[UVTriangle.B];
+			UVTriangle.C = ElementIDMap[UVTriangle.C];
+			UVOverlay->SetTriangle(TriangleID, UVTriangle);
+		}
+	}
+
+	return true;
+}
+
 void FDynamicMeshUVEditor::TransformUVElements(const TArray<int32>& ElementIDs, TFunctionRef<FVector2f(const FVector2f&)> TransformFunc)
 {
 	for (int32 elemid : ElementIDs)
diff --git a/Engine/Plugins/Runtime/GeometryProcessing/Source/DynamicMesh/Public/Parameterization/DynamicMeshUVEditor.h b/Engine/Plugins/Runtime/GeometryProcessing/Source/DynamicMesh/Public/Parameterization/DynamicMeshUVEditor.h
--- a/Engine/Plugins/Runtime/GeometryProcessing/Source/DynamicMesh/Public/Parameterization/DynamicMeshUVEditor.h
+++ b/Engine/Plugins/Runtime/GeometryProcessing/Source/DynamicMesh/Public/Parameterization/DynamicMeshUVEditor.h
@@ -64,6 +64,12 @@ public:
 	*/
 	void ResetUVs(const TArray<int32>& Triangles);
 
+	/**
+	 * Copy UVs from another overlay
+	 */
+	bool CopyUVLayer(FDynamicMeshUVOverlay* FromUVOverlay);
+
+
 	/**
 	 * Create new UV island for each Triangle, by planar projection onto plane of Triangle. No transforms/etc are applied.
 	 */
diff --git a/Engine/Plugins/Runtime/GeometryProcessing/Source/DynamicMesh/Public/Polygroups/PolygroupSet.h b/Engine/Plugins/Runtime/GeometryProcessing/Source/DynamicMesh/Public/Polygroups/PolygroupSet.h
--- a/Engine/Plugins/Runtime/GeometryProcessing/Source/DynamicMesh/Public/Polygroups/PolygroupSet.h
+++ b/Engine/Plugins/Runtime/GeometryProcessing/Source/DynamicMesh/Public/Polygroups/PolygroupSet.h
@@ -29,6 +29,18 @@ struct DYNAMICMESH_API FPolygroupLayer
 	/** Construct a FPolygroupLayer for an extended layer */
 	static FPolygroupLayer Layer(int32 Index) { return FPolygroupLayer{ false, Index }; }
 
+	bool operator==(const FPolygroupLayer& OtherLayer) const 
+	{
+		if (bIsDefaultLayer || OtherLayer.bIsDefaultLayer)
+		{
+			return bIsDefaultLayer && OtherLayer.bIsDefaultLayer;
+		}
+		else
+		{
+			return LayerIndex == OtherLayer.LayerIndex;
+		}
+	}
+
 	/** @return true if the specified layer (default or extended) exist and is initialized on the given Mesh */
 	bool CheckExists(const FDynamicMesh3* Mesh);
 };
