Mesh Simplification: Minimal Planar case.  Fix crashing bug

diff --git a/Engine/Plugins/Experimental/MeshModelingToolsetExp/Source/ModelingOperatorsEditorOnly/Private/CleaningOps/SimplifyMeshOp.cpp b/Engine/Plugins/Experimental/MeshModelingToolsetExp/Source/ModelingOperatorsEditorOnly/Private/CleaningOps/SimplifyMeshOp.cpp
--- a/Engine/Plugins/Experimental/MeshModelingToolsetExp/Source/ModelingOperatorsEditorOnly/Private/CleaningOps/SimplifyMeshOp.cpp
+++ b/Engine/Plugins/Experimental/MeshModelingToolsetExp/Source/ModelingOperatorsEditorOnly/Private/CleaningOps/SimplifyMeshOp.cpp
@@ -176,13 +176,13 @@ void FSimplifyMeshOp::CalculateResult(FProgressCancel* Progress)
 		{
 			FMeshNormals::QuickComputeVertexNormals(*ResultMesh, false);
 		}
-		ComputeSimplify<FAttrMeshSimplification>(TargetMesh, bReproject, OriginalTriCount, *OriginalMesh, *OriginalMeshSpatial,
+		ComputeSimplify<FQEMSimplification>(TargetMesh, bReproject, OriginalTriCount, *OriginalMesh, *OriginalMeshSpatial,
 			MeshBoundaryConstraint,
 			GroupBoundaryConstraint,
 			MaterialBoundaryConstraint,
 			bPreserveSharpEdges, bAllowSeamCollapse,
 			ESimplifyTargetType::MinimalPlanar, TargetPercentage, TargetCount, TargetEdgeLength, MinimalPlanarAngleThresh,
-			FAttrMeshSimplification::ESimplificationCollapseModes::MinimalQuadricPositionError, bUseQuadricMemory,
+			FQEMSimplification::ESimplificationCollapseModes::MinimalQuadricPositionError, bUseQuadricMemory,
 			UseGeometricTolerance);
 	}
 	else if (SimplifierType == ESimplifyType::MinimalExistingVertex)
