Geometry Tools :  SimplifyMeshTool, clean up UI  disable unused options when the Minimal Planar simplification approach is being used.

diff --git a/Engine/Plugins/Experimental/MeshModelingToolsetExp/Source/MeshModelingToolsEditorOnly/Private/SimplifyMeshTool.cpp b/Engine/Plugins/Experimental/MeshModelingToolsetExp/Source/MeshModelingToolsEditorOnly/Private/SimplifyMeshTool.cpp
--- a/Engine/Plugins/Experimental/MeshModelingToolsetExp/Source/MeshModelingToolsEditorOnly/Private/SimplifyMeshTool.cpp
+++ b/Engine/Plugins/Experimental/MeshModelingToolsetExp/Source/MeshModelingToolsEditorOnly/Private/SimplifyMeshTool.cpp
@@ -47,7 +47,9 @@ USimplifyMeshToolProperties::USimplifyMeshToolProperties()
 	SimplifierType = ESimplifyType::QEM;
 	TargetMode = ESimplifyTargetType::Percentage;
 	TargetPercentage = 50;
-	TargetCount = 1000;
+	TargetTriangleCount = 1000;
+	TargetVertexCount = 1000;
+	MinimalAngleThreshold = 0.01;
 	TargetEdgeLength = 5.0;
 	bReproject = false;
 	bPreventNormalFlips = true;
@@ -192,7 +194,8 @@ TUniquePtr<FDynamicMeshOperator> USimplifyMeshTool::MakeNewOperator()
 	Op->bAllowSeamCollapse = !SimplifyProperties->bPreserveSharpEdges;
 	Op->bReproject = SimplifyProperties->bReproject;
 	Op->SimplifierType = SimplifyProperties->SimplifierType;
-	Op->TargetCount = SimplifyProperties->TargetCount;
+	Op->TargetCount = ( SimplifyProperties->TargetMode == ESimplifyTargetType::VertexCount) ?  SimplifyProperties->TargetVertexCount : SimplifyProperties->TargetTriangleCount;
+	Op->MinimalPlanarAngleThresh = SimplifyProperties->MinimalAngleThreshold;
 	Op->TargetEdgeLength = SimplifyProperties->TargetEdgeLength;
 	Op->TargetMode = SimplifyProperties->TargetMode;
 	Op->TargetPercentage = SimplifyProperties->TargetPercentage;
diff --git a/Engine/Plugins/Experimental/MeshModelingToolsetExp/Source/MeshModelingToolsEditorOnly/Public/SimplifyMeshTool.h b/Engine/Plugins/Experimental/MeshModelingToolsetExp/Source/MeshModelingToolsEditorOnly/Public/SimplifyMeshTool.h
--- a/Engine/Plugins/Experimental/MeshModelingToolsetExp/Source/MeshModelingToolsEditorOnly/Public/SimplifyMeshTool.h
+++ b/Engine/Plugins/Experimental/MeshModelingToolsetExp/Source/MeshModelingToolsEditorOnly/Public/SimplifyMeshTool.h
@@ -44,20 +44,28 @@ public:
 	ESimplifyType SimplifierType;
 
 	/** Simplification Target Type  */
-	UPROPERTY(EditAnywhere, Category = Options)
+	UPROPERTY(EditAnywhere, Category = Options, meta = (EditCondition = "SimplifierType != ESimplifyType::MinimalPlanar"))
 	ESimplifyTargetType TargetMode;
 
 	/** Target percentage of original triangle count */
-	UPROPERTY(EditAnywhere, Category = Options, meta = (UIMin = "0", UIMax = "100", EditCondition = "TargetMode == ESimplifyTargetType::Percentage"))
+	UPROPERTY(EditAnywhere, Category = Options, meta = (UIMin = "0", UIMax = "100", EditCondition = "SimplifierType != ESimplifyType::MinimalPlanar && TargetMode == ESimplifyTargetType::Percentage"))
 	int TargetPercentage;
 
 	/** Target edge length */
-	UPROPERTY(EditAnywhere, Category = Options, meta = (UIMin = "3.0", UIMax = "10.0", ClampMin = "0.001", ClampMax = "1000.0", EditCondition = "TargetMode == ESimplifyTargetType::EdgeLength && SimplifierType != ESimplifyType::UEStandard"))
+	UPROPERTY(EditAnywhere, Category = Options, meta = (UIMin = "3.0", UIMax = "10.0", ClampMin = "0.001", ClampMax = "1000.0", EditCondition = "TargetMode == ESimplifyTargetType::EdgeLength && SimplifierType != ESimplifyType::UEStandard && SimplifierType != ESimplifyType::MinimalPlanar"))
 	float TargetEdgeLength;
 
 	/** Target triangle count */
-	UPROPERTY(EditAnywhere, Category = Options, meta = (UIMin = "4", UIMax = "10000", ClampMin = "1", ClampMax = "9999999999", EditCondition = "TargetMode == ESimplifyTargetType::TriangleCount || TargetMode == ESimplifyTargetType::VertexCount"))
-	int TargetCount;
+	UPROPERTY(EditAnywhere, Category = Options, meta = (UIMin = "4", UIMax = "10000", ClampMin = "1", ClampMax = "9999999999", EditCondition = "TargetMode == ESimplifyTargetType::TriangleCount && SimplifierType != ESimplifyType::MinimalPlanar"))
+	int TargetTriangleCount;
+
+	/** Target vertex count */
+	UPROPERTY(EditAnywhere, Category = Options, meta = (UIMin = "4", UIMax = "10000", ClampMin = "1", ClampMax = "9999999999", EditCondition = "TargetMode == ESimplifyTargetType::VertexCount && SimplifierType != ESimplifyType::MinimalPlanar"))
+	int TargetVertexCount;
+
+	/** Angle threshold in degrees used for testing if two triangles should be considered coplanar, or two lines collinear */
+	UPROPERTY()
+	float MinimalAngleThreshold = 0.01;
 
 	/** If true, UVs and Normals are discarded  */
 	UPROPERTY(EditAnywhere, Category = Options)
