Add polygroup simplifier to simplify tool + improve behavior for angle tolerance w/ smooth turns over many edges

diff --git a/Engine/Plugins/Experimental/MeshModelingToolsetExp/Source/MeshModelingToolsEditorOnly/Private/SimplifyMeshTool.cpp b/Engine/Plugins/Experimental/MeshModelingToolsetExp/Source/MeshModelingToolsEditorOnly/Private/SimplifyMeshTool.cpp
--- a/Engine/Plugins/Experimental/MeshModelingToolsetExp/Source/MeshModelingToolsEditorOnly/Private/SimplifyMeshTool.cpp
+++ b/Engine/Plugins/Experimental/MeshModelingToolsetExp/Source/MeshModelingToolsEditorOnly/Private/SimplifyMeshTool.cpp
@@ -204,6 +204,7 @@ TUniquePtr<FDynamicMeshOperator> USimplifyMeshTool::MakeNewOperator()
 	Op->MaterialBoundaryConstraint = (EEdgeRefineFlags)SimplifyProperties->MaterialBoundaryConstraint;
 	Op->bGeometricDeviationConstraint = SimplifyProperties->bGeometricConstraint;
 	Op->GeometricTolerance = SimplifyProperties->GeometricTolerance;
+	Op->PolyEdgeAngleTolerance = SimplifyProperties->PolyEdgeAngleTolerance;
 	FTransform LocalToWorld = (FTransform)UE::ToolTarget::GetLocalToWorldTransform(Target);
 	Op->SetTransform(LocalToWorld);
 
diff --git a/Engine/Plugins/Experimental/MeshModelingToolsetExp/Source/MeshModelingToolsEditorOnly/Public/SimplifyMeshTool.h b/Engine/Plugins/Experimental/MeshModelingToolsetExp/Source/MeshModelingToolsEditorOnly/Public/SimplifyMeshTool.h
--- a/Engine/Plugins/Experimental/MeshModelingToolsetExp/Source/MeshModelingToolsEditorOnly/Public/SimplifyMeshTool.h
+++ b/Engine/Plugins/Experimental/MeshModelingToolsetExp/Source/MeshModelingToolsEditorOnly/Public/SimplifyMeshTool.h
@@ -44,39 +44,49 @@ public:
 	ESimplifyType SimplifierType;
 
 	/** Simplification Target Type  */
-	UPROPERTY(EditAnywhere, Category = Options, meta = (EditCondition = "SimplifierType != ESimplifyType::MinimalPlanar"))
+	UPROPERTY(EditAnywhere, Category = Options, meta = (EditCondition = "SimplifierType != ESimplifyType::MinimalPlanar && SimplifierType != ESimplifyType::MinimalPolygroup"))
 	ESimplifyTargetType TargetMode;
 
 	/** Target percentage of original triangle count */
-	UPROPERTY(EditAnywhere, Category = Options, meta = (UIMin = "0", UIMax = "100", EditCondition = "SimplifierType != ESimplifyType::MinimalPlanar && TargetMode == ESimplifyTargetType::Percentage"))
+	UPROPERTY(EditAnywhere, Category = Options, meta = (UIMin = "0", UIMax = "100",
+				EditCondition = "SimplifierType != ESimplifyType::MinimalPolygroup && SimplifierType != ESimplifyType::MinimalPlanar && TargetMode == ESimplifyTargetType::Percentage"))
 	int TargetPercentage;
 
 	/** Target edge length */
-	UPROPERTY(EditAnywhere, Category = Options, meta = (UIMin = "3.0", UIMax = "10.0", ClampMin = "0.001", ClampMax = "1000.0", EditCondition = "TargetMode == ESimplifyTargetType::EdgeLength && SimplifierType != ESimplifyType::UEStandard && SimplifierType != ESimplifyType::MinimalPlanar"))
+	UPROPERTY(EditAnywhere, Category = Options, meta = (UIMin = "3.0", UIMax = "10.0", ClampMin = "0.001", ClampMax = "1000.0",
+		EditCondition = "TargetMode == ESimplifyTargetType::EdgeLength && SimplifierType != ESimplifyType::UEStandard && SimplifierType != ESimplifyType::MinimalPlanar && SimplifierType != ESimplifyType::MinimalPolygroup"))
 	float TargetEdgeLength;
 
 	/** Target triangle count */
-	UPROPERTY(EditAnywhere, Category = Options, meta = (UIMin = "4", UIMax = "10000", ClampMin = "1", ClampMax = "9999999999", EditCondition = "TargetMode == ESimplifyTargetType::TriangleCount && SimplifierType != ESimplifyType::MinimalPlanar"))
+	UPROPERTY(EditAnywhere, Category = Options, meta = (UIMin = "4", UIMax = "10000", ClampMin = "1", ClampMax = "9999999999",
+				EditCondition = "TargetMode == ESimplifyTargetType::TriangleCount && SimplifierType != ESimplifyType::MinimalPlanar && SimplifierType != ESimplifyType::MinimalPolygroup"))
 	int TargetTriangleCount;
 
 	/** Target vertex count */
-	UPROPERTY(EditAnywhere, Category = Options, meta = (UIMin = "4", UIMax = "10000", ClampMin = "1", ClampMax = "9999999999", EditCondition = "TargetMode == ESimplifyTargetType::VertexCount && SimplifierType != ESimplifyType::MinimalPlanar"))
+	UPROPERTY(EditAnywhere, Category = Options, meta = (UIMin = "4", UIMax = "10000", ClampMin = "1", ClampMax = "9999999999",
+				EditCondition = "TargetMode == ESimplifyTargetType::VertexCount && SimplifierType != ESimplifyType::MinimalPlanar"))
 	int TargetVertexCount;
 
 	/** Angle threshold in degrees used for testing if two triangles should be considered coplanar, or two lines collinear */
 	UPROPERTY()
 	float MinimalAngleThreshold = 0.01;
 
+	//~ Note PolyEdgeAngleTolerance is very similar to MinimalAngleThreshold, but not redundant b/c the useful ranges are very different (MinimalAngleThreshold should generally be kept very small)
+	/** Threshold angle change (in degrees) along a polygroup edge, above which a vertex must be added */
+	UPROPERTY(EditAnywhere, Category = Options, meta = (UIMin = "0.001", ClampMin = "0.0", UIMax = "90.0", ClampMax = "180.0", EditCondition = "SimplifierType == ESimplifyType::MinimalPolygroup"))
+	float PolyEdgeAngleTolerance = .1;
+
 	/** If true, UVs and Normals are discarded  */
 	UPROPERTY(EditAnywhere, Category = Options)
 	bool bDiscardAttributes;
 
 	/** If true, then simplification will consider geometric deviation with the input mesh  */
-	UPROPERTY(EditAnywhere, Category = Options)
+	UPROPERTY(EditAnywhere, Category = Options, meta = (EditCondition = "SimplifierType != ESimplifyType::MinimalPolygroup"))
 	bool bGeometricConstraint;
 
 	/** Geometric deviation tolerance used when bGeometricConstraint is enabled, to limit the geometric deviation between the simplified and original meshes */
-	UPROPERTY(EditAnywhere, Category = Options, meta = (UIMin = "0.0", UIMax = "10.0", ClampMin = "0.0", ClampMax = "10000000.0", EditCondition = "bGeometricConstraint && SimplifierType != ESimplifyType::UEStandard"))
+	UPROPERTY(EditAnywhere, Category = Options, meta = (UIMin = "0.0", UIMax = "10.0", ClampMin = "0.0", ClampMax = "10000000.0",
+				EditCondition = "bGeometricConstraint && SimplifierType != ESimplifyType::UEStandard && SimplifierType != ESimplifyType::MinimalPolygroup"))
 	float GeometricTolerance;
 
 	/** Display colors corresponding to the mesh's polygon groups */
diff --git a/Engine/Plugins/Experimental/MeshModelingToolsetExp/Source/ModelingOperatorsEditorOnly/Private/CleaningOps/SimplifyMeshOp.cpp b/Engine/Plugins/Experimental/MeshModelingToolsetExp/Source/ModelingOperatorsEditorOnly/Private/CleaningOps/SimplifyMeshOp.cpp
--- a/Engine/Plugins/Experimental/MeshModelingToolsetExp/Source/ModelingOperatorsEditorOnly/Private/CleaningOps/SimplifyMeshOp.cpp
+++ b/Engine/Plugins/Experimental/MeshModelingToolsetExp/Source/ModelingOperatorsEditorOnly/Private/CleaningOps/SimplifyMeshOp.cpp
@@ -15,6 +15,9 @@
 #include "MeshDescription.h"
 #include "DynamicMesh/MeshNormals.h"
 #include "DynamicMesh/Operations/MergeCoincidentMeshEdges.h"
+#include "GroupTopology.h"
+#include "Operations/PolygroupRemesh.h"
+#include "ConstrainedDelaunay2.h"
 #include "OverlappingCorners.h"
 #include "StaticMeshOperations.h"
 
@@ -198,6 +201,19 @@ void FSimplifyMeshOp::CalculateResult(FProgressCancel* Progress)
 			FQEMSimplification::ESimplificationCollapseModes::MinimalExistingVertexError, bUseQuadricMemory,
 			UseGeometricTolerance);
 	}
+	else if (SimplifierType == ESimplifyType::MinimalPolygroup)
+	{
+		ResultMesh->Copy(*OriginalMesh, true, true, true, !bDiscardAttributes);
+		FGroupTopology Topology(ResultMesh.Get(), true);
+		if (Progress && Progress->Cancelled())
+		{
+			return;
+		}
+
+		FPolygroupRemesh Remesh(ResultMesh.Get(), &Topology, ConstrainedDelaunayTriangulate<double>);
+		Remesh.SimplificationAngleTolerance = PolyEdgeAngleTolerance;
+		bool bSuccess = Remesh.Compute();
+	}
 	else // SimplifierType == ESimplifyType::UEStandard
 	{
 		const FMeshDescription* SrcMeshDescription = OriginalMeshDescription.Get();
diff --git a/Engine/Plugins/Experimental/MeshModelingToolsetExp/Source/ModelingOperatorsEditorOnly/Public/CleaningOps/SimplifyMeshOp.h b/Engine/Plugins/Experimental/MeshModelingToolsetExp/Source/ModelingOperatorsEditorOnly/Public/CleaningOps/SimplifyMeshOp.h
--- a/Engine/Plugins/Experimental/MeshModelingToolsetExp/Source/ModelingOperatorsEditorOnly/Public/CleaningOps/SimplifyMeshOp.h
+++ b/Engine/Plugins/Experimental/MeshModelingToolsetExp/Source/ModelingOperatorsEditorOnly/Public/CleaningOps/SimplifyMeshOp.h
@@ -51,6 +51,9 @@ enum class ESimplifyType : uint8
 	/** Collapse any spurious edges but do not change the 3D shape. */
 	MinimalPlanar = 4 UMETA(DisplayName = "Minimal Shape-Preserving"),
 
+	/** Only preserve polygroup boundaries; ignore all other shape features */
+	MinimalPolygroup = 5 UMETA(DisplayName = "Minimal Polygroup-Preserving"),
+
 };
 
 namespace UE
@@ -76,6 +79,9 @@ public:
 	/** Angle threshold in degrees used for testing if two triangles should be considered coplanar, or two lines collinear */
 	float MinimalPlanarAngleThresh = 0.01f;
 
+	// For minimal polygroup-preserving simplification: Threshold angle change (in degrees) along a polygroup edge, above which a vertex must be added
+	float PolyEdgeAngleTolerance = 0.1f;
+
 	bool bGeometricDeviationConstraint = false;
 	float GeometricTolerance = 0.0f;
 
diff --git a/Engine/Plugins/Runtime/GeometryProcessing/Source/DynamicMesh/Private/Operations/PolygroupRemesh.cpp b/Engine/Plugins/Runtime/GeometryProcessing/Source/DynamicMesh/Private/Operations/PolygroupRemesh.cpp
--- a/Engine/Plugins/Runtime/GeometryProcessing/Source/DynamicMesh/Private/Operations/PolygroupRemesh.cpp
+++ b/Engine/Plugins/Runtime/GeometryProcessing/Source/DynamicMesh/Private/Operations/PolygroupRemesh.cpp
@@ -98,6 +98,7 @@ bool FPolygroupRemesh::Compute()
 			}
 		}
 		FVector3d PrevEdge = FirstEdge;
+		FVector3d LastKeptEdge = FirstEdge;
 		int32 FirstKeptIdx = KeptVertices[SpanV[0]] ? 0 : -1;
 		int32 LastKeptIdx = FirstKeptIdx;
 		int32 NumKept = int32(KeptVertices[SpanV[0]]) + int32(KeptVertices[SpanV.Last()]);
@@ -107,7 +108,8 @@ bool FPolygroupRemesh::Compute()
 			FVector3d NextEdge = Next - Center;
 			NextEdge.Normalize();
 			int32 SpanVID = SpanV[SpanIdx];
-			if (KeptVertices[SpanVID] || NextEdge.Dot(PrevEdge) < DotTolerance)
+			// track deviation from (1) the most recent edge and (2) the last edge from a vertex we kept, so gentle slow turns across many edges are also captured
+			if (KeptVertices[SpanVID] || NextEdge.Dot(PrevEdge) < DotTolerance || NextEdge.Dot(LastKeptEdge) < DotTolerance)
 			{
 				NumKept++;
 				KeptVertices[SpanVID] = true;
@@ -122,6 +124,7 @@ bool FPolygroupRemesh::Compute()
 					}
 				}
 
+				LastKeptEdge = NextEdge;
 				LastKeptIdx = SpanIdx;
 				if (FirstKeptIdx == -1)
 				{
