MeshModelingTools: Fix flip not working in CubeGrid tool, fix Escape not correctly clearing selection there, fix incorrect face being picked on transformed meshes and transformed grid plane.

#rnx

diff --git a/Engine/Plugins/Experimental/MeshModelingToolsetExp/Source/MeshModelingToolsExp/Private/CubeGridTool.cpp b/Engine/Plugins/Experimental/MeshModelingToolsetExp/Source/MeshModelingToolsExp/Private/CubeGridTool.cpp
--- a/Engine/Plugins/Experimental/MeshModelingToolsetExp/Source/MeshModelingToolsExp/Private/CubeGridTool.cpp
+++ b/Engine/Plugins/Experimental/MeshModelingToolsetExp/Source/MeshModelingToolsExp/Private/CubeGridTool.cpp
@@ -1217,7 +1217,7 @@ bool UCubeGridTool::GetHitGridFace(const FRay& WorldRay, FCubeGrid::FCubeFace& F
 		FVector3d ClampedGridPoint = CubeGrid->ToGridPoint(IntersectionPoint);
 		ClampedGridPoint.Z = 0;
 		FaceOut = FCubeGrid::FCubeFace(ClampedGridPoint,
-			WorldRay.Origin.Z >= 0 ? FCubeGrid::EFaceDirection::PositiveZ : FCubeGrid::EFaceDirection::NegativeZ,
+			CubeGrid->ToGridPoint(WorldRay.Origin).Z >= 0 ? FCubeGrid::EFaceDirection::PositiveZ : FCubeGrid::EFaceDirection::NegativeZ,
 			CubeGrid->GetGridPower());
 		BestHitT = WorldRay.GetParameter(IntersectionPoint);
 	}
@@ -1255,7 +1255,8 @@ bool UCubeGridTool::GetHitGridFace(const FRay& WorldRay, FCubeGrid::FCubeFace& F
 			{
 				BestHitT = HitT;
 				ConvertToFaceHit(*CubeGrid, Settings->FaceSelectionMode,
-					WorldRay, BestHitT, CurrentMesh->GetTriNormal(Tid), FaceOut, Settings->PlaneTolerance);
+					WorldRay, BestHitT, CurrentMeshTransform.TransformNormal(CurrentMesh->GetTriNormal(Tid)), 
+					FaceOut, Settings->PlaneTolerance);
 			}
 		}
 	}
@@ -1658,7 +1659,7 @@ void UCubeGridTool::ApplyFlipSelection()
 	GetToolManager()->BeginUndoTransaction(LOCTEXT("FlipTransactionName", "Flip Selection"));
 
 	FSelection NewSelection = Selection;
-	Selection.Direction = FCubeGrid::FlipDir(Selection.Direction);
+	NewSelection.Direction = FCubeGrid::FlipDir(Selection.Direction);
 	SetSelection(NewSelection, true);
 
 	GetToolManager()->EndUndoTransaction();
@@ -2048,6 +2049,7 @@ bool UCubeGridTool::ExecuteNestedCancelCommand()
 	else if (bHaveSelection)
 	{
 		ClearSelection(true);
+		return true;
 	}
 	return false;
 }
