MeshModelingTools: Fix normals being incorrectly set in our box generator. Fix wrong diagonal being picked occasionally in corner mode of CubeGrid tool.

#rnx

diff --git a/Engine/Plugins/Experimental/MeshModelingToolsetExp/Source/MeshModelingToolsExp/Private/CubeGridTool.cpp b/Engine/Plugins/Experimental/MeshModelingToolsetExp/Source/MeshModelingToolsExp/Private/CubeGridTool.cpp
--- a/Engine/Plugins/Experimental/MeshModelingToolsetExp/Source/MeshModelingToolsExp/Private/CubeGridTool.cpp
+++ b/Engine/Plugins/Experimental/MeshModelingToolsetExp/Source/MeshModelingToolsExp/Private/CubeGridTool.cpp
@@ -1597,7 +1597,16 @@ void UCubeGridTool::UpdateCornerModeLineSet()
 			CornerWelded[i] = !CornerSelectedFlags[i];
 		}
 
-		int DiagStartIdx = (CornerWelded[0] && CornerWelded[2]) ? 1 : 0;
+		// The choice of diagonal here lines up with the generator in CubeGridBooleanOp. The
+		// indices look a little different because we're accounting for mirroring that happens in
+		// subtract mode, though we would have actually gotten the same results either way in
+		// the cases we care about (the nonplanar ones)
+		int DiagStartIdx = 1;
+		if (CornerWelded[0] != CornerWelded[2] ||
+			(!CornerWelded[0] && CornerWelded[1] && CornerWelded[3]))
+		{
+			DiagStartIdx = 0;
+		}
 		DiagStartIdx = Settings->bCrosswiseDiagonal ? 1 - DiagStartIdx : DiagStartIdx;
 		
 		bool bDiagonalWelded = CornerWelded[DiagStartIdx] && CornerWelded[DiagStartIdx + 2];
diff --git a/Engine/Plugins/Experimental/MeshModelingToolsetExp/Source/ModelingOperators/Private/CompositionOps/CubeGridBooleanOp.cpp b/Engine/Plugins/Experimental/MeshModelingToolsetExp/Source/ModelingOperators/Private/CompositionOps/CubeGridBooleanOp.cpp
--- a/Engine/Plugins/Experimental/MeshModelingToolsetExp/Source/ModelingOperators/Private/CompositionOps/CubeGridBooleanOp.cpp
+++ b/Engine/Plugins/Experimental/MeshModelingToolsetExp/Source/ModelingOperators/Private/CompositionOps/CubeGridBooleanOp.cpp
@@ -184,7 +184,12 @@ protected:
 		// corner results in a pyramid with a square (rather than triangular) base, and
 		// a checkerboard raised/lowered state results in a raised ridge. In the crosswise
 		// case we flip the diagonal.
-		int DiagFaceIdx1 = (CornerInfo.WeldedAtBase[0] && CornerInfo.WeldedAtBase[2]) ? 1 : 0;
+		int DiagFaceIdx1 = 0;
+		if (CornerInfo.WeldedAtBase[1] != CornerInfo.WeldedAtBase[3] || 
+			(!CornerInfo.WeldedAtBase[1] && CornerInfo.WeldedAtBase[0] && CornerInfo.WeldedAtBase[2]))
+		{
+			DiagFaceIdx1 = 1;
+		}
 		DiagFaceIdx1 = bCrosswiseDiagonal ? 1 - DiagFaceIdx1 : DiagFaceIdx1;
 
 		int DiagFaceIdx2 = DiagFaceIdx1 + 2;
diff --git a/Engine/Source/Runtime/GeometryCore/Public/Generators/GridBoxMeshGenerator.h b/Engine/Source/Runtime/GeometryCore/Public/Generators/GridBoxMeshGenerator.h
--- a/Engine/Source/Runtime/GeometryCore/Public/Generators/GridBoxMeshGenerator.h
+++ b/Engine/Source/Runtime/GeometryCore/Public/Generators/GridBoxMeshGenerator.h
@@ -193,6 +193,8 @@ public:
 
 				FVector3f Normal(0, 0, 0);
 				Normal[Dim] = (2 * Side - 1);
+				Normal = (FVector3f)Box.Frame.FromFrameVector((FVector3d)Normal);
+
 				int MajorFaceInd = FaceIdxBase + Side;
 
 				int FaceUVStartInd = CurrentUVIdx;
