add missing explicit namespace in checkSlow
#rnx

diff --git a/Engine/Plugins/Experimental/MeshModelingToolsetExp/Source/ModelingComponents/Public/Mechanics/CurveControlPointsMechanic.h b/Engine/Plugins/Experimental/MeshModelingToolsetExp/Source/ModelingComponents/Public/Mechanics/CurveControlPointsMechanic.h
--- a/Engine/Plugins/Experimental/MeshModelingToolsetExp/Source/ModelingComponents/Public/Mechanics/CurveControlPointsMechanic.h
+++ b/Engine/Plugins/Experimental/MeshModelingToolsetExp/Source/ModelingComponents/Public/Mechanics/CurveControlPointsMechanic.h
@@ -156,7 +156,7 @@ protected:
 		 */
 		void SetPointCoordinates(int32 PointID, const FVector3d& NewCoordinates)
 		{
-			checkSlow(VectorUtil::IsFinite(NewCoordinates));
+			checkSlow(UE::Geometry::VectorUtil::IsFinite(NewCoordinates));
 			check(IsValidPoint(PointID));
 
 			Vertices[PointID] = NewCoordinates;
