Various improvements and to marquee and click selection for ULatticeControlPointsMechanic and UPolygonSelectionMechanic

- Fixed marquee selection so it works when starting on a click-selectable item in ULatticeControlPointsMechanic and UPolygonSelectionMechanic (used USingleClickOrDragInputBehavior)
- Fixed issue where selection changes that do nothing were added to undo/redo history in ULatticeControlPointsMechanic and UPolygonSelectionMechanic
- Addressed a TODO to make ULatticeControlPointsMechanic click behavior use control and shift modifiers in the same way as the marquee mechanic
- Refactored ULatticeControlPointsMechanic to simplify the selection update code and make it use the same path for clicks and marquee select (removed ChangeSelection and ClearSelection)

#rnx

diff --git a/Engine/Plugins/Experimental/MeshModelingToolsetExp/Source/ModelingComponents/Private/Mechanics/LatticeControlPointsMechanic.cpp b/Engine/Plugins/Experimental/MeshModelingToolsetExp/Source/ModelingComponents/Private/Mechanics/LatticeControlPointsMechanic.cpp
--- a/Engine/Plugins/Experimental/MeshModelingToolsetExp/Source/ModelingComponents/Private/Mechanics/LatticeControlPointsMechanic.cpp
+++ b/Engine/Plugins/Experimental/MeshModelingToolsetExp/Source/ModelingComponents/Private/Mechanics/LatticeControlPointsMechanic.cpp
@@ -2,7 +2,7 @@
 
 #include "Mechanics/LatticeControlPointsMechanic.h"
 
-#include "BaseBehaviors/SingleClickBehavior.h"
+#include "BaseBehaviors/SingleClickOrDragBehavior.h"
 #include "BaseBehaviors/MouseHoverBehavior.h"
 #include "BaseBehaviors/ClickDragBehavior.h"
 #include "BaseGizmos/TransformGizmoUtil.h"
@@ -38,6 +38,18 @@ namespace LatticeControlPointMechanicLocals
 			}
 		}
 	}
+
+	// TODO Move to SetUtil namespace and call it in FGroupTopologySelection::operator==
+	bool IsEqual(const TSet<int32>& A, const TSet<int32>& B)
+	{
+		if (A.Num() != B.Num())
+		{
+			return false;
+		}
+
+		// No need to also check B.Includes(A) here since counts match and sets contain unique elements
+		return A.Includes(B);
+	}
 }
 
 
@@ -46,19 +58,20 @@ void ULatticeControlPointsMechanic::Setup(UInteractiveTool* ParentToolIn)
 	UInteractionMechanic::Setup(ParentToolIn);
 
 	MarqueeMechanic = NewObject<URectangleMarqueeMechanic>(this);
+	MarqueeMechanic->bUseExternalClickDragBehavior = true;
 	MarqueeMechanic->Setup(ParentToolIn);
 	MarqueeMechanic->OnDragRectangleStarted.AddUObject(this, &ULatticeControlPointsMechanic::OnDragRectangleStarted);
 	MarqueeMechanic->OnDragRectangleChanged.AddUObject(this, &ULatticeControlPointsMechanic::OnDragRectangleChanged);
 	MarqueeMechanic->OnDragRectangleFinished.AddUObject(this, &ULatticeControlPointsMechanic::OnDragRectangleFinished);
 
-	USingleClickInputBehavior* ClickBehavior = NewObject<USingleClickInputBehavior>();
-	ClickBehavior->Initialize(this);
-	ParentTool->AddInputBehavior(ClickBehavior);
+	USingleClickOrDragInputBehavior* ClickOrDragBehavior = NewObject<USingleClickOrDragInputBehavior>();
+	ClickOrDragBehavior->Initialize(this, MarqueeMechanic);
+	ClickOrDragBehavior->Modifiers.RegisterModifier(ShiftModifierID, FInputDeviceState::IsShiftKeyDown);
+	ClickOrDragBehavior->Modifiers.RegisterModifier(CtrlModifierID, FInputDeviceState::IsCtrlKeyDown);
+	ParentTool->AddInputBehavior(ClickOrDragBehavior);
 
 	UMouseHoverBehavior* HoverBehavior = NewObject<UMouseHoverBehavior>();
 	HoverBehavior->Initialize(this);
-	HoverBehavior->Modifiers.RegisterModifier(ShiftModifierID, FInputDeviceState::IsShiftKeyDown);
-	HoverBehavior->Modifiers.RegisterModifier(CtrlModifierID, FInputDeviceState::IsCtrlKeyDown);
 	ParentTool->AddInputBehavior(HoverBehavior);
 
 	DrawnControlPoints = NewObject<UPointSetComponent>();
@@ -394,86 +407,65 @@ bool ULatticeControlPointsMechanic::HitTest(const FInputDeviceRay& ClickPos, FIn
 FInputRayHit ULatticeControlPointsMechanic::IsHitByClick(const FInputDeviceRay& ClickPos)
 {
 	FInputRayHit Result;
-	HitTest(ClickPos, Result);
-	return Result;
-}
-
-void ULatticeControlPointsMechanic::OnClicked(const FInputDeviceRay& ClickPos)
-{
-	FGeometrySet3::FNearest Nearest;
-
-	if (GeometrySet.FindNearestPointToRay(ClickPos.WorldRay, Nearest, GeometrySetToleranceTest))
+	if (HitTest(ClickPos, Result))
 	{
-		ParentTool->GetToolManager()->BeginUndoTransaction(LatticePointSelectionTransactionText);
-		// TODO: Have the modifier keys for single-select match the behavior of marquee select
-		ChangeSelection(Nearest.ID, ShouldAddToSelectionFunc());
-		ParentTool->GetToolManager()->EndUndoTransaction();
+		return Result;
 	}
 
+	// Return a hit so we always capture and can clear the selection
+	return FInputRayHit(TNumericLimits<float>::Max());
 }
 
-void ULatticeControlPointsMechanic::ChangeSelection(int32 NewPointID, bool bAddToSelection)
+void ULatticeControlPointsMechanic::OnClicked(const FInputDeviceRay& ClickPos)
 {
-	// If not adding to selection, clear it
-	bool bAnyChange = false;
-	if (!bAddToSelection && SelectedPointIDs.Num() > 0)
+	using namespace LatticeControlPointMechanicLocals;
+	
+	ParentTool->GetToolManager()->BeginUndoTransaction(LatticePointSelectionTransactionText);
+	
+	const TSet<int32> PreClickSelection = SelectedPointIDs;
+	
+	FGeometrySet3::FNearest Nearest;
+	if (GeometrySet.FindNearestPointToRay(ClickPos.WorldRay, Nearest, GeometrySetToleranceTest))
 	{
-		TSet<int32> PointsToDeselect;
-
-		for (int32 PointID : SelectedPointIDs)
+		if (ShouldAddToSelectionFunc())
 		{
-			// We check for validity here because we'd like to be able to use this function to deselect points after
-			// deleting them.
-			if (DrawnControlPoints->IsPointValid(PointID))
+			if (ShouldRemoveFromSelectionFunc())
+			{
+				Toggle(SelectedPointIDs, TSet<int>{Nearest.ID});
+			}
+			else
 			{
-				PointsToDeselect.Add(PointID);
-				DrawnControlPoints->SetPointColor(PointID, NormalPointColor);
+				SelectedPointIDs.Add(Nearest.ID);
 			}
 		}
-
-		FTransform PreviousTransform = PointTransformProxy->GetTransform();
-		SelectedPointIDs.Empty();
-		UpdateGizmoLocation();
-		FTransform NewTransform = PointTransformProxy->GetTransform();
-
-		ParentTool->GetToolManager()->EmitObjectChange(this, MakeUnique<FLatticeControlPointsMechanicSelectionChange>(
-			PointsToDeselect, false, PreviousTransform, NewTransform, CurrentChangeStamp), LatticePointDeselectionTransactionText);
-		bAnyChange = true;
-	}
-
-	for (const TPair<int32, FColor>& Override : ColorOverrides)
-	{
-		DrawnControlPoints->SetPointColor(Override.Key, Override.Value);
-	}
-
-	// We check for validity here because giving an invalid id (such as -1) with bAddToSelection == false
-	// is an easy way to clear the selection.
-	if ((NewPointID >= 0)  && (NewPointID < ControlPoints.Num()))
-	{
-		FTransform PreviousTransform = PointTransformProxy->GetTransform();
-
-		if (bAddToSelection && DeselectPoint(NewPointID))
+		else if (ShouldRemoveFromSelectionFunc())
 		{
-			UpdateGizmoLocation();
-			FTransform NewTransform = PointTransformProxy->GetTransform();
-			ParentTool->GetToolManager()->EmitObjectChange(this, MakeUnique<FLatticeControlPointsMechanicSelectionChange>(
-				NewPointID, false, PreviousTransform, NewTransform, CurrentChangeStamp), LatticePointDeselectionTransactionText);
+			SelectedPointIDs.Remove(Nearest.ID);
 		}
 		else
 		{
-			SelectPoint(NewPointID);
-			UpdateGizmoLocation();
-			FTransform NewTransform = PointTransformProxy->GetTransform();
-			ParentTool->GetToolManager()->EmitObjectChange(this, MakeUnique<FLatticeControlPointsMechanicSelectionChange>(
-				NewPointID, true, PreviousTransform, NewTransform, CurrentChangeStamp), LatticePointSelectionTransactionText);
+			// Neither key pressed.
+			SelectedPointIDs = TSet<int>{Nearest.ID};
 		}
-		bAnyChange = true;
 	}
-
-	if (bAnyChange)
+	else
 	{
+		SelectedPointIDs.Empty();
+	}
+	
+	if (!IsEqual(PreClickSelection, SelectedPointIDs))
+	{
+		const FTransform Transform = PointTransformProxy->GetTransform();
+		ParentTool->GetToolManager()->EmitObjectChange(this, MakeUnique<FLatticeControlPointsMechanicSelectionChange>(
+			PreClickSelection, false, Transform, Transform, CurrentChangeStamp), LatticePointDeselectionTransactionText);
+		ParentTool->GetToolManager()->EmitObjectChange(this, MakeUnique<FLatticeControlPointsMechanicSelectionChange>(
+			SelectedPointIDs, true, Transform, Transform, CurrentChangeStamp), LatticePointSelectionTransactionText);
 		OnSelectionChanged.Broadcast();
 	}
+
+	ParentTool->GetToolManager()->EndUndoTransaction();
+	
+	UpdateDrawables();
 }
 
 void ULatticeControlPointsMechanic::UpdateGizmoLocation()
@@ -543,11 +535,6 @@ void ULatticeControlPointsMechanic::SelectPoint(int32 PointID)
 	}
 }
 
-void ULatticeControlPointsMechanic::ClearSelection()
-{
-	ChangeSelection(-1, false);
-}
-
 FInputRayHit ULatticeControlPointsMechanic::BeginHoverSequenceHitTest(const FInputDeviceRay& PressPos)
 {
 	FInputRayHit Result;
@@ -628,6 +615,8 @@ void ULatticeControlPointsMechanic::OnDragRectangleStarted()
 
 void ULatticeControlPointsMechanic::OnDragRectangleChanged(const FCameraRectangle& Rectangle)
 {
+	using namespace LatticeControlPointMechanicLocals;
+	
 	TSet<int> DragSelection;
 	for (int32 PointID = 0; PointID < ControlPoints.Num(); ++PointID)
 	{
@@ -643,7 +632,7 @@ void ULatticeControlPointsMechanic::OnDragRectangleChanged(const FCameraRectangl
 		SelectedPointIDs = PreDragSelection;
 		if (ShouldRemoveFromSelectionFunc())
 		{
-			LatticeControlPointMechanicLocals::Toggle(SelectedPointIDs, DragSelection);
+			Toggle(SelectedPointIDs, DragSelection);
 		}
 		else
 		{
@@ -665,37 +654,22 @@ void ULatticeControlPointsMechanic::OnDragRectangleChanged(const FCameraRectangl
 
 void ULatticeControlPointsMechanic::OnDragRectangleFinished(const FCameraRectangle& Rectangle, bool bCancelled)
 {
-	ParentTool->GetToolManager()->BeginUndoTransaction(LatticePointSelectionTransactionText);
-
-	FTransform PreviousTransform = PointTransformProxy->GetTransform();
+	using namespace LatticeControlPointMechanicLocals;
 	
-	UpdateGizmoLocation();
-
-	FTransform NewTransform = PointTransformProxy->GetTransform();
-
-	bool bAnyChange = false;
+	ParentTool->GetToolManager()->BeginUndoTransaction(LatticePointSelectionTransactionText);
 
-	if (PreDragSelection.Num() > 0)
+	if (!IsEqual(PreDragSelection, SelectedPointIDs))
 	{
+		const FTransform Transform = PointTransformProxy->GetTransform();
 		ParentTool->GetToolManager()->EmitObjectChange(this, MakeUnique<FLatticeControlPointsMechanicSelectionChange>(
-			PreDragSelection, false, PreviousTransform, NewTransform, CurrentChangeStamp), LatticePointDeselectionTransactionText);
-		bAnyChange = true;
-	}
-
-	if (SelectedPointIDs.Num() > 0)
-	{
+			PreDragSelection, false, Transform, Transform, CurrentChangeStamp), LatticePointDeselectionTransactionText);
 		ParentTool->GetToolManager()->EmitObjectChange(this, MakeUnique<FLatticeControlPointsMechanicSelectionChange>(
-			SelectedPointIDs, true, PreviousTransform, NewTransform, CurrentChangeStamp), LatticePointSelectionTransactionText);
-		bAnyChange = true;
+			SelectedPointIDs, true, Transform, Transform, CurrentChangeStamp), LatticePointSelectionTransactionText);
+		OnSelectionChanged.Broadcast();
 	}
 
 	ParentTool->GetToolManager()->EndUndoTransaction();
 
-	if (bAnyChange)
-	{
-		OnSelectionChanged.Broadcast();
-	}
-
 	UpdateDrawables();
 }
 
diff --git a/Engine/Plugins/Experimental/MeshModelingToolsetExp/Source/ModelingComponents/Private/Selection/PolygonSelectionMechanic.cpp b/Engine/Plugins/Experimental/MeshModelingToolsetExp/Source/ModelingComponents/Private/Selection/PolygonSelectionMechanic.cpp
--- a/Engine/Plugins/Experimental/MeshModelingToolsetExp/Source/ModelingComponents/Private/Selection/PolygonSelectionMechanic.cpp
+++ b/Engine/Plugins/Experimental/MeshModelingToolsetExp/Source/ModelingComponents/Private/Selection/PolygonSelectionMechanic.cpp
@@ -2,7 +2,7 @@
 
 #include "Selection/PolygonSelectionMechanic.h"
 
-#include "BaseBehaviors/ClickDragBehavior.h"
+#include "BaseBehaviors/SingleClickOrDragBehavior.h"
 #include "BaseBehaviors/MouseHoverBehavior.h"
 #include "BaseBehaviors/SingleClickBehavior.h"
 #include "InteractiveToolManager.h"
@@ -27,12 +27,8 @@ void UPolygonSelectionMechanic::Setup(UInteractiveTool* ParentToolIn)
 
 	TopoSelector = MakeShared<FGroupTopologySelector, ESPMode::ThreadSafe>();
 
-	ClickBehavior = NewObject<USingleClickInputBehavior>();
-	ClickBehavior->Initialize(this);
-	ClickBehavior->SetDefaultPriority(BasePriority);
-	ParentToolIn->AddInputBehavior(ClickBehavior, this);
-
 	HoverBehavior = NewObject<UMouseHoverBehavior>();
+	// We use modifiers on hover to change the highlighting according to what can be selected  	
 	HoverBehavior->Modifiers.RegisterModifier(ShiftModifierID, FInputDeviceState::IsShiftKeyDown);
 	HoverBehavior->Modifiers.RegisterModifier(CtrlModifierID, FInputDeviceState::IsCtrlKeyDown);
 	HoverBehavior->Initialize(this);
@@ -40,12 +36,20 @@ void UPolygonSelectionMechanic::Setup(UInteractiveTool* ParentToolIn)
 	ParentToolIn->AddInputBehavior(HoverBehavior, this);
 
 	MarqueeMechanic = NewObject<URectangleMarqueeMechanic>();
+	MarqueeMechanic->bUseExternalClickDragBehavior = true;
 	MarqueeMechanic->Setup(ParentToolIn);
 	MarqueeMechanic->OnDragRectangleStarted.AddUObject(this, &UPolygonSelectionMechanic::OnDragRectangleStarted);
 	MarqueeMechanic->OnDragRectangleChanged.AddUObject(this, &UPolygonSelectionMechanic::OnDragRectangleChanged);
 	MarqueeMechanic->OnDragRectangleFinished.AddUObject(this, &UPolygonSelectionMechanic::OnDragRectangleFinished);
 	MarqueeMechanic->SetBasePriority(BasePriority.MakeLower());
 
+	ClickOrDragBehavior = NewObject<USingleClickOrDragInputBehavior>();
+	ClickOrDragBehavior->Initialize(this, MarqueeMechanic);
+	ClickOrDragBehavior->Modifiers.RegisterModifier(ShiftModifierID, FInputDeviceState::IsShiftKeyDown);
+	ClickOrDragBehavior->Modifiers.RegisterModifier(CtrlModifierID, FInputDeviceState::IsCtrlKeyDown);
+	ClickOrDragBehavior->SetDefaultPriority(BasePriority);
+	ParentTool->AddInputBehavior(ClickOrDragBehavior);
+
 	Properties = NewObject<UPolygonSelectionMechanicProperties>(this);
 	if (bAddSelectionFilterPropertiesToParentTool)
 	{
@@ -178,9 +182,9 @@ void UPolygonSelectionMechanic::SetIsEnabled(bool bOn)
 void UPolygonSelectionMechanic::SetBasePriority(const FInputCapturePriority &Priority)
 {
 	BasePriority = Priority;
-	if (ClickBehavior)
+	if (ClickOrDragBehavior)
 	{
-		ClickBehavior->SetDefaultPriority(Priority);
+		ClickOrDragBehavior->SetDefaultPriority(Priority);
 	}
 	if (HoverBehavior)
 	{
@@ -341,6 +345,7 @@ FGroupTopologySelector::FSelectionSettings UPolygonSelectionMechanic::GetTopoSel
 
 	if (!PersistentSelection.IsEmpty() && (ShouldAddToSelectionFunc() || ShouldRemoveFromSelectionFunc()))
 	{
+		// If we have a selection and we're adding/removing/toggling elements make sure we only hit elements with compatible types
 		Settings.bEnableFaceHits = Settings.bEnableFaceHits && PersistentSelection.SelectedGroupIDs.Num() > 0;
 		Settings.bEnableEdgeHits = Settings.bEnableEdgeHits && PersistentSelection.SelectedEdgeIDs.Num() > 0;
 		Settings.bEnableCornerHits = Settings.bEnableCornerHits && PersistentSelection.SelectedCornerIDs.Num() > 0;
@@ -416,11 +421,12 @@ bool UPolygonSelectionMechanic::UpdateHighlight(const FRay& WorldRay)
 					// We use the triangle normals because the normal overlay isn't guaranteed to be valid as we edit the mesh
 					FVector3d TriangleNormal = Mesh->GetTriNormal(Tid);
 
+					// The UV's and colors here don't currently get used by HighlightedFaceMaterial, but we set them anyway
 					FIndex3i VertIndices = Mesh->GetTriangle(Tid);
 					DrawnTriangleSetComponent->AddTriangle(FRenderableTriangle(HighlightedFaceMaterial,
-						FRenderableTriangleVertex((FVector)Mesh->GetVertex(VertIndices.A), (FVector2D)Mesh->GetVertexUV(VertIndices.A), (FVector)TriangleNormal, UE::Geometry::ToFColor((FVector4f)Mesh->GetVertexColor(VertIndices.A))),
-						FRenderableTriangleVertex((FVector)Mesh->GetVertex(VertIndices.B), (FVector2D)Mesh->GetVertexUV(VertIndices.B), (FVector)TriangleNormal, UE::Geometry::ToFColor((FVector4f)Mesh->GetVertexColor(VertIndices.B))),
-						FRenderableTriangleVertex((FVector)Mesh->GetVertex(VertIndices.C), (FVector2D)Mesh->GetVertexUV(VertIndices.C), (FVector)TriangleNormal, UE::Geometry::ToFColor((FVector4f)Mesh->GetVertexColor(VertIndices.C))) ));
+						FRenderableTriangleVertex((FVector)Mesh->GetVertex(VertIndices.A), (FVector2D)Mesh->GetVertexUV(VertIndices.A), (FVector)TriangleNormal, FLinearColor(Mesh->GetVertexColor(VertIndices.A)).ToFColor(true)),
+						FRenderableTriangleVertex((FVector)Mesh->GetVertex(VertIndices.B), (FVector2D)Mesh->GetVertexUV(VertIndices.B), (FVector)TriangleNormal, FLinearColor(Mesh->GetVertexColor(VertIndices.B)).ToFColor(true)),
+						FRenderableTriangleVertex((FVector)Mesh->GetVertex(VertIndices.C), (FVector2D)Mesh->GetVertexUV(VertIndices.C), (FVector)TriangleNormal, FLinearColor(Mesh->GetVertexColor(VertIndices.C)).ToFColor(true)) ));
 				}
 
 				CurrentlyHighlightedGroups.Add(Gid);
@@ -446,7 +452,8 @@ bool UPolygonSelectionMechanic::UpdateSelection(const FRay& WorldRay, FVector3d&
 		TargetTransform.InverseTransformVector((FVector3d)WorldRay.Direction));
 	UE::Geometry::Normalize(LocalRay.Direction);
 
-	bool bSelectionModified = false;
+	const FGroupTopologySelection PreviousSelection = PersistentSelection;
+
 	FVector3d LocalPosition, LocalNormal;
 	FGroupTopologySelection Selection;
 	FGroupTopologySelector::FSelectionSettings TopoSelectorSettings = GetTopoSelectorSettings(CameraState.bIsOrthographic);
@@ -483,22 +490,20 @@ bool UPolygonSelectionMechanic::UpdateSelection(const FRay& WorldRay, FVector3d&
 		{
 			PersistentSelection = Selection;
 		}
-
-		bSelectionModified = true;
 	}
 	else
 	{
-		bSelectionModified = (PersistentSelection.IsEmpty() == false);
 		PersistentSelection.Clear();
 	}
 
-	if (bSelectionModified)
+	if (PersistentSelection != PreviousSelection)
 	{
 		SelectionTimestamp++;
 		OnSelectionChanged.Broadcast();
+		return true;
 	}
 
-	return bSelectionModified;
+	return false;
 }
 
 
@@ -534,11 +539,8 @@ FInputRayHit UPolygonSelectionMechanic::IsHitByClick(const FInputDeviceRay& Clic
 		return FInputRayHit(OutHit.Distance);
 	}
 
-	// We didn't hit selectable geometry. If we're using the marquee mechanic, give it a
-	// chance to capture. Otherwise still capture ourselves so that we are able to clear selections.
-	return MarqueeMechanic->IsEnabled() ?
-		FInputRayHit() // bHit is false
-		: FInputRayHit(TNumericLimits<float>::Max()); // bHit is true
+	// Return a hit so we always capture and can clear the selection
+	return FInputRayHit(TNumericLimits<float>::Max());
 }
 
 void UPolygonSelectionMechanic::OnClicked(const FInputDeviceRay& ClickPos)
@@ -600,8 +602,6 @@ void UPolygonSelectionMechanic::OnDragRectangleStarted()
 	PreDragPersistentSelection = PersistentSelection;
 	PreDragTopoSelectorSettings = GetTopoSelectorSettings(false);
 	PreDragTopoSelectorSettings.bIgnoreOcclusion = Properties->bMarqueeIgnoreOcclusion; // uses a separate setting for marquee
-	
-	PersistentSelection.Clear();
 }
 
 void UPolygonSelectionMechanic::OnDragRectangleChanged(const FCameraRectangle& CurrentRectangle)
@@ -637,7 +637,7 @@ void UPolygonSelectionMechanic::OnDragRectangleChanged(const FCameraRectangle& C
 
 void UPolygonSelectionMechanic::OnDragRectangleFinished(const FCameraRectangle& Rectangle, bool bCancelled)
 {
-	if (!(PersistentSelection == PreDragPersistentSelection))
+	if (PersistentSelection != PreDragPersistentSelection)
 	{
 		SelectionTimestamp++;
 		OnSelectionChanged.Broadcast();
diff --git a/Engine/Plugins/Experimental/MeshModelingToolsetExp/Source/ModelingComponents/Public/Mechanics/LatticeControlPointsMechanic.h b/Engine/Plugins/Experimental/MeshModelingToolsetExp/Source/ModelingComponents/Public/Mechanics/LatticeControlPointsMechanic.h
--- a/Engine/Plugins/Experimental/MeshModelingToolsetExp/Source/ModelingComponents/Public/Mechanics/LatticeControlPointsMechanic.h
+++ b/Engine/Plugins/Experimental/MeshModelingToolsetExp/Source/ModelingComponents/Public/Mechanics/LatticeControlPointsMechanic.h
@@ -187,10 +187,6 @@ protected:
 	FQuat	GizmoStartRotation;
 	FVector GizmoStartScale;
 
-	// These issue undo/redo change objects, and must therefore not be called in undo/redo code.
-	void ChangeSelection(int32 NewPointID, bool AddToSelection);
-	void ClearSelection();
-
 	// All of the following do not issue undo/redo change objects.
 	bool HitTest(const FInputDeviceRay& ClickPos, FInputRayHit& ResultOut);
 	void SelectPoint(int32 PointID);
diff --git a/Engine/Plugins/Experimental/MeshModelingToolsetExp/Source/ModelingComponents/Public/Selection/PolygonSelectionMechanic.h b/Engine/Plugins/Experimental/MeshModelingToolsetExp/Source/ModelingComponents/Public/Selection/PolygonSelectionMechanic.h
--- a/Engine/Plugins/Experimental/MeshModelingToolsetExp/Source/ModelingComponents/Public/Selection/PolygonSelectionMechanic.h
+++ b/Engine/Plugins/Experimental/MeshModelingToolsetExp/Source/ModelingComponents/Public/Selection/PolygonSelectionMechanic.h
@@ -19,7 +19,7 @@ class FPolygonSelectionMechanicSelectionChange;
 class UPersistentMeshSelection;
 class UMouseHoverBehavior;
 class URectangleMarqueeMechanic;
-class USingleClickInputBehavior;
+class USingleClickOrDragInputBehavior;
 
 using UE::Geometry::FDynamicMeshAABBTree3;
 PREDECLARE_USE_GEOMETRY_CLASS(FCompactMaps);
@@ -339,7 +339,7 @@ protected:
 	TObjectPtr<UMouseHoverBehavior> HoverBehavior;
 
 	UPROPERTY()
-	TObjectPtr<USingleClickInputBehavior> ClickBehavior;
+	TObjectPtr<USingleClickOrDragInputBehavior> ClickOrDragBehavior;
 
 	UPROPERTY()
 	TObjectPtr<URectangleMarqueeMechanic> MarqueeMechanic;
diff --git a/Engine/Plugins/Runtime/GeometryProcessing/Source/DynamicMesh/Public/GroupTopology.h b/Engine/Plugins/Runtime/GeometryProcessing/Source/DynamicMesh/Public/GroupTopology.h
--- a/Engine/Plugins/Runtime/GeometryProcessing/Source/DynamicMesh/Public/GroupTopology.h
+++ b/Engine/Plugins/Runtime/GeometryProcessing/Source/DynamicMesh/Public/GroupTopology.h
@@ -105,6 +105,7 @@ struct DYNAMICMESH_API FGroupTopologySelection
 		return SelectedCornerIDs.Num() == Selection.SelectedCornerIDs.Num()
 			&& SelectedEdgeIDs.Num() == Selection.SelectedEdgeIDs.Num()
 			&& SelectedGroupIDs.Num() == Selection.SelectedGroupIDs.Num()
+			// No need to also check Selection.Contains(*this) here since counts match and sets contain unique elements
 			&& Contains(Selection);
 	}
 
