ModelingTools: fix broken accept flow in GroupPaintTool, was baking in sculpt transform incorrectly
#rnx

diff --git a/Engine/Plugins/Experimental/MeshModelingToolsetExp/Source/MeshModelingToolsExp/Private/MeshGroupPaintTool.cpp b/Engine/Plugins/Experimental/MeshModelingToolsetExp/Source/MeshModelingToolsExp/Private/MeshGroupPaintTool.cpp
--- a/Engine/Plugins/Experimental/MeshModelingToolsetExp/Source/MeshModelingToolsExp/Private/MeshGroupPaintTool.cpp
+++ b/Engine/Plugins/Experimental/MeshModelingToolsetExp/Source/MeshModelingToolsExp/Private/MeshGroupPaintTool.cpp
@@ -253,26 +253,18 @@ void UMeshGroupPaintTool::Shutdown(EToolShutdownType ShutdownType)
 	FilterProperties->SaveProperties(this);
 	PolygroupLayerProperties->SaveProperties(this, TEXT("MeshGroupPaintTool"));
 
+	UMeshSculptToolBase::Shutdown(ShutdownType);
+}
 
-	// TODO
-	// TODO
-	// TODO
-	// Bake should not have to replace entire mesh just to update groups...
 
-	// do our own bake
-	if (ShutdownType == EToolShutdownType::Accept)
+void UMeshGroupPaintTool::CommitResult(UBaseDynamicMeshComponent* Component, bool bModifiedTopology)
+{
+	GetToolManager()->BeginUndoTransaction(LOCTEXT("GroupPaintToolTransactionName", "Paint Groups"));
+	Component->ProcessMesh([&](const FDynamicMesh3& CurMesh)
 	{
-		GetToolManager()->BeginUndoTransaction(LOCTEXT("GroupPaintMeshToolTransactionName", "Paint Groups"));
-		DynamicMeshComponent->ProcessMesh([&](const FDynamicMesh3& ReadMesh)
-		{
-			UE::ToolTarget::CommitDynamicMeshUpdate(Target, ReadMesh, true);
-		});
-		GetToolManager()->EndUndoTransaction();
-	}
-
-
-	// this call will unregister and destroy DynamicMeshComponent
-	UMeshSculptToolBase::Shutdown(EToolShutdownType::Completed);
+		UE::ToolTarget::CommitDynamicMeshUpdate(Target, CurMesh, true);
+	});
+	GetToolManager()->EndUndoTransaction();
 }
 
 
diff --git a/Engine/Plugins/Experimental/MeshModelingToolsetExp/Source/MeshModelingToolsExp/Public/MeshGroupPaintTool.h b/Engine/Plugins/Experimental/MeshModelingToolsetExp/Source/MeshModelingToolsExp/Public/MeshGroupPaintTool.h
--- a/Engine/Plugins/Experimental/MeshModelingToolsetExp/Source/MeshModelingToolsExp/Public/MeshGroupPaintTool.h
+++ b/Engine/Plugins/Experimental/MeshModelingToolsetExp/Source/MeshModelingToolsExp/Public/MeshGroupPaintTool.h
@@ -301,6 +301,9 @@ public:
 
 	bool IsInBrushSubMode() const;
 
+	virtual void CommitResult(UBaseDynamicMeshComponent* Component, bool bModifiedTopology) override;
+
+
 public:
 
 	UPROPERTY()
