fix crash on trying to sculpt after clicking 'discard attributes' in dynamic mesh sculpt tool

#rnx

diff --git a/Engine/Plugins/Experimental/MeshModelingToolsetExp/Source/MeshModelingToolsExp/Private/DynamicMeshSculptTool.cpp b/Engine/Plugins/Experimental/MeshModelingToolsetExp/Source/MeshModelingToolsExp/Private/DynamicMeshSculptTool.cpp
--- a/Engine/Plugins/Experimental/MeshModelingToolsetExp/Source/MeshModelingToolsExp/Private/DynamicMeshSculptTool.cpp
+++ b/Engine/Plugins/Experimental/MeshModelingToolsetExp/Source/MeshModelingToolsExp/Private/DynamicMeshSculptTool.cpp
@@ -270,13 +270,13 @@ void UDynamicMeshSculptTool::Setup()
 		if (bHaveUVSeams)
 		{
 			GetToolManager()->DisplayMessage(
-				LOCTEXT("UVSeamWarning", "This mesh has UV seams which may limit remeshing. Consider clearing the UV layers using the Remesh Tool."),
+				LOCTEXT("UVSeamWarning", "This mesh has UV seams which may limit remeshing. Consider clearing the UV layers using \"Discard Attributes\" or the Remesh Tool."),
 				EToolMessageLevel::UserWarning);
 		}
 		else if (bHaveNormalSeams)
 		{
 			GetToolManager()->DisplayMessage(
-				LOCTEXT("NormalSeamWarning", "This mesh has Hard Normal seams which may limit remeshing. Consider clearing Hard Normals using the Remesh Tool."),
+				LOCTEXT("NormalSeamWarning", "This mesh has Hard Normal seams which may limit remeshing. Consider clearing Hard Normals using \"Discard Attributes,\" or the Remesh or Normals Tool."),
 				EToolMessageLevel::UserWarning);
 		}
 	}
@@ -2072,6 +2072,10 @@ void UDynamicMeshSculptTool::RecalculateNormals_Overlay(const TSet<int32>& Trian
 			if (Mesh->IsTriangle(TriangleID))
 			{
 				FIndex3i TriElems = Normals->GetTriangle(TriangleID);
+				if (TriElems.A == FDynamicMesh3::InvalidID)
+				{
+					continue;
+				}
 				for (int j = 0; j < 3; ++j)
 				{
 					int elemid = TriElems[j];
@@ -2683,8 +2687,11 @@ void UDynamicMeshSculptTool::DiscardAttributes()
 {
 	TSharedPtr<FDynamicMesh3, ESPMode::ThreadSafe> BeforeMesh = MakeShared<FDynamicMesh3, ESPMode::ThreadSafe>(*DynamicMeshComponent->GetMesh());
 	TSharedPtr<FDynamicMesh3, ESPMode::ThreadSafe> AfterMesh = MakeShared<FDynamicMesh3, ESPMode::ThreadSafe>(*BeforeMesh);
+	// Reset attributes and compute per-vertex normals
+	// (we need to clear rather than permanently discard the attributes because UDynamicMesh::EditMeshInternal would automatically add them back anyway)
 	AfterMesh->DiscardAttributes();
-	FMeshNormals::QuickComputeVertexNormals(*AfterMesh);
+	AfterMesh->EnableAttributes();
+	FMeshNormals::InitializeOverlayToPerVertexNormals(AfterMesh->Attributes()->PrimaryNormals(), false);
 
 	TUniquePtr<FMeshReplacementChange> ReplaceChange = MakeUnique<FMeshReplacementChange>(BeforeMesh, AfterMesh);
 	DynamicMeshComponent->ApplyChange(ReplaceChange.Get(), false);
@@ -2692,6 +2699,9 @@ void UDynamicMeshSculptTool::DiscardAttributes()
 	// other internals are still valid?
 
 	GetToolManager()->EmitObjectChange(DynamicMeshComponent, MoveTemp(ReplaceChange), LOCTEXT("SculptDiscardAttribsChange", "Discard Attributes"));
+
+	// The tool's warning messages about seams are no longer relevant after discarding attributes, so clear warning messages
+	GetToolManager()->DisplayMessage(FText(), EToolMessageLevel::UserWarning);
 }
 
 #undef LOCTEXT_NAMESPACE
diff --git a/Engine/Source/Runtime/GeometryCore/Private/DynamicMesh/DynamicMeshOverlay.cpp b/Engine/Source/Runtime/GeometryCore/Private/DynamicMesh/DynamicMeshOverlay.cpp
--- a/Engine/Source/Runtime/GeometryCore/Private/DynamicMesh/DynamicMeshOverlay.cpp
+++ b/Engine/Source/Runtime/GeometryCore/Private/DynamicMesh/DynamicMeshOverlay.cpp
@@ -672,6 +672,17 @@ bool TDynamicMeshOverlay<RealType, ElementSize>::HasInteriorSeamEdges() const
 		FIndex2i et = ParentMesh->GetEdgeT(eid);
 		if (et.B != FDynamicMesh3::InvalidID)
 		{
+			bool bASet = IsSetTriangle(et.A), bBSet = IsSetTriangle(et.B);
+			if (bASet != bBSet)
+			{
+				// seam between triangles with elements and triangles without
+				return true;
+			}
+			else if (!bASet)
+			{
+				// neither triangle has set elements
+				continue;
+			}
 			FIndex2i ev = ParentMesh->GetEdgeV(eid);
 			int base_a = ev.A, base_b = ev.B;
 
