Add brute force Marquee selection for Vertices and Edges to UV Editor

#rnx

diff --git a/Engine/Plugins/Experimental/MeshModelingToolsetExp/Source/ModelingComponents/Private/Selection/MeshSelectionMechanic.cpp b/Engine/Plugins/Experimental/MeshModelingToolsetExp/Source/ModelingComponents/Private/Selection/MeshSelectionMechanic.cpp
--- a/Engine/Plugins/Experimental/MeshModelingToolsetExp/Source/ModelingComponents/Private/Selection/MeshSelectionMechanic.cpp
+++ b/Engine/Plugins/Experimental/MeshModelingToolsetExp/Source/ModelingComponents/Private/Selection/MeshSelectionMechanic.cpp
@@ -3,7 +3,7 @@
 #include "Selection/MeshSelectionMechanic.h"
 
 #include "BaseBehaviors/BehaviorTargetInterfaces.h"
-#include "BaseBehaviors/SingleClickBehavior.h"
+#include "BaseBehaviors/SingleClickOrDragBehavior.h"
 #include "Drawing/LineSetComponent.h"
 #include "Drawing/PointSetComponent.h"
 #include "Drawing/PreviewGeometryActor.h"
@@ -73,17 +73,44 @@ namespace MeshSelectionMechanicLocals
 			Set.Add(Item);
 		}
 	}
-}
+	
+	UE::Geometry::FDynamicMeshSelection::EType ToCompatibleDynamicMeshSelectionType(const EMeshSelectionMechanicMode& Mode)
+	{
+		switch (Mode)
+		{
+			case EMeshSelectionMechanicMode::Mesh:
+			case EMeshSelectionMechanicMode::Component:
+			case EMeshSelectionMechanicMode::Triangle:
+				return UE::Geometry::FDynamicMeshSelection::EType::Triangle;
+			case EMeshSelectionMechanicMode::Edge:
+				return UE::Geometry::FDynamicMeshSelection::EType::Edge;
+			case EMeshSelectionMechanicMode::Vertex:
+				return UE::Geometry::FDynamicMeshSelection::EType::Vertex;
+		}
+		checkNoEntry();
+		return UE::Geometry::FDynamicMeshSelection::EType::Vertex;
+	}
+} // namespace MeshSelectionMechanicLocals
 
 void UMeshSelectionMechanic::Setup(UInteractiveTool* ParentToolIn)
 {
 	UInteractionMechanic::Setup(ParentToolIn);
 
-	USingleClickInputBehavior* ClickBehavior = NewObject<USingleClickInputBehavior>();
-	ClickBehavior->Modifiers.RegisterModifier(CtrlModifierID, FInputDeviceState::IsCtrlKeyDown);
-	ClickBehavior->Modifiers.RegisterModifier(ShiftModifierID, FInputDeviceState::IsShiftKeyDown);
-	ClickBehavior->Initialize(this);
-	ParentTool->AddInputBehavior(ClickBehavior);
+	// This will be the target for the click drag behavior below
+	MarqueeMechanic = NewObject<URectangleMarqueeMechanic>();
+	MarqueeMechanic->bUseExternalClickDragBehavior = true;
+	MarqueeMechanic->Setup(ParentToolIn);
+	MarqueeMechanic->OnDragRectangleStarted.AddUObject(this, &UMeshSelectionMechanic::OnDragRectangleStarted);
+	MarqueeMechanic->OnDragRectangleChanged.AddUObject(this, &UMeshSelectionMechanic::OnDragRectangleChanged);
+	MarqueeMechanic->OnDragRectangleFinished.AddUObject(this, &UMeshSelectionMechanic::OnDragRectangleFinished);
+
+	USingleClickOrDragInputBehavior* ClickOrDragBehavior = NewObject<USingleClickOrDragInputBehavior>();
+	ClickOrDragBehavior->Initialize(this, MarqueeMechanic);
+	ClickOrDragBehavior->Modifiers.RegisterModifier(ShiftModifierID, FInputDeviceState::IsShiftKeyDown);
+	ClickOrDragBehavior->Modifiers.RegisterModifier(CtrlModifierID, FInputDeviceState::IsCtrlKeyDown);
+	ParentTool->AddInputBehavior(ClickOrDragBehavior);
+
+	ClearCurrentSelection();
 
 	LineSet = NewObject<ULineSetComponent>();
 	LineSet->SetLineMaterial(ToolSetupUtil::GetDefaultLineComponentMaterial(
@@ -326,206 +353,302 @@ void UMeshSelectionMechanic::SetDrawnElementsTransform(const FTransform& Transfo
 
 void UMeshSelectionMechanic::Render(IToolsContextRenderAPI* RenderAPI)
 {
-	// TODO:We do this in other places, but should CameraState be cached somewhere else?
+	// Cache the camera state
+	MarqueeMechanic->Render(RenderAPI);
 	GetParentTool()->GetToolManager()->GetContextQueriesAPI()->GetCurrentViewState(CameraState);
+}
 
-	if (CurrentSelection.IsEmpty())
-	{
-		return;
-	}
-
-
+void UMeshSelectionMechanic::DrawHUD(FCanvas* Canvas, IToolsContextRenderAPI* RenderAPI)
+{
+	MarqueeMechanic->DrawHUD(Canvas, RenderAPI);
 }
 
 FInputRayHit UMeshSelectionMechanic::IsHitByClick(const FInputDeviceRay& ClickPos)
 {
-	// We grab any clicks we are offered because even if we miss the mesh, we want to be able to clear selection.
+	// TODO We could account for what modifiers would do and return an accurate depth here but for now this simple code works fine
+
+	// Return a hit so we always capture and can clear the selection
 	FInputRayHit Hit; 
 	Hit.bHit = true;
-	return Hit;
+	return Hit;	
 }
 
-void UMeshSelectionMechanic::OnClicked(const FInputDeviceRay& ClickPos)
+void UMeshSelectionMechanic::ClearCurrentSelection()
 {
-	FDynamicMeshSelection OriginalSelection = CurrentSelection;
+	CurrentSelection.SelectedIDs.Empty();
+	CurrentSelection.Mesh = nullptr;
+	CurrentSelection.Type = MeshSelectionMechanicLocals::ToCompatibleDynamicMeshSelectionType(SelectionMode);
+	CurrentSelectionIndex = IndexConstants::InvalidID;
+}
 
-	const bool bIncompatibleSelectionMode =
-		(SelectionMode == EMeshSelectionMechanicMode::Component && CurrentSelection.Type != FDynamicMeshSelection::EType::Triangle) ||
-		(SelectionMode == EMeshSelectionMechanicMode::Edge && CurrentSelection.Type != FDynamicMeshSelection::EType::Edge) ||
-		(SelectionMode == EMeshSelectionMechanicMode::Vertex && CurrentSelection.Type != FDynamicMeshSelection::EType::Vertex) || 
-		(SelectionMode == EMeshSelectionMechanicMode::Triangle && CurrentSelection.Type != FDynamicMeshSelection::EType::Triangle) || 
-		(SelectionMode == EMeshSelectionMechanicMode::Mesh && CurrentSelection.Type != FDynamicMeshSelection::EType::Triangle);
-	if (bIncompatibleSelectionMode || ShouldRestartSelection())
+void UMeshSelectionMechanic::UpdateCurrentSelection(const TSet<int32>& NewSelection, bool CalledFromOnDragRectangleChanged)
+{
+	if (!ShouldRestartSelection() && CalledFromOnDragRectangleChanged)
 	{
-		CurrentSelection.SelectedIDs.Empty();
-		CurrentSelection.Mesh = nullptr;
+		// If we're modifying (adding/removing/toggling, but not restarting) the selection, we should start with the
+		// selection cached in OnDragRectangleStarted, otherwise multiple changes get accumulated as the rectangle is
+		// swept around
+		CurrentSelection.SelectedIDs = PreDragSelection.SelectedIDs;
 	}
-
-	int32 HitTid = IndexConstants::InvalidID;
-	double RayT = 0; 
-	for (int32 i = 0; i < MeshSpatials.Num(); ++i)
+	
+	if (ShouldRestartSelection())
+	{
+		CurrentSelection.SelectedIDs = NewSelection;
+	}
+	else if (ShouldAddToSelection())
 	{
-		// We only add/toggle/remove from selections in the currently selected mesh. If this isn't the case or if we're
-		// removing from a selection but there is nothing to remove, then we can break early.
-		if ((!ShouldRestartSelection() && i != CurrentSelectionIndex) ||
-			(ShouldRemoveFromSelection() && CurrentSelection.IsEmpty()))
+		CurrentSelection.SelectedIDs.Append(NewSelection);
+	}
+	else if (ShouldToggleFromSelection())
+	{
+		for (int32 Index : NewSelection)
 		{
-			continue;
+			MeshSelectionMechanicLocals::ToggleItem(CurrentSelection.SelectedIDs, Index);
 		}
+	}
+	else if (ShouldRemoveFromSelection())
+	{
+		CurrentSelection.SelectedIDs = CurrentSelection.SelectedIDs.Difference(NewSelection);
+	}
+	else
+	{
+		checkNoEntry();
+	}
 
-		// Minor note: Using FRay3d in this statement instead of FRay causes an internal compiler error in VS
-		// in Development Editor builds, for no good reason.
-		FRay LocalRay(
-			MeshTransforms[i].InverseTransformPosition(ClickPos.WorldRay.Origin),
-			MeshTransforms[i].InverseTransformVector(ClickPos.WorldRay.Direction));
+	if (CurrentSelection.SelectedIDs.IsEmpty())
+	{
+		ClearCurrentSelection();
+	}
+	
+	CurrentSelection.Type = MeshSelectionMechanicLocals::ToCompatibleDynamicMeshSelectionType(SelectionMode);
+}
 
-		if (MeshSpatials[i]->FindNearestHitTriangle(LocalRay, RayT, HitTid))
-		{
-			// It should be okay to reassign this. If we're in multiselect, and it misses a hit, this just won't happen and selection won't be altered.
-			CurrentSelection.Mesh = MeshSpatials[i]->GetMesh();
-			CurrentSelection.TopologyTimestamp = CurrentSelection.Mesh->GetTopologyChangeStamp();
-			CurrentSelectionIndex = i;
-			break;
-		}
+void UMeshSelectionMechanic::OnClicked(const FInputDeviceRay& ClickPos)
+{
+	FDynamicMeshSelection OriginalSelection = CurrentSelection;
+	
+	if (CurrentSelection.Type != MeshSelectionMechanicLocals::ToCompatibleDynamicMeshSelectionType(SelectionMode))
+	{
+		ClearCurrentSelection();
 	}
 
-	if (HitTid != IndexConstants::InvalidID)
+	TSet<int32> ClickSelectedIDs; // TODO Maybe using a TVariant would have been better here...
 	{
-		if (SelectionMode == EMeshSelectionMechanicMode::Component)
+		int32 HitTid = IndexConstants::InvalidID;
+		for (int32 MeshIndex = 0; MeshIndex < MeshSpatials.Num(); ++MeshIndex)
 		{
-			FMeshConnectedComponents MeshSelectedComponent(CurrentSelection.Mesh);
-			TArray<int32> SeedTriangles;
-			SeedTriangles.Add(HitTid);
-			MeshSelectedComponent.FindTrianglesConnectedToSeeds(SeedTriangles);
-			
-			if (ShouldAddToSelection() || ShouldRestartSelection())
-			{
-				CurrentSelection.SelectedIDs.Append(TSet<int>(MeshSelectedComponent.Components[0].Indices));
-			}
-			else if (ShouldRemoveFromSelection())
-			{
-				CurrentSelection.SelectedIDs = CurrentSelection.SelectedIDs.Difference(TSet<int>(MeshSelectedComponent.Components[0].Indices));
-			}
-			else if (ShouldToggleFromSelection())
+			FRay LocalRay(
+				MeshTransforms[MeshIndex].InverseTransformPosition(ClickPos.WorldRay.Origin),
+				MeshTransforms[MeshIndex].InverseTransformVector(ClickPos.WorldRay.Direction));
+
+			double RayT = 0; 
+			if (MeshSpatials[MeshIndex]->FindNearestHitTriangle(LocalRay, RayT, HitTid))
 			{
-				for (int I : MeshSelectedComponent.Components[0].Indices) {
-					MeshSelectionMechanicLocals::ToggleItem(CurrentSelection.SelectedIDs, I);
-				}
+				CurrentSelection.Mesh = MeshSpatials[MeshIndex]->GetMesh();
+				CurrentSelection.TopologyTimestamp = CurrentSelection.Mesh->GetTopologyChangeStamp();
+				CurrentSelectionIndex = MeshIndex;
+				break;
 			}
-			CurrentSelection.Type = FDynamicMeshSelection::EType::Triangle;
 		}
-		// TODO: We'll need the ability to hit occluded triangles to see if there is a better edge
-		// to snap to.
-		else if (SelectionMode == EMeshSelectionMechanicMode::Edge)
-		{
-			// Try to snap to one of the edges.
-			FIndex3i Eids = CurrentSelection.Mesh->GetTriEdges(HitTid);
 
-			FGeometrySet3 GeometrySet;
-			for (int i = 0; i < 3; ++i)
+		if (HitTid != IndexConstants::InvalidID)
+		{
+			if (SelectionMode == EMeshSelectionMechanicMode::Component)
 			{
-				FIndex2i Vids = CurrentSelection.Mesh->GetEdgeV(Eids[i]);
-				FPolyline3d Polyline(CurrentSelection.Mesh->GetVertex(Vids.A), CurrentSelection.Mesh->GetVertex(Vids.B));
-				GeometrySet.AddCurve(Eids[i], Polyline);
+				FMeshConnectedComponents MeshSelectedComponent(CurrentSelection.Mesh);
+				TArray<int32> SeedTriangles;
+				SeedTriangles.Add(HitTid);
+				MeshSelectedComponent.FindTrianglesConnectedToSeeds(SeedTriangles);
+				ClickSelectedIDs = TSet<int32>(SeedTriangles);	
 			}
-			FGeometrySet3::FNearest Result;
-			if (GeometrySet.FindNearestCurveToRay(ClickPos.WorldRay, Result, 
-				[this](const FVector3d& Position1, const FVector3d& Position2) {
-					return ToolSceneQueriesUtil::PointSnapQuery(CameraState,
-						Position1, Position2,
-						ToolSceneQueriesUtil::GetDefaultVisualAngleSnapThreshD()); }))
+			else if (SelectionMode == EMeshSelectionMechanicMode::Edge)
 			{
-				if (ShouldAddToSelection() || ShouldRestartSelection())
-				{
-					CurrentSelection.SelectedIDs.Add(Result.ID);
-				}
-				else if (ShouldRemoveFromSelection())
+				// TODO: We'll need the ability to hit occluded triangles to see if there is a better edge to snap to.
+
+				// Try to snap to one of the edges.
+				FIndex3i Eids = CurrentSelection.Mesh->GetTriEdges(HitTid);
+
+				FGeometrySet3 GeometrySet;
+				for (int i = 0; i < 3; ++i)
 				{
-					CurrentSelection.SelectedIDs.Remove(Result.ID);
+					FIndex2i Vids = CurrentSelection.Mesh->GetEdgeV(Eids[i]);
+					FPolyline3d Polyline(CurrentSelection.Mesh->GetVertex(Vids.A), CurrentSelection.Mesh->GetVertex(Vids.B));
+					GeometrySet.AddCurve(Eids[i], Polyline);
 				}
-				else if (ShouldToggleFromSelection())
+				FGeometrySet3::FNearest Result;
+				if (GeometrySet.FindNearestCurveToRay(ClickPos.WorldRay, Result, 
+					[this](const FVector3d& Position1, const FVector3d& Position2) {
+						return ToolSceneQueriesUtil::PointSnapQuery(CameraState,
+							Position1, Position2,
+							ToolSceneQueriesUtil::GetDefaultVisualAngleSnapThreshD()); }))
 				{
-					MeshSelectionMechanicLocals::ToggleItem(CurrentSelection.SelectedIDs, Result.ID);
+					ClickSelectedIDs = TSet<int32>{Result.ID};
 				}
-				CurrentSelection.Type = FDynamicMeshSelection::EType::Edge;
 			}
-		}
-		else if (SelectionMode == EMeshSelectionMechanicMode::Vertex)
-		{
-			// TODO: Improve this to handle super narrow, sliver triangles better, where testing near vertices can be difficult.
+			else if (SelectionMode == EMeshSelectionMechanicMode::Vertex)
+			{
+				// TODO: Improve this to handle super narrow, sliver triangles better, where testing near vertices can be difficult.
 
-			// Try to snap to one of the vertices
-			FIndex3i Vids = CurrentSelection.Mesh->GetTriangle(HitTid);
+				// Try to snap to one of the vertices
+				FIndex3i Vids = CurrentSelection.Mesh->GetTriangle(HitTid);
 
-			FGeometrySet3 GeometrySet;
-			for (int i = 0; i < 3; ++i)
-			{
-				GeometrySet.AddPoint(Vids[i], CurrentSelection.Mesh->GetTriVertex(HitTid, i));
-			}
-			FGeometrySet3::FNearest Result;
-			if (GeometrySet.FindNearestPointToRay(ClickPos.WorldRay, Result,
-				[this](const FVector3d& Position1, const FVector3d& Position2) {
-					return ToolSceneQueriesUtil::PointSnapQuery(CameraState,
-						Position1, Position2,
-						ToolSceneQueriesUtil::GetDefaultVisualAngleSnapThreshD()); }))
-			{
-				if (ShouldAddToSelection() || ShouldRestartSelection())
-				{
-					CurrentSelection.SelectedIDs.Add(Result.ID);
-				}
-				else if (ShouldRemoveFromSelection())
+				FGeometrySet3 GeometrySet;
+				for (int i = 0; i < 3; ++i)
 				{
-					CurrentSelection.SelectedIDs.Remove(Result.ID);
+					GeometrySet.AddPoint(Vids[i], CurrentSelection.Mesh->GetTriVertex(HitTid, i));
 				}
-				else if (ShouldToggleFromSelection())
+				FGeometrySet3::FNearest Result;
+				if (GeometrySet.FindNearestPointToRay(ClickPos.WorldRay, Result,
+					[this](const FVector3d& Position1, const FVector3d& Position2) {
+						return ToolSceneQueriesUtil::PointSnapQuery(CameraState,
+							Position1, Position2,
+							ToolSceneQueriesUtil::GetDefaultVisualAngleSnapThreshD()); }))
 				{
-					MeshSelectionMechanicLocals::ToggleItem(CurrentSelection.SelectedIDs, Result.ID);
+					ClickSelectedIDs = TSet<int32>{Result.ID};
 				}
-				CurrentSelection.Type = FDynamicMeshSelection::EType::Vertex;
 			}
-		}
-		else if (SelectionMode == EMeshSelectionMechanicMode::Triangle)
-		{
-			if (ShouldAddToSelection() || ShouldRestartSelection())
+			else if (SelectionMode == EMeshSelectionMechanicMode::Triangle)
 			{
-				CurrentSelection.SelectedIDs.Add(HitTid);
+				ClickSelectedIDs = TSet<int32>{HitTid};
 			}
-			else if (ShouldRemoveFromSelection())
+			else if (SelectionMode == EMeshSelectionMechanicMode::Mesh)
 			{
-				CurrentSelection.SelectedIDs.Remove(HitTid);
+				for (int32 Tid : CurrentSelection.Mesh->TriangleIndicesItr())
+				{
+					ClickSelectedIDs.Add(Tid);
+				}
 			}
-			else if (ShouldToggleFromSelection())
+			else
 			{
-				MeshSelectionMechanicLocals::ToggleItem(CurrentSelection.SelectedIDs, HitTid);
+				checkNoEntry();
 			}
-			CurrentSelection.Type = FDynamicMeshSelection::EType::Triangle;
 		}
-		else if (SelectionMode == EMeshSelectionMechanicMode::Mesh)
+	}
+
+	// TODO Perhaps selection clearing should happen only if the click occurs further than some threshold from the meshes
+	UpdateCurrentSelection(ClickSelectedIDs);
+
+	if (OriginalSelection != CurrentSelection)
+	{
+		UpdateCentroid();
+		RebuildDrawnElements(FTransform(GetCurrentSelectionCentroid()));
+		EmitSelectionChange(OriginalSelection, CurrentSelection, true);
+		OnSelectionChanged.Broadcast();
+	}
+}
+
+void UMeshSelectionMechanic::OnDragRectangleStarted()
+{
+	if (CurrentSelection.Type != MeshSelectionMechanicLocals::ToCompatibleDynamicMeshSelectionType(SelectionMode))
+	{
+		ClearCurrentSelection();	
+	}
+	
+	PreDragSelection = CurrentSelection;
+}
+
+void UMeshSelectionMechanic::OnDragRectangleChanged(const FCameraRectangle& CurrentRectangle)
+{
+	auto FindIDsInsideCurrentRectangle = [this](auto&& AddIDsToSelection)
+	{
+		TSet<int32> RectangleSelectedIDs;
+		
+		if (CurrentSelectionIndex == IndexConstants::InvalidID)
 		{
-			for (int32 Tid : CurrentSelection.Mesh->TriangleIndicesItr())
+			for (int32 MeshIndex = 0; MeshIndex < MeshSpatials.Num(); ++MeshIndex)
 			{
-				if (ShouldAddToSelection() || ShouldRestartSelection())
+				const FDynamicMesh3& Mesh = *(MeshSpatials[MeshIndex]->GetMesh());
+				AddIDsToSelection(Mesh, MeshTransforms[MeshIndex], RectangleSelectedIDs);
+	
+				if (!RectangleSelectedIDs.IsEmpty())
 				{
-					CurrentSelection.SelectedIDs.Add(Tid);
+					// Pick the first selectable mesh for now, maybe we should try to be smarter
+					CurrentSelectionIndex = MeshIndex;
+					CurrentSelection.Mesh = MeshSpatials[CurrentSelectionIndex]->GetMesh();
+					break;
 				}
-				else if (ShouldRemoveFromSelection())
+			}
+		}
+		else
+		{
+			ensure(CurrentSelectionIndex != IndexConstants::InvalidID);
+			ensure(CurrentSelection.Mesh == MeshSpatials[CurrentSelectionIndex]->GetMesh());
+			
+			const FDynamicMesh3& Mesh = *(MeshSpatials[CurrentSelectionIndex]->GetMesh());
+			AddIDsToSelection(Mesh, MeshTransforms[CurrentSelectionIndex], RectangleSelectedIDs);
+		}
+
+		return RectangleSelectedIDs;
+	};
+	
+	TSet<int32> RectangleSelectedIDs;
+	if (SelectionMode == EMeshSelectionMechanicMode::Vertex)
+	{
+		auto AddVertexIDsToSelection = [this, &CurrentRectangle]
+			(const FDynamicMesh3& Mesh, const FTransform& Transform, TSet<int32>& OutSelectedIDs)
+		{
+			for (int32 VertexIndex : Mesh.VertexIndicesItr())
+			{
+				FVector3d Point = Transform.InverseTransformPosition(Mesh.GetVertex(VertexIndex));
+				if (CurrentRectangle.IsProjectedPointInRectangle(Point))
 				{
-					CurrentSelection.SelectedIDs.Remove(Tid);
+					OutSelectedIDs.Add(VertexIndex);
 				}
-				else if (ShouldToggleFromSelection())
+			}
+		};
+	
+		RectangleSelectedIDs = FindIDsInsideCurrentRectangle(AddVertexIDsToSelection);
+	}
+	else if (SelectionMode == EMeshSelectionMechanicMode::Edge)
+	{
+		auto AddEdgeIDsToSelection = [this, &CurrentRectangle]
+			(const FDynamicMesh3& Mesh, const FTransform& Transform, TSet<int32>& OutSelectedIDs)
+		{
+			for (int32 EdgeIndex : Mesh.EdgeIndicesItr())
+			{
+				const FDynamicMesh3::FEdge& Edge = Mesh.GetEdgeRef(EdgeIndex);
+				FVector3d StartPoint = Transform.InverseTransformPosition(Mesh.GetVertex(Edge.Vert.A));
+				FVector3d EndPoint = Transform.InverseTransformPosition(Mesh.GetVertex(Edge.Vert.B));
+   				
+				if (CurrentRectangle.IsProjectedSegmentIntersectingRectangle(StartPoint, EndPoint))
 				{
-					MeshSelectionMechanicLocals::ToggleItem(CurrentSelection.SelectedIDs, Tid);
+					OutSelectedIDs.Add(EdgeIndex);
 				}
 			}
-			CurrentSelection.Type = FDynamicMeshSelection::EType::Triangle;
-		}
+		};
+
+		RectangleSelectedIDs = FindIDsInsideCurrentRectangle(AddEdgeIDsToSelection);
+	}
+	else if (SelectionMode == EMeshSelectionMechanicMode::Triangle)
+	{
+		// TODO Implement me
+	}
+	else if (SelectionMode == EMeshSelectionMechanicMode::Component)
+	{
+		// TODO Implement me
+	}
+	else if (SelectionMode == EMeshSelectionMechanicMode::Mesh)
+	{
+		// TODO Implement me
+	}
+	else
+	{
+		checkNoEntry();
 	}
 
-	if (OriginalSelection != CurrentSelection)
+	UpdateCurrentSelection(RectangleSelectedIDs, true);
+	
+	UpdateCentroid();
+	RebuildDrawnElements(FTransform(GetCurrentSelectionCentroid()));
+}
+
+void UMeshSelectionMechanic::OnDragRectangleFinished(const FCameraRectangle& Rectangle, bool bCancelled)
+{
+	if (PreDragSelection != CurrentSelection)
 	{
 		UpdateCentroid();
 		RebuildDrawnElements(FTransform(GetCurrentSelectionCentroid()));
-		EmitSelectionChange(OriginalSelection, CurrentSelection, true);
+		EmitSelectionChange(PreDragSelection, CurrentSelection, true);
 		OnSelectionChanged.Broadcast();
 	}
 }
diff --git a/Engine/Plugins/Experimental/MeshModelingToolsetExp/Source/ModelingComponents/Public/Selection/MeshSelectionMechanic.h b/Engine/Plugins/Experimental/MeshModelingToolsetExp/Source/ModelingComponents/Public/Selection/MeshSelectionMechanic.h
--- a/Engine/Plugins/Experimental/MeshModelingToolsetExp/Source/ModelingComponents/Public/Selection/MeshSelectionMechanic.h
+++ b/Engine/Plugins/Experimental/MeshModelingToolsetExp/Source/ModelingComponents/Public/Selection/MeshSelectionMechanic.h
@@ -10,6 +10,7 @@
 #include "InteractionMechanic.h"
 #include "InteractiveTool.h"
 #include "Selection/DynamicMeshSelection.h"
+#include "Mechanics/RectangleMarqueeMechanic.h"
 #include "ToolContextInterfaces.h" //FViewCameraState
 
 #include "MeshSelectionMechanic.generated.h"
@@ -17,6 +18,7 @@
 class ULineSetComponent;
 class UPointSetComponent;
 class APreviewGeometryActor;
+struct FCameraRectangle;
 
 enum class EMeshSelectionMechanicMode
 {
@@ -78,7 +80,11 @@ public:
 	virtual void SetSelection(const FDynamicMeshSelection& Selection, bool bBroadcast = false, bool bEmitChange = false);
 
 	virtual void Render(IToolsContextRenderAPI* RenderAPI) override;
-	//virtual void DrawHUD(FCanvas* Canvas, IToolsContextRenderAPI* RenderAPI); // Will need once we have marquee
+	virtual void DrawHUD(FCanvas* Canvas, IToolsContextRenderAPI* RenderAPI);
+
+	void OnDragRectangleStarted();
+	void OnDragRectangleChanged(const FCameraRectangle& CurrentRectangle);
+	void OnDragRectangleFinished(const FCameraRectangle& Rectangle, bool bCancelled);
 
 	// IClickBehaviorTarget implementation
 	virtual FInputRayHit IsHitByClick(const FInputDeviceRay& ClickPos) override;
@@ -109,6 +115,9 @@ protected:
 	bool ShouldToggleFromSelection() const { return bCtrlToggle && bShiftToggle; }
 	bool ShouldRestartSelection() const { return !bCtrlToggle && !bShiftToggle; }
 
+	UPROPERTY()
+	TObjectPtr<URectangleMarqueeMechanic> MarqueeMechanic;
+
 	UPROPERTY()
 	TObjectPtr<APreviewGeometryActor> PreviewGeometryActor = nullptr;
 
@@ -122,6 +131,7 @@ protected:
 	TArray<TSharedPtr<FDynamicMeshAABBTree3>> MeshSpatials;
 	TArray<FTransform> MeshTransforms;
 	FDynamicMeshSelection CurrentSelection;
+	FDynamicMeshSelection PreDragSelection;
 	int32 CurrentSelectionIndex = IndexConstants::InvalidID;
 	FViewCameraState CameraState;
 
@@ -140,5 +150,9 @@ protected:
 	FVector3d CurrentSelectionCentroid;
 	int32 CentroidTimestamp = -1;
 	void UpdateCentroid();
+
+private:
+	void ClearCurrentSelection();
+	void UpdateCurrentSelection(const TSet<int32>& NewSelection, bool CalledFromOnRectangleChanged = false);
 };
 
diff --git a/Engine/Plugins/Experimental/UVEditor/Source/UVEditorTools/Private/UVSelectTool.cpp b/Engine/Plugins/Experimental/UVEditor/Source/UVEditorTools/Private/UVSelectTool.cpp
--- a/Engine/Plugins/Experimental/UVEditor/Source/UVEditorTools/Private/UVSelectTool.cpp
+++ b/Engine/Plugins/Experimental/UVEditor/Source/UVEditorTools/Private/UVSelectTool.cpp
@@ -739,6 +739,11 @@ void UUVSelectTool::Render(IToolsContextRenderAPI* RenderAPI)
 	SelectionMechanic->Render(RenderAPI);
 }
 
+void UUVSelectTool::DrawHUD(FCanvas* Canvas, IToolsContextRenderAPI* RenderAPI)
+{
+	SelectionMechanic->DrawHUD(Canvas, RenderAPI);
+}
+
 void UUVSelectTool::OnTick(float DeltaTime)
 {
 	ApplyGizmoTransform();
diff --git a/Engine/Plugins/Experimental/UVEditor/Source/UVEditorTools/Public/UVSelectTool.h b/Engine/Plugins/Experimental/UVEditor/Source/UVEditorTools/Public/UVSelectTool.h
--- a/Engine/Plugins/Experimental/UVEditor/Source/UVEditorTools/Public/UVSelectTool.h
+++ b/Engine/Plugins/Experimental/UVEditor/Source/UVEditorTools/Public/UVSelectTool.h
@@ -52,10 +52,10 @@ protected:
 UENUM()
 enum class EUVSelectToolSelectionMode : uint8
 {
-	Island,
-	Edge,
 	Vertex,
+	Edge,
 	Triangle,
+	Island,
 	Mesh
 };
 
@@ -136,6 +136,7 @@ public:
 	virtual void Shutdown(EToolShutdownType ShutdownType) override;
 	virtual void OnPropertyModified(UObject* PropertySet, FProperty* Property) override;
 	virtual void Render(IToolsContextRenderAPI* RenderAPI) override;
+	virtual void DrawHUD(FCanvas* Canvas, IToolsContextRenderAPI* RenderAPI) override;
 	virtual void OnTick(float DeltaTime) override;
 	virtual bool HasCancel() const override { return false; }
 	virtual bool HasAccept() const override { return false; }
