ModelingTools: Fix color texture output color space to correct asset previews.

#rnx

diff --git a/Engine/Plugins/Experimental/MeshModelingToolsetExp/Source/MeshModelingToolsExp/Private/BakeMeshAttributeMapsTool.cpp b/Engine/Plugins/Experimental/MeshModelingToolsetExp/Source/MeshModelingToolsExp/Private/BakeMeshAttributeMapsTool.cpp
--- a/Engine/Plugins/Experimental/MeshModelingToolsetExp/Source/MeshModelingToolsExp/Private/BakeMeshAttributeMapsTool.cpp
+++ b/Engine/Plugins/Experimental/MeshModelingToolsetExp/Source/MeshModelingToolsExp/Private/BakeMeshAttributeMapsTool.cpp
@@ -866,9 +866,6 @@ EBakeOpState UBakeMeshAttributeMapsTool::UpdateResult_Texture2DImage()
 			GetToolManager()->DisplayMessage(LOCTEXT("CannotReadTextureWarning", "Cannot read from the source texture"), EToolMessageLevel::UserWarning);
 			return EBakeOpState::Invalid;
 		}
-
-		// The read texture data is always in linear space.
-		NewSettings.bSRGB = false; 
 	}
 
 	if (!(CachedTexture2DImageSettings == NewSettings))
@@ -908,9 +905,6 @@ EBakeOpState UBakeMeshAttributeMapsTool::UpdateResult_MultiTexture()
 	}
 
 	CachedMultiTextures.Reset();
-
-	// The read texture data is always in linear space.
-	NewSettings.bSRGB = false;
 	
 	for ( auto& InputTexture : MultiTextureProps->MaterialIDSourceTextureMap)
 	{
diff --git a/Engine/Plugins/Experimental/MeshModelingToolsetExp/Source/MeshModelingToolsExp/Private/BakeMeshAttributeMapsToolBase.cpp b/Engine/Plugins/Experimental/MeshModelingToolsetExp/Source/MeshModelingToolsExp/Private/BakeMeshAttributeMapsToolBase.cpp
--- a/Engine/Plugins/Experimental/MeshModelingToolsetExp/Source/MeshModelingToolsExp/Private/BakeMeshAttributeMapsToolBase.cpp
+++ b/Engine/Plugins/Experimental/MeshModelingToolsetExp/Source/MeshModelingToolsExp/Private/BakeMeshAttributeMapsToolBase.cpp
@@ -231,9 +231,11 @@ void UBakeMeshAttributeMapsToolBase::OnMapsUpdated(const TUniquePtr<UE::Geometry
 
 		auto UpdateCachedMap = [this, &NewResult, &EvalIdx, &BakeDimensions](const EBakeMapType BakeMapType, const FTexture2DBuilder::ETextureType TexType, const int32 ResultIdx) -> void
 		{
+			// For 8-bit color textures, ensure that the source data is in sRGB.
+			const bool bConvertToSRGB = TexType == FTexture2DBuilder::ETextureType::Color;
 			FTexture2DBuilder TextureBuilder;
 			TextureBuilder.Initialize(TexType, BakeDimensions);
-			TextureBuilder.Copy(*NewResult->GetBakeResults(EvalIdx)[ResultIdx]);
+			TextureBuilder.Copy(*NewResult->GetBakeResults(EvalIdx)[ResultIdx], bConvertToSRGB);
 			TextureBuilder.Commit(false);
 
 			// The CachedMap & CachedMapIndices can be thrown out of sync if updated during
@@ -248,7 +250,8 @@ void UBakeMeshAttributeMapsToolBase::OnMapsUpdated(const TUniquePtr<UE::Geometry
 		{
 		case EMeshMapEvaluatorType::Normal:
 		{
-			UpdateCachedMap(EBakeMapType::TangentSpaceNormalMap, FTexture2DBuilder::ETextureType::NormalMap, 0);
+			constexpr EBakeMapType MapType = EBakeMapType::TangentSpaceNormalMap;
+			UpdateCachedMap(MapType, GetTextureType(MapType), 0);
 			break;
 		}
 		case EMeshMapEvaluatorType::Occlusion:
@@ -258,17 +261,20 @@ void UBakeMeshAttributeMapsToolBase::OnMapsUpdated(const TUniquePtr<UE::Geometry
 			int32 OcclusionIdx = 0;
 			if ((bool)(OcclusionEval->OcclusionType & EMeshOcclusionMapType::AmbientOcclusion))
 			{
-				UpdateCachedMap(EBakeMapType::AmbientOcclusion, FTexture2DBuilder::ETextureType::AmbientOcclusion, OcclusionIdx++);
+				constexpr EBakeMapType MapType = EBakeMapType::AmbientOcclusion;
+				UpdateCachedMap(MapType, GetTextureType(MapType), OcclusionIdx++);
 			}
 			if ((bool)(OcclusionEval->OcclusionType & EMeshOcclusionMapType::BentNormal))
 			{
-				UpdateCachedMap(EBakeMapType::BentNormal, FTexture2DBuilder::ETextureType::NormalMap, OcclusionIdx++);
+				constexpr EBakeMapType MapType = EBakeMapType::BentNormal;
+				UpdateCachedMap(MapType, GetTextureType(MapType), OcclusionIdx++);
 			}
 			break;
 		}
 		case EMeshMapEvaluatorType::Curvature:
 		{
-			UpdateCachedMap(EBakeMapType::Curvature, FTexture2DBuilder::ETextureType::ColorLinear, 0);
+			constexpr EBakeMapType MapType = EBakeMapType::Curvature;
+			UpdateCachedMap(MapType, GetTextureType(MapType), 0);
 			break;
 		}
 		case EMeshMapEvaluatorType::Property:
@@ -297,19 +303,19 @@ void UBakeMeshAttributeMapsToolBase::OnMapsUpdated(const TUniquePtr<UE::Geometry
 				break;
 			}
 
-			UpdateCachedMap(MapType, FTexture2DBuilder::ETextureType::ColorLinear, 0);
+			UpdateCachedMap(MapType, GetTextureType(MapType), 0);
 			break;
 		}
 		case EMeshMapEvaluatorType::ResampleImage:
 		{
-			constexpr FTexture2DBuilder::ETextureType TexType = FTexture2DBuilder::ETextureType::ColorLinear;
-			UpdateCachedMap(EBakeMapType::Texture2DImage, TexType, 0);
+			constexpr EBakeMapType MapType = EBakeMapType::Texture2DImage;
+			UpdateCachedMap(EBakeMapType::Texture2DImage, GetTextureType(MapType), 0);
 			break;
 		}
 		case EMeshMapEvaluatorType::MultiResampleImage:
 		{
-			constexpr FTexture2DBuilder::ETextureType TexType = FTexture2DBuilder::ETextureType::ColorLinear;
-			UpdateCachedMap(EBakeMapType::MultiTexture, TexType, 0);
+			constexpr EBakeMapType MapType = EBakeMapType::MultiTexture;
+			UpdateCachedMap(MapType, GetTextureType(MapType), 0);
 			break;
 		}
 		default:
@@ -351,7 +357,7 @@ TArray<EBakeMapType> UBakeMeshAttributeMapsToolBase::GetMapTypesArray(const int3
 
 FTexture2DBuilder::ETextureType UBakeMeshAttributeMapsToolBase::GetTextureType(EBakeMapType MapType)
 {
-	FTexture2DBuilder::ETextureType TexType = FTexture2DBuilder::ETextureType::ColorLinear;
+	FTexture2DBuilder::ETextureType TexType = FTexture2DBuilder::ETextureType::Color;
 	switch (MapType)
 	{
 	default:
@@ -369,9 +375,11 @@ FTexture2DBuilder::ETextureType UBakeMeshAttributeMapsToolBase::GetTextureType(E
 	case EBakeMapType::Curvature:
 	case EBakeMapType::NormalImage:
 	case EBakeMapType::FaceNormalImage:
+	case EBakeMapType::PositionImage:
+		TexType = FTexture2DBuilder::ETextureType::ColorLinear;
+		break;
 	case EBakeMapType::MaterialID:
 	case EBakeMapType::VertexColorImage:
-	case EBakeMapType::PositionImage:
 	case EBakeMapType::Texture2DImage:
 	case EBakeMapType::MultiTexture:
 		break;
@@ -380,7 +388,7 @@ FTexture2DBuilder::ETextureType UBakeMeshAttributeMapsToolBase::GetTextureType(E
 }
 
 
-void UBakeMeshAttributeMapsToolBase::GetTextureName(EBakeMapType MapType, const FString& BaseName, FString& TexName)
+void UBakeMeshAttributeMapsToolBase::GetTextureName(const EBakeMapType MapType, const FString& BaseName, FString& TexName)
 {
 	switch (MapType)
 	{
diff --git a/Engine/Plugins/Experimental/MeshModelingToolsetExp/Source/MeshModelingToolsExp/Private/BakeMeshAttributeVertexTool.cpp b/Engine/Plugins/Experimental/MeshModelingToolsetExp/Source/MeshModelingToolsExp/Private/BakeMeshAttributeVertexTool.cpp
--- a/Engine/Plugins/Experimental/MeshModelingToolsetExp/Source/MeshModelingToolsExp/Private/BakeMeshAttributeVertexTool.cpp
+++ b/Engine/Plugins/Experimental/MeshModelingToolsetExp/Source/MeshModelingToolsExp/Private/BakeMeshAttributeVertexTool.cpp
@@ -879,12 +879,10 @@ EBakeOpState UBakeMeshAttributeVertexTool::UpdateResult_Texture2DImage()
 		return EBakeOpState::Invalid;
 	}
 
-	// The read texture data is always in linear space.
-	NewSettings.bSRGB = false;
 
 	{
 		CachedTextureImage = MakeShared<UE::Geometry::TImageBuilder<FVector4f>, ESPMode::ThreadSafe>();
-		if (!UE::AssetUtils::ReadTexture(TextureSettings->SourceTexture, *CachedTextureImage, /*bPreferPlatformData*/ false))
+		if (!UE::AssetUtils::ReadTexture(TextureSettings->SourceTexture, *CachedTextureImage, bPreferPlatformData))
 		{
 			GetToolManager()->DisplayMessage(LOCTEXT("CannotReadTextureWarning", "Cannot read from the source texture"), EToolMessageLevel::UserWarning);
 			return EBakeOpState::Invalid;
@@ -924,9 +922,6 @@ EBakeOpState UBakeMeshAttributeVertexTool::UpdateResult_MultiTexture()
 
 	CachedMultiTextures.Reset();
 
-	// The read texture data is always in linear space.
-	NewSettings.bSRGB = false;
-	
 	for (auto& InputTexture : MultiTextureSettings->MaterialIDSourceTextureMap)
 	{
 		UTexture2D* Texture = InputTexture.Value;
@@ -938,7 +933,7 @@ EBakeOpState UBakeMeshAttributeVertexTool::UpdateResult_MultiTexture()
 
 		int32 MaterialID = InputTexture.Key;
 		CachedMultiTextures.Add(MaterialID, MakeShared<UE::Geometry::TImageBuilder<FVector4f>, ESPMode::ThreadSafe>());
-		if (!UE::AssetUtils::ReadTexture(Texture, *CachedMultiTextures[MaterialID], /*bPreferPlatformData*/ false))
+		if (!UE::AssetUtils::ReadTexture(Texture, *CachedMultiTextures[MaterialID], bPreferPlatformData))
 		{
 			GetToolManager()->DisplayMessage(LOCTEXT("CannotReadTextureWarning", "Cannot read from the source texture"), EToolMessageLevel::UserWarning);
 			return EBakeOpState::Invalid;
diff --git a/Engine/Plugins/Experimental/MeshModelingToolsetExp/Source/MeshModelingToolsExp/Public/BakeMeshAttributeMapsToolBase.h b/Engine/Plugins/Experimental/MeshModelingToolsetExp/Source/MeshModelingToolsExp/Public/BakeMeshAttributeMapsToolBase.h
--- a/Engine/Plugins/Experimental/MeshModelingToolsetExp/Source/MeshModelingToolsExp/Public/BakeMeshAttributeMapsToolBase.h
+++ b/Engine/Plugins/Experimental/MeshModelingToolsetExp/Source/MeshModelingToolsExp/Public/BakeMeshAttributeMapsToolBase.h
@@ -231,7 +231,7 @@ protected:
 	//
 	// Utilities
 	//
-	const bool bPreferPlatformData = true;
+	const bool bPreferPlatformData = false;
 	
 	/** @return An enumerated map type from a bitfield */
 	static EBakeMapType GetMapTypes(const int32& MapTypes);
diff --git a/Engine/Plugins/Experimental/MeshModelingToolsetExp/Source/MeshModelingToolsExp/Public/BakeMeshAttributeToolCommon.h b/Engine/Plugins/Experimental/MeshModelingToolsetExp/Source/MeshModelingToolsExp/Public/BakeMeshAttributeToolCommon.h
--- a/Engine/Plugins/Experimental/MeshModelingToolsetExp/Source/MeshModelingToolsExp/Public/BakeMeshAttributeToolCommon.h
+++ b/Engine/Plugins/Experimental/MeshModelingToolsetExp/Source/MeshModelingToolsExp/Public/BakeMeshAttributeToolCommon.h
@@ -309,11 +309,10 @@ struct FTexture2DImageSettings
 {
 	FImageDimensions Dimensions;
 	int32 UVLayer = 0;
-	bool bSRGB = true;
 
 	bool operator==(const FTexture2DImageSettings& Other) const
 	{
-		return Dimensions == Other.Dimensions && UVLayer == Other.UVLayer && bSRGB == Other.bSRGB;
+		return Dimensions == Other.Dimensions && UVLayer == Other.UVLayer;
 	}
 };
 
diff --git a/Engine/Plugins/Experimental/MeshModelingToolsetExp/Source/MeshModelingToolsExp/Public/BakeMeshAttributeVertexTool.h b/Engine/Plugins/Experimental/MeshModelingToolsetExp/Source/MeshModelingToolsExp/Public/BakeMeshAttributeVertexTool.h
--- a/Engine/Plugins/Experimental/MeshModelingToolsetExp/Source/MeshModelingToolsExp/Public/BakeMeshAttributeVertexTool.h
+++ b/Engine/Plugins/Experimental/MeshModelingToolsetExp/Source/MeshModelingToolsExp/Public/BakeMeshAttributeVertexTool.h
@@ -216,7 +216,8 @@ protected:
 	TObjectPtr<UMaterialInstanceDynamic> WorkingPreviewMaterial;
 	float SecondsBeforeWorkingMaterial = 0.75;
 
-	UWorld* TargetWorld = nullptr;
+	UPROPERTY()
+	TObjectPtr<UWorld> TargetWorld = nullptr;
 
 	TUniquePtr<TGenericDataBackgroundCompute<UE::Geometry::FMeshVertexBaker>> Compute = nullptr;
 	void OnResultUpdated(const TUniquePtr<UE::Geometry::FMeshVertexBaker>& NewResult);
@@ -238,6 +239,8 @@ protected:
 	void UpdateColorTopology();
 	void UpdateResult();
 
+	const bool bPreferPlatformData = false;
+
 	EBakeOpState OpState = EBakeOpState::Evaluate;
 
 	struct FBakeSettings
