ModelingTools: Remove obsolete GaussianBlur properties.
Improve BakeMapType display names and tool tips.

#rnx

diff --git a/Engine/Plugins/Experimental/MeshModelingToolsetExp/Source/MeshModelingToolsExp/Private/BakeMeshAttributeMapsTool.cpp b/Engine/Plugins/Experimental/MeshModelingToolsetExp/Source/MeshModelingToolsExp/Private/BakeMeshAttributeMapsTool.cpp
--- a/Engine/Plugins/Experimental/MeshModelingToolsetExp/Source/MeshModelingToolsExp/Private/BakeMeshAttributeMapsTool.cpp
+++ b/Engine/Plugins/Experimental/MeshModelingToolsetExp/Source/MeshModelingToolsExp/Private/BakeMeshAttributeMapsTool.cpp
@@ -334,8 +334,6 @@ void UBakeMeshAttributeMapsTool::Setup()
 	OcclusionMapProps->WatchProperty(OcclusionMapProps->MaxDistance, [this](float) { bInputsDirty = true; });
 	OcclusionMapProps->WatchProperty(OcclusionMapProps->SpreadAngle, [this](float) { bInputsDirty = true; });
 	OcclusionMapProps->WatchProperty(OcclusionMapProps->Distribution, [this](EOcclusionMapDistribution) { bInputsDirty = true; });
-	OcclusionMapProps->WatchProperty(OcclusionMapProps->BlurRadius, [this](float) { bInputsDirty = true; });
-	OcclusionMapProps->WatchProperty(OcclusionMapProps->bGaussianBlur, [this](float) { bInputsDirty = true; });
 	OcclusionMapProps->WatchProperty(OcclusionMapProps->BiasAngle, [this](float) { bInputsDirty = true; });
 	OcclusionMapProps->WatchProperty(OcclusionMapProps->NormalSpace, [this](ENormalMapSpace) { bInputsDirty = true; });
 
@@ -349,8 +347,6 @@ void UBakeMeshAttributeMapsTool::Setup()
 	CurvatureMapProps->WatchProperty(CurvatureMapProps->CurvatureType, [this](EBakedCurvatureTypeMode) { bInputsDirty = true; });
 	CurvatureMapProps->WatchProperty(CurvatureMapProps->ColorMode, [this](EBakedCurvatureColorMode) { bInputsDirty = true; });
 	CurvatureMapProps->WatchProperty(CurvatureMapProps->Clamping, [this](EBakedCurvatureClampMode) { bInputsDirty = true; });
-	CurvatureMapProps->WatchProperty(CurvatureMapProps->BlurRadius, [this](float) { bInputsDirty = true; });
-	CurvatureMapProps->WatchProperty(CurvatureMapProps->bGaussianBlur, [this](float) { bInputsDirty = true; });
 
 
 	Texture2DProps = NewObject<UBakedTexture2DImageProperties>(this);
@@ -739,7 +735,6 @@ EBakeOpState UBakeMeshAttributeMapsTool::UpdateResult_Occlusion()
 	OcclusionMapSettings.OcclusionRays = OcclusionMapProps->OcclusionRays;
 	OcclusionMapSettings.SpreadAngle = OcclusionMapProps->SpreadAngle;
 	OcclusionMapSettings.Distribution = OcclusionMapProps->Distribution;
-	OcclusionMapSettings.BlurRadius = (OcclusionMapProps->bGaussianBlur) ? OcclusionMapProps->BlurRadius : 0.0;
 	OcclusionMapSettings.BiasAngle = OcclusionMapProps->BiasAngle;
 	OcclusionMapSettings.NormalSpace = OcclusionMapProps->NormalSpace;
 
@@ -805,7 +800,6 @@ EBakeOpState UBakeMeshAttributeMapsTool::UpdateResult_Curvature()
 		CurvatureMapSettings.ClampMode = (int32)FMeshCurvatureMapEvaluator::EClampMode::Negative;
 		break;
 	}
-	CurvatureMapSettings.BlurRadius = (CurvatureMapProps->bGaussianBlur) ? CurvatureMapProps->BlurRadius : 0.0;
 
 	if (!(CachedCurvatureMapSettings == CurvatureMapSettings))
 	{
diff --git a/Engine/Plugins/Experimental/MeshModelingToolsetExp/Source/MeshModelingToolsExp/Private/BakeMeshAttributeVertexTool.cpp b/Engine/Plugins/Experimental/MeshModelingToolsetExp/Source/MeshModelingToolsExp/Private/BakeMeshAttributeVertexTool.cpp
--- a/Engine/Plugins/Experimental/MeshModelingToolsetExp/Source/MeshModelingToolsExp/Private/BakeMeshAttributeVertexTool.cpp
+++ b/Engine/Plugins/Experimental/MeshModelingToolsetExp/Source/MeshModelingToolsExp/Private/BakeMeshAttributeVertexTool.cpp
@@ -341,8 +341,6 @@ void UBakeMeshAttributeVertexTool::Setup()
 	OcclusionSettings->WatchProperty(OcclusionSettings->MaxDistance, [this](float) { OpState = EBakeOpState::Evaluate; });
 	OcclusionSettings->WatchProperty(OcclusionSettings->SpreadAngle, [this](float) { OpState = EBakeOpState::Evaluate; });
 	OcclusionSettings->WatchProperty(OcclusionSettings->Distribution, [this](EOcclusionMapDistribution) { OpState = EBakeOpState::Evaluate; });
-	OcclusionSettings->WatchProperty(OcclusionSettings->BlurRadius, [this](float) { OpState = EBakeOpState::Evaluate; });
-	OcclusionSettings->WatchProperty(OcclusionSettings->bGaussianBlur, [this](float) { OpState = EBakeOpState::Evaluate; });
 	OcclusionSettings->WatchProperty(OcclusionSettings->BiasAngle, [this](float) { OpState = EBakeOpState::Evaluate; });
 	OcclusionSettings->WatchProperty(OcclusionSettings->NormalSpace, [this](ENormalMapSpace) { OpState = EBakeOpState::Evaluate; });
 
@@ -355,8 +353,6 @@ void UBakeMeshAttributeVertexTool::Setup()
 	CurvatureSettings->WatchProperty(CurvatureSettings->CurvatureType, [this](EBakedCurvatureTypeMode) { OpState = EBakeOpState::Evaluate; });
 	CurvatureSettings->WatchProperty(CurvatureSettings->ColorMode, [this](EBakedCurvatureColorMode) { OpState = EBakeOpState::Evaluate; });
 	CurvatureSettings->WatchProperty(CurvatureSettings->Clamping, [this](EBakedCurvatureClampMode) { OpState = EBakeOpState::Evaluate; });
-	CurvatureSettings->WatchProperty(CurvatureSettings->BlurRadius, [this](float) { OpState = EBakeOpState::Evaluate; });
-	CurvatureSettings->WatchProperty(CurvatureSettings->bGaussianBlur, [this](float) { OpState = EBakeOpState::Evaluate; });
 
 	TextureSettings = NewObject<UBakedTexture2DImageProperties>(this);
 	TextureSettings->RestoreProperties(this);
@@ -782,7 +778,6 @@ EBakeOpState UBakeMeshAttributeVertexTool::UpdateResult_Occlusion()
 	OcclusionMapSettings.OcclusionRays = OcclusionSettings->OcclusionRays;
 	OcclusionMapSettings.SpreadAngle = OcclusionSettings->SpreadAngle;
 	OcclusionMapSettings.Distribution = OcclusionSettings->Distribution;
-	OcclusionMapSettings.BlurRadius = (OcclusionSettings->bGaussianBlur) ? OcclusionSettings->BlurRadius : 0.0;
 	OcclusionMapSettings.BiasAngle = OcclusionSettings->BiasAngle;
 	OcclusionMapSettings.NormalSpace = OcclusionSettings->NormalSpace;
 
@@ -843,7 +838,6 @@ EBakeOpState UBakeMeshAttributeVertexTool::UpdateResult_Curvature()
 		CurvatureMapSettings.ClampMode = (int32)FMeshCurvatureMapEvaluator::EClampMode::Negative;
 		break;
 	}
-	CurvatureMapSettings.BlurRadius = (CurvatureSettings->bGaussianBlur) ? CurvatureSettings->BlurRadius : 0.0;
 
 	if (!(CachedCurvatureMapSettings == CurvatureMapSettings))
 	{
diff --git a/Engine/Plugins/Experimental/MeshModelingToolsetExp/Source/MeshModelingToolsExp/Public/BakeMeshAttributeMapsToolBase.h b/Engine/Plugins/Experimental/MeshModelingToolsetExp/Source/MeshModelingToolsExp/Public/BakeMeshAttributeMapsToolBase.h
--- a/Engine/Plugins/Experimental/MeshModelingToolsetExp/Source/MeshModelingToolsExp/Public/BakeMeshAttributeMapsToolBase.h
+++ b/Engine/Plugins/Experimental/MeshModelingToolsetExp/Source/MeshModelingToolsExp/Public/BakeMeshAttributeMapsToolBase.h
@@ -27,17 +27,28 @@ UENUM(meta=(Bitflags, UseEnumValuesAsMaskValuesInEditor="true"))
 enum class EBakeMapType
 {
 	None                   = 0 UMETA(Hidden),
-	TangentSpaceNormalMap  = 1 << 0,
-	AmbientOcclusion       = 1 << 1,
-	BentNormal             = 1 << 2,
-	Curvature              = 1 << 3,
-	Texture2DImage         = 1 << 4,
-	NormalImage            = 1 << 5,
-	FaceNormalImage        = 1 << 6,
-	PositionImage          = 1 << 7,
-	MaterialID             = 1 << 8,
-	MultiTexture           = 1 << 9,
-	VertexColorImage       = 1 << 10,
+	/* Sample normals from the detail mesh in tangent space */
+	TangentSpaceNormalMap  = 1 << 0 UMETA(DisplayName="Tangent Space Normals"),
+	/* Sample ambient occlusion from the detail mesh */
+	AmbientOcclusion       = 1 << 1 UMETA(DisplayName="Ambient Occlusion"),
+	/* Sample normals skewed towards the least occluded direction from the detail mesh */
+	BentNormal             = 1 << 2 UMETA(DisplayName="Bent Normals"),
+	/* Sample mesh curvatures from the detail mesh */
+	Curvature              = 1 << 3 UMETA(DisplayName="Curvature"),
+	/* Sample a source texture from the detail mesh UVs */
+	Texture2DImage         = 1 << 4 UMETA(DisplayName="Texture"),
+	/* Sample object space normals from the detail mesh */
+	NormalImage            = 1 << 5 UMETA(DisplayName="Object Space Normals"),
+	/* Sample object space face normals from the detail mesh */
+	FaceNormalImage        = 1 << 6 UMETA(DisplayName="Face Normals"),
+	/* Sample bounding box relative positions from the detail mesh */
+	PositionImage          = 1 << 7 UMETA(DisplayName="Position"),
+	/* Sample material IDs as unique colors from the detail mesh */
+	MaterialID             = 1 << 8 UMETA(DisplayName="Material ID"),
+	/* Sample a source texture per material ID on the detail mesh */
+	MultiTexture           = 1 << 9 UMETA(DisplayName="MultiTexture"),
+	/* Sample the interpolated vertex colors from the detail mesh */
+	VertexColorImage       = 1 << 10 UMETA(DisplayName="Vertex Colors"),
 	Occlusion              = (AmbientOcclusion | BentNormal) UMETA(Hidden),
 	All                    = 0x7FF UMETA(Hidden)
 };
diff --git a/Engine/Plugins/Experimental/MeshModelingToolsetExp/Source/MeshModelingToolsExp/Public/BakeMeshAttributeToolCommon.h b/Engine/Plugins/Experimental/MeshModelingToolsetExp/Source/MeshModelingToolsExp/Public/BakeMeshAttributeToolCommon.h
--- a/Engine/Plugins/Experimental/MeshModelingToolsetExp/Source/MeshModelingToolsExp/Public/BakeMeshAttributeToolCommon.h
+++ b/Engine/Plugins/Experimental/MeshModelingToolsetExp/Source/MeshModelingToolsExp/Public/BakeMeshAttributeToolCommon.h
@@ -88,14 +88,6 @@ public:
 	UPROPERTY(EditAnywhere, Category = OcclusionMap)
 	EOcclusionMapDistribution Distribution = EOcclusionMapDistribution::Cosine;
 
-	/** Whether or not to apply Gaussian Blur to computed AO Map (recommended) */
-	UPROPERTY(EditAnywhere, Category = "OcclusionMap|Ambient Occlusion")
-	bool bGaussianBlur = true;
-
-	/** Pixel Radius of Gaussian Blur Kernel */
-	UPROPERTY(EditAnywhere, Category = "OcclusionMap|Ambient Occlusion", meta = (UIMin = "0", UIMax = "10.0", ClampMin = "0", ClampMax = "100.0"))
-	float BlurRadius = 2.25;
-
 	/** Contribution of AO rays that are within this angle (degrees) from horizontal are attenuated. This reduces faceting artifacts. */
 	UPROPERTY(EditAnywhere, Category = "OcclusionMap|Ambient Occlusion", meta = (UIMin = "0", UIMax = "45.0", ClampMin = "0", ClampMax = "89.9"))
 	float BiasAngle = 15.0;
@@ -182,14 +174,6 @@ public:
 	/** Clamping to apply to curvature values before scaling to color range */
 	UPROPERTY(EditAnywhere, Category = CurvatureMap)
 	EBakedCurvatureClampMode Clamping = EBakedCurvatureClampMode::None;
-
-	/** Whether or not to apply Gaussian Blur to computed Map */
-	UPROPERTY(EditAnywhere, Category = CurvatureMap)
-	bool bGaussianBlur = false;
-
-	/** Pixel Radius of Gaussian Blur Kernel */
-	UPROPERTY(EditAnywhere, Category = CurvatureMap, meta = (UIMin = "0", UIMax = "10.0", ClampMin = "0", ClampMax = "100.0"))
-	float BlurRadius = 2.25;
 };
 
 
@@ -260,7 +244,6 @@ struct FOcclusionMapSettings
 	float MaxDistance;
 	float SpreadAngle;
 	EOcclusionMapDistribution Distribution;
-	float BlurRadius;
 	float BiasAngle;
 	ENormalMapSpace NormalSpace;
 
@@ -271,7 +254,6 @@ struct FOcclusionMapSettings
 			MaxDistance == Other.MaxDistance &&
 			SpreadAngle == Other.SpreadAngle &&
 			Distribution == Other.Distribution &&
-			BlurRadius == Other.BlurRadius &&
 			BiasAngle == Other.BiasAngle &&
 			NormalSpace == Other.NormalSpace;
 	}
@@ -287,11 +269,10 @@ struct FCurvatureMapSettings
 	int32 ColorMode = 0;
 	int32 ClampMode = 0;
 	float MaxDistance = 1.0;
-	float BlurRadius = 1.0;
 
 	bool operator==(const FCurvatureMapSettings& Other) const
 	{
-		return Dimensions == Other.Dimensions && RayCount == Other.RayCount && CurvatureType == Other.CurvatureType && RangeMultiplier == Other.RangeMultiplier && MinRangeMultiplier == Other.MinRangeMultiplier && ColorMode == Other.ColorMode && ClampMode == Other.ClampMode && MaxDistance == Other.MaxDistance && BlurRadius == Other.BlurRadius;
+		return Dimensions == Other.Dimensions && RayCount == Other.RayCount && CurvatureType == Other.CurvatureType && RangeMultiplier == Other.RangeMultiplier && MinRangeMultiplier == Other.MinRangeMultiplier && ColorMode == Other.ColorMode && ClampMode == Other.ClampMode && MaxDistance == Other.MaxDistance;
 	}
 };
 
diff --git a/Engine/Plugins/Experimental/MeshModelingToolsetExp/Source/MeshModelingToolsExp/Public/BakeMeshAttributeVertexTool.h b/Engine/Plugins/Experimental/MeshModelingToolsetExp/Source/MeshModelingToolsExp/Public/BakeMeshAttributeVertexTool.h
--- a/Engine/Plugins/Experimental/MeshModelingToolsetExp/Source/MeshModelingToolsExp/Public/BakeMeshAttributeVertexTool.h
+++ b/Engine/Plugins/Experimental/MeshModelingToolsetExp/Source/MeshModelingToolsExp/Public/BakeMeshAttributeVertexTool.h
@@ -48,16 +48,26 @@ enum class EBakeVertexMode
 UENUM()
 enum class EBakeVertexTypeColor
 {
-	TangentSpaceNormal,
-	AmbientOcclusion,
-	BentNormal,
-	Curvature,
-	Texture2DImage,
-	NormalImage,
-	FaceNormalImage,
-	PositionImage,
-	MaterialID,
-	MultiTexture
+	/* Sample tangent space normals from the detail mesh */
+	TangentSpaceNormal UMETA(DisplayName="Tangent Space Normals"),
+	/* Sample ambient occlusion from the detail mesh */
+	AmbientOcclusion UMETA(DisplayName="Ambient Occlusion"),
+	/* Sample normals skewed towards the least occluded direction from the detail mesh */
+	BentNormal UMETA(DisplayName="Bent Normals"),
+	/* Sample mesh curvatures from the detail mesh */
+	Curvature UMETA(DisplayName="Curvature"),
+	/* Sample a source texture from the detail mesh UVs */
+	Texture2DImage UMETA(DisplayName="Texture"),
+	/* Sample object space normals from the detail mesh */
+	NormalImage UMETA(DisplayName="Object Space Normals"),
+	/* Sample object space face normals from the detail mesh */
+	FaceNormalImage UMETA(DisplayName="Face Normals"),
+	/* Sample bounding box relative positions from the detail mesh */
+	PositionImage UMETA(DisplayName="Position"),
+	/* Sample material IDs as unique colors from the detail mesh */
+	MaterialID UMETA(DisplayName="Material ID"),
+	/* Sample a source texture per material ID on the detail mesh */
+	MultiTexture UMETA(DisplayName="MultiTexture")
 };
 
 
@@ -65,8 +75,10 @@ UENUM()
 enum class EBakeVertexTypeChannel
 {
 	None,
-	AmbientOcclusion,
-	Curvature
+	/* Sample ambient occlusion from the detail mesh */
+	AmbientOcclusion UMETA(DisplayName="Ambient Occlusion"),
+	/* Sample mesh curvatures from the detail mesh */
+	Curvature UMETA(DisplayName="Curvature"),
 };
 
 
