Refactor FCameraRectangle to make it possible to implement marquee selection more efficiently in the UV editor where the meshes are known to have 0 z-coordinate

#rnx

diff --git a/Engine/Plugins/Experimental/MeshModelingToolsetExp/Source/ModelingComponents/Private/Mechanics/RectangleMarqueeMechanic.cpp b/Engine/Plugins/Experimental/MeshModelingToolsetExp/Source/ModelingComponents/Private/Mechanics/RectangleMarqueeMechanic.cpp
--- a/Engine/Plugins/Experimental/MeshModelingToolsetExp/Source/ModelingComponents/Private/Mechanics/RectangleMarqueeMechanic.cpp
+++ b/Engine/Plugins/Experimental/MeshModelingToolsetExp/Source/ModelingComponents/Private/Mechanics/RectangleMarqueeMechanic.cpp
@@ -5,127 +5,128 @@
 #include "CanvasTypes.h"
 #include "CanvasItem.h"
 #include "ToolSceneQueriesUtil.h"
-#include "TransformTypes.h"
 
 #include "ExplicitUseGeometryMathTypes.h"		// using UE::Geometry::(math types)
+#include "Intersection/IntrSegment2Segment2.h"
 using namespace UE::Geometry;
 
 #define LOCTEXT_NAMESPACE "URectangleMarqueeMechanic"
 
-
-FCameraRectangle::FCameraRectangle(const FViewCameraState& CachedCameraState,
-								   const FRay& DragStartWorldRay,
-								   const FRay& DragEndWorldRay)
+void FCameraRectangle::Initialize()
 {
-	// Create a plane just in front of the camera
-	CameraOrigin = CachedCameraState.Position;
-	CameraPlane = FPlane(CachedCameraState.Position + CachedCameraState.Forward(), CachedCameraState.Forward());
-	bCameraIsOrthographic = CachedCameraState.bIsOrthographic;
-
-	// Intersect the drag rays with the camera plane and compute their coordinates in the camera basis
-	UBasisVector = CachedCameraState.Right();
-	VBasisVector = CachedCameraState.Up();
-
-	FVector StartIntersection = FMath::RayPlaneIntersection(DragStartWorldRay.Origin,
-															DragStartWorldRay.Direction,
-															CameraPlane);
-	FVector2D Start2D = PlaneCoordinates(StartIntersection);
-
-	FVector CurrentIntersection = FMath::RayPlaneIntersection(DragEndWorldRay.Origin,
-															  DragEndWorldRay.Direction,
-															  CameraPlane);
-
-	FVector2D Current2D = PlaneCoordinates(CurrentIntersection);
+	bIsInitialized = true;
+	CameraPlane = FPlane3(CameraState.Forward(), CameraState.Position + CameraState.Forward());
+	RectangleInCameraPlane = GetSelectionRectangleUV();
+}
 
-	RectangleCorners = FBox2D(Start2D, Start2D);
-	RectangleCorners += Current2D;	// Initialize this way so we don't have to care about min/max
+FCameraRectangle::FAxisAlignedBox2 FCameraRectangle::GetSelectionRectangleUV(double OffsetFromCameraPlane) const
+{
+	ensure(bIsInitialized);
+	
+	FAxisAlignedBox2 Result;
+
+	FPlane3 SelectionPlane = CameraPlane;
+	SelectionPlane.Constant += OffsetFromCameraPlane;
+	
+	FVector ProjectedStart = FMath::RayPlaneIntersection(RectangleStartRay.WorldRay.Origin, RectangleStartRay.WorldRay.Direction, (FPlane)SelectionPlane);
+	FVector ProjectedEnd = FMath::RayPlaneIntersection(RectangleEndRay.WorldRay.Origin, RectangleEndRay.WorldRay.Direction, (FPlane)SelectionPlane);
+	FVector2 StartUV = PlaneCoordinates(SelectionPlane, ProjectedStart);
+	FVector2 EndUV = PlaneCoordinates(SelectionPlane, ProjectedEnd);
+
+	// Initialize this way so we don't have to care about min/max
+	Result.Contain(StartUV);
+	Result.Contain(EndUV);
+
+	return Result;
 }
 
 bool FCameraRectangle::IsProjectedPointInRectangle(const FVector& Point) const
 {
+	ensure(bIsInitialized);
+	
 	FVector ProjectedPoint;
-	if (bCameraIsOrthographic)
+	if (CameraState.bIsOrthographic)
 	{
-		// project directly to plane
-		ProjectedPoint = FVector::PointPlaneProject(Point, CameraPlane);
+		ProjectedPoint = OrthographicProjection(CameraPlane, Point);
 	}
 	else
 	{
-		// If it's behind the camera rectangle, then not contained.
-		if (CameraPlane.PlaneDot(Point) < 0)
+		// If it's not in front of the camera plane, its not contained in the camera rectangle
+		if (CameraPlane.DistanceTo(Point) <= 0)
 		{
 			return false;
 		}
-
-		// intersect along the eye-to-point ray
-		ProjectedPoint = FMath::RayPlaneIntersection(CameraOrigin,
-													 Point - CameraOrigin,
-													 CameraPlane);
+		ProjectedPoint = PerspectiveProjection(CameraPlane, Point);
 	}
 
-	FVector2D Point2D = PlaneCoordinates(ProjectedPoint);
-	return RectangleCorners.IsInside(Point2D);
+	FVector2 PointUV = PlaneCoordinates(CameraPlane, ProjectedPoint);
+	return RectangleInCameraPlane.Contains(PointUV);
 }
 
 bool FCameraRectangle::IsProjectedSegmentIntersectingRectangle(const FVector& Endpoint1, const FVector& Endpoint2) const
 {
+	ensure(bIsInitialized);
+	
 	FVector ProjectedEndpoint1;
 	FVector ProjectedEndpoint2;
 
-	if (bCameraIsOrthographic)
+	if (CameraState.bIsOrthographic)
 	{
-		ProjectedEndpoint1 = FVector::PointPlaneProject(Endpoint1, CameraPlane);
-		ProjectedEndpoint2 = FVector::PointPlaneProject(Endpoint2, CameraPlane);
+		ProjectedEndpoint1 = OrthographicProjection(CameraPlane, Endpoint1);
+		ProjectedEndpoint2 = OrthographicProjection(CameraPlane, Endpoint2);
 	}
 	else
 	{
-		// We'll have to crop the segment to the portion in front of the camera plane
-		bool bPoint1InFrontOfCamera = CameraPlane.PlaneDot(Endpoint1) > 0;
-		bool bPoint2InFrontOfCamera = CameraPlane.PlaneDot(Endpoint2) > 0;
-		
-		if (!bPoint1InFrontOfCamera && !bPoint2InFrontOfCamera)
-		{
-			return false; // Segment is behind the camera plane
-		}
+		ProjectedEndpoint1 = Endpoint1;
+		ProjectedEndpoint2 = Endpoint2;
 
-		FVector IntersectionPoint;
-		if (bPoint1InFrontOfCamera != bPoint2InFrontOfCamera)
+		// We have to crop the segment to the portion in front of the camera plane
+		FPlane3::EClipSegmentType ClipType = CameraPlane.ClipSegment(ProjectedEndpoint1, ProjectedEndpoint2);
+		
+		// Since the selection plane is identical to the clipping plane there's no need to reproject clipped points
+		switch (ClipType)
 		{
-			// Get the intersection point so we can replace the endpoint that is behind the camera with it.
-			ensure(FMath::SegmentPlaneIntersection(Endpoint1, Endpoint2, CameraPlane, IntersectionPoint));
+			case FPlane3::FullyClipped:
+				return false; // Segment is behind the camera plane hence not in the camera rectangle
+			case FPlane3::FirstClipped:
+				ProjectedEndpoint2 = PerspectiveProjection(CameraPlane, ProjectedEndpoint2);
+				break;
+			case FPlane3::SecondClipped:
+				ProjectedEndpoint1 = PerspectiveProjection(CameraPlane, ProjectedEndpoint1);
+				break;
+			case FPlane3::NotClipped:
+			default:
+				ProjectedEndpoint1 = PerspectiveProjection(CameraPlane, ProjectedEndpoint1);
+				ProjectedEndpoint2 = PerspectiveProjection(CameraPlane, ProjectedEndpoint2);
 		}
-
-		ProjectedEndpoint1 = bPoint1InFrontOfCamera ? FMath::RayPlaneIntersection(CameraOrigin, Endpoint1 - CameraOrigin, CameraPlane)
-			: IntersectionPoint;
-		ProjectedEndpoint2 = bPoint2InFrontOfCamera ? FMath::RayPlaneIntersection(CameraOrigin, Endpoint2 - CameraOrigin, CameraPlane)
-			: IntersectionPoint;
 	}
 
-	FVector2D Endpoint1PlaneCoord = PlaneCoordinates(ProjectedEndpoint1);
-	FVector2D Endpoint2PlaneCoord = PlaneCoordinates(ProjectedEndpoint2);
+	// TODO Port Wild Magic IntrSegment2Box2 (which requires porting IntrLine2Box2) so we can call segment-box intersection here
+
+	FVector2 Endpoint1UV = PlaneCoordinates(CameraPlane, ProjectedEndpoint1);
+	FVector2 Endpoint2UV = PlaneCoordinates(CameraPlane, ProjectedEndpoint2);
+	FSegment2 ProjectedSegmentUV(Endpoint1UV, Endpoint2UV);
 
 	// If either endpoint is inside, then definitely (at least partially) contained
-	if (RectangleCorners.IsInside(Endpoint1PlaneCoord) || RectangleCorners.IsInside(Endpoint2PlaneCoord))
+	if (RectangleInCameraPlane.Contains(ProjectedSegmentUV.StartPoint()) ||
+		RectangleInCameraPlane.Contains(ProjectedSegmentUV.EndPoint()))
 	{
 		return true;
 	}
 
-	// If both outside, have to do some intersections with the box sides. The function we have for this
-	// uses FVectors instead of FVector2Ds.
-	ProjectedEndpoint1 = FVector(Endpoint1PlaneCoord, 0);
-	ProjectedEndpoint2 = FVector(Endpoint2PlaneCoord, 0);
-	FVector Throwaway;
-	return FMath::SegmentIntersection2D(ProjectedEndpoint1, ProjectedEndpoint2, 
-		FVector(RectangleCorners.Min, 0), FVector(RectangleCorners.Max.X, RectangleCorners.Min.Y, 0), 
-		Throwaway)
+	// If both outside, have to do some intersections with the box sides
 
-		|| FMath::SegmentIntersection2D(ProjectedEndpoint1, ProjectedEndpoint2, 
-			FVector(RectangleCorners.Max.X, RectangleCorners.Min.Y, 0), FVector(RectangleCorners.Max, 0), 
-			Throwaway)
+	if (ProjectedSegmentUV.Intersects(FSegment2(RectangleInCameraPlane.GetCorner(0), RectangleInCameraPlane.GetCorner(1))))
+	{
+		return true;
+	}
 
-		|| FMath::SegmentIntersection2D(ProjectedEndpoint1, ProjectedEndpoint2,
-			FVector(RectangleCorners.Max, 0), FVector(RectangleCorners.Min.X, RectangleCorners.Max.Y, 0), 
-			Throwaway);
+	if (ProjectedSegmentUV.Intersects(FSegment2(RectangleInCameraPlane.GetCorner(1), RectangleInCameraPlane.GetCorner(2))))
+	{
+		return true;
+	}
+
+	return ProjectedSegmentUV.Intersects(FSegment2(RectangleInCameraPlane.GetCorner(3), RectangleInCameraPlane.GetCorner(2)));
 
 	// Don't need to intersect with the fourth side because segment would have to intersect two sides
 	// of box if both endpoints are outside the box.
@@ -179,7 +180,7 @@ TPair<FInputCapturePriority, FInputCapturePriority> URectangleMarqueeMechanic::G
 void URectangleMarqueeMechanic::Render(IToolsContextRenderAPI* RenderAPI)
 {
 	// Cache the camera state
-	GetParentTool()->GetToolManager()->GetContextQueriesAPI()->GetCurrentViewState(CachedCameraState);
+	GetParentTool()->GetToolManager()->GetContextQueriesAPI()->GetCurrentViewState(CameraRectangle.CameraState);
 }
 
 FInputRayHit URectangleMarqueeMechanic::CanBeginClickDragSequence(const FInputDeviceRay& PressPos)
@@ -197,8 +198,9 @@ void URectangleMarqueeMechanic::OnClickPress(const FInputDeviceRay& PressPos)
 		return;
 	}
 
-	DragStartScreenPosition = PressPos.ScreenPosition;
-	DragStartWorldRay = PressPos.WorldRay;
+	CameraRectangle.RectangleStartRay = PressPos;
+	CameraRectangle.RectangleEndRay = PressPos;
+	CameraRectangle.Initialize();
 
 	OnDragRectangleStarted.Broadcast();
 }
@@ -211,25 +213,22 @@ void URectangleMarqueeMechanic::OnClickDrag(const FInputDeviceRay& DragPos)
 	}
 
 	bIsDragging = true;
-	DragCurrentScreenPosition = DragPos.ScreenPosition;
-	FRay DragCurrentWorldRay = DragPos.WorldRay;
-
-	FCameraRectangle Rectangle(CachedCameraState, DragStartWorldRay, DragCurrentWorldRay);
-	LastRectangle = Rectangle;
-
-	OnDragRectangleChanged.Broadcast(Rectangle);
+	CameraRectangle.RectangleEndRay = DragPos;
+	CameraRectangle.Initialize();
+	
+	OnDragRectangleChanged.Broadcast(CameraRectangle);
 }
 
 void URectangleMarqueeMechanic::OnClickRelease(const FInputDeviceRay& ReleasePos)
 {
 	bIsDragging = false;
-	OnDragRectangleFinished.Broadcast(LastRectangle, false);
+	OnDragRectangleFinished.Broadcast(CameraRectangle, false);
 }
 
 void URectangleMarqueeMechanic::OnTerminateDragSequence()
 {
 	bIsDragging = false;
-	OnDragRectangleFinished.Broadcast(LastRectangle, true);
+	OnDragRectangleFinished.Broadcast(CameraRectangle, true);
 }
 
 
@@ -239,8 +238,8 @@ void URectangleMarqueeMechanic::DrawHUD(FCanvas* Canvas, IToolsContextRenderAPI*
 	bool bThisViewHasFocus = !!(State & EViewInteractionState::Focused);
 	if (bThisViewHasFocus && bIsDragging)
 	{
-		FVector2D Start = DragStartScreenPosition;
-		FVector2D Curr = DragCurrentScreenPosition;
+		FVector2D Start = CameraRectangle.RectangleStartRay.ScreenPosition;
+		FVector2D Curr = CameraRectangle.RectangleEndRay.ScreenPosition;
 		FCanvasBoxItem BoxItem(Start / Canvas->GetDPIScale(), (Curr - Start) / Canvas->GetDPIScale());
 		BoxItem.SetColor(FLinearColor::White);
 		Canvas->DrawItem(BoxItem);
diff --git a/Engine/Plugins/Experimental/MeshModelingToolsetExp/Source/ModelingComponents/Private/Selection/GroupTopologySelector.cpp b/Engine/Plugins/Experimental/MeshModelingToolsetExp/Source/ModelingComponents/Private/Selection/GroupTopologySelector.cpp
--- a/Engine/Plugins/Experimental/MeshModelingToolsetExp/Source/ModelingComponents/Private/Selection/GroupTopologySelector.cpp
+++ b/Engine/Plugins/Experimental/MeshModelingToolsetExp/Source/ModelingComponents/Private/Selection/GroupTopologySelector.cpp
@@ -599,15 +599,17 @@ bool FGroupTopologySelector::FindSelectedElement(const FSelectionSettings& Setti
 	FDynamicMeshAABBTree3* Spatial = GetSpatial();
 
 	// Needed for occlusion test, which happens in local space.
-	FVector3d LocalCameraOrigin = TargetTransform.InverseTransformPosition((FVector3d)CameraRectangle.CameraOrigin);
+	FVector3d LocalCameraOrigin = TargetTransform.InverseTransformPosition((FVector3d)CameraRectangle.CameraState.Position);
 	
 	// Corner selection takes priority over edges.
 	if (Settings.bEnableCornerHits)
 	{
 		GeometrySet.ParallelFindAllPointsSatisfying(
-			[&CameraRectangle, &TargetTransform, &Settings, &LocalCameraOrigin, Spatial](const FVector3d& PointPosition)
+			[&CameraRectangle, &TargetTransform, &Settings, &LocalCameraOrigin, Spatial]
+			(const FVector3d& PointPosition)
 			{
-				return CameraRectangle.IsProjectedPointInRectangle((FVector)TargetTransform.TransformPosition(PointPosition))
+				FVector TransformedPoint = TargetTransform.TransformPosition(PointPosition);
+				return CameraRectangle.IsProjectedPointInRectangle(TransformedPoint)
 					&& (Settings.bIgnoreOcclusion || !IsOccluded(PointPosition, LocalCameraOrigin, Spatial));
 			},
 			ResultOut.SelectedCornerIDs);
diff --git a/Engine/Plugins/Experimental/MeshModelingToolsetExp/Source/ModelingComponents/Private/Selection/MeshSelectionMechanic.cpp b/Engine/Plugins/Experimental/MeshModelingToolsetExp/Source/ModelingComponents/Private/Selection/MeshSelectionMechanic.cpp
--- a/Engine/Plugins/Experimental/MeshModelingToolsetExp/Source/ModelingComponents/Private/Selection/MeshSelectionMechanic.cpp
+++ b/Engine/Plugins/Experimental/MeshModelingToolsetExp/Source/ModelingComponents/Private/Selection/MeshSelectionMechanic.cpp
@@ -644,7 +644,7 @@ void UMeshSelectionMechanic::OnDragRectangleChanged(const FCameraRectangle& Curr
 
 void UMeshSelectionMechanic::OnDragRectangleFinished(const FCameraRectangle& Rectangle, bool bCancelled)
 {
-	if (PreDragSelection != CurrentSelection)
+	if (!bCancelled && (PreDragSelection != CurrentSelection))
 	{
 		UpdateCentroid();
 		RebuildDrawnElements(FTransform(GetCurrentSelectionCentroid()));
diff --git a/Engine/Plugins/Experimental/MeshModelingToolsetExp/Source/ModelingComponents/Public/Mechanics/RectangleMarqueeMechanic.h b/Engine/Plugins/Experimental/MeshModelingToolsetExp/Source/ModelingComponents/Public/Mechanics/RectangleMarqueeMechanic.h
--- a/Engine/Plugins/Experimental/MeshModelingToolsetExp/Source/ModelingComponents/Public/Mechanics/RectangleMarqueeMechanic.h
+++ b/Engine/Plugins/Experimental/MeshModelingToolsetExp/Source/ModelingComponents/Public/Mechanics/RectangleMarqueeMechanic.h
@@ -6,46 +6,80 @@
 #include "InputBehavior.h"
 #include "InteractionMechanic.h"
 #include "BaseBehaviors/BehaviorTargetInterfaces.h"
-#include "TransformTypes.h"
+#include "GeometryBase.h"
+#include "BoxTypes.h"
+#include "PlaneTypes.h"
+#include "SegmentTypes.h"
 #include "VectorTypes.h"
 
 #include "RectangleMarqueeMechanic.generated.h"
 
 class UClickDragInputBehavior;
 
+PREDECLARE_USE_GEOMETRY_CLASS(FGeometrySet3);
+
 /// Struct containing:
 ///		- camera information, 
-///		- a 3D plane just in front of the camera,
-///		- a 2D basis for coordinates in this plane, and
-///		- the corners of a rectangle contained in this plane, in this 2D basis
+///		- input device rays used to define the corners of a rectangle contained in a selection plane, the device ray
+///		  screen positions could be used to change the selection behavior when dragging from the top right to bottom
+///		  left or vice-versa
 /// 
 struct MODELINGCOMPONENTS_API FCameraRectangle
 {
-	FCameraRectangle() {};
-	FCameraRectangle(const FViewCameraState& CachedCameraState,
-					 const FRay& DragStartWorldRay,
-					 const FRay& DragEndWorldRay);
-
-	FVector CameraOrigin;
-	bool bCameraIsOrthographic;
-	FPlane CameraPlane;
-	FVector UBasisVector;
-	FVector VBasisVector;
-	FBox2D RectangleCorners;
-
-	// Project the given 3D point to the camera plane and test if it's in the rectangle
-	bool IsProjectedPointInRectangle(const FVector& Point) const;
+	using FGeometrySet3 = UE::Geometry::FGeometrySet3;
+	using FAxisAlignedBox2 = UE::Geometry::TAxisAlignedBox2<FVector::FReal>;
+	using FPlane3 = UE::Geometry::TPlane3<FVector::FReal>;
+	using FSegment2 = UE::Geometry::TSegment2<FVector::FReal>;
+	using FVector2 = UE::Math::TVector2<FVector::FReal>;
+
+	FViewCameraState CameraState;
+	FInputDeviceRay RectangleStartRay = FInputDeviceRay(FRay());
+	FInputDeviceRay RectangleEndRay = FInputDeviceRay(FRay());
+	FPlane3 CameraPlane;
+	FAxisAlignedBox2 RectangleInCameraPlane;
+
+	// This function must be called before other member functions whenever camera state or start/end rays are updated
+	void Initialize();
+	
+	bool bIsInitialized = false;
 
-	// Project the given segment to the camera plane and test if it intersects the rectangle
+	// Return an axis aligned box (the region of a marquee selection) in the UV-coordinate system of a plane offset from
+	// CameraPlane by the given distance in the CameraPlane normal direction. This function was added in order to
+	// implement optimizations (avoid projecting points) when selecting UV editor meshes, which lie in the XY plane
+	FAxisAlignedBox2 GetSelectionRectangleUV(double OffsetFromCameraPlane = 0.) const;
+
+	// Return true if the given 3D geometry projected to the camera plane is inside or intersecting the rectangle, and
+	// false otherwise
+	bool IsProjectedPointInRectangle(const FVector& Point) const;
 	bool IsProjectedSegmentIntersectingRectangle(const FVector& Endpoint1, const FVector& Endpoint2) const;
+	// bool IsProjectedTriangleIntersectingRectangle(const FVector& A, const FVector& B, const FVector& C) const;
 
-	// TODO: Add a way to test rectangle against triangles.
+	// Return the 3D point obtained by projecting the given 3D Point onto the given projection plane
+	// Note: Assumes an orthographic camera
+	static FVector OrthographicProjection(const FPlane3& Plane, const FVector& Point)
+	{
+		return FVector::PointPlaneProject(Point, (FPlane)Plane);
+	}
+
+	// Return the 3D point obtained by projecting the given 3D Point onto the given projection plane
+	// Note: OrthographicProjection is more efficient if the camera is known to be orthographic
+	FVector PerspectiveProjection(const FPlane3& Plane, const FVector& Point) const
+	{
+		ensure(bIsInitialized);
+		return FMath::RayPlaneIntersection(CameraState.Position, Point - CameraState.Position, (FPlane)Plane);
+	}
 
-	FVector2D PlaneCoordinates(const FVector& Point) const
+	// Given a 3D point lying in the given plane, return the UV coordinates of the point expressed in the following a
+	// two dimensional parameterization of the given plane:
+	// - the 2D origin (0,0) is located at the foot of the projection of cartesian point (0,0,0) onto the plane
+	// - the U basis vector is the camera right vector
+	// - the V basis vector is the camera up vector
+	FVector2 PlaneCoordinates(const FPlane3& Plane, const FVector& PointInPlane) const
 	{
-		float U = FVector::DotProduct(Point - CameraPlane.GetOrigin(), UBasisVector);
-		float V = FVector::DotProduct(Point - CameraPlane.GetOrigin(), VBasisVector);
-		return FVector2D{ U,V };
+		ensure(bIsInitialized);
+		FVector::FReal U = FVector::DotProduct(PointInPlane - ((FPlane)Plane).GetOrigin(), CameraState.Right());
+		FVector::FReal V = FVector::DotProduct(PointInPlane - ((FPlane)Plane).GetOrigin(), CameraState.Up());
+		return FVector2{U, V};
 	}
 };
 
@@ -129,7 +163,7 @@ protected:
 	UPROPERTY()
 	TObjectPtr<UClickDragInputBehavior> ClickDragBehavior = nullptr;
 
-	FViewCameraState CachedCameraState;
+	FCameraRectangle CameraRectangle;
 
 	FInputCapturePriority BasePriority = FInputCapturePriority(FInputCapturePriority::DEFAULT_TOOL_PRIORITY);
 
@@ -143,8 +177,4 @@ private:
 
 	bool bIsEnabled;
 	bool bIsDragging;
-	FVector2D DragStartScreenPosition;
-	FRay DragStartWorldRay;
-	FVector2D DragCurrentScreenPosition;
-	FCameraRectangle LastRectangle;
 };
diff --git a/Engine/Source/Runtime/GeometryCore/Public/Intersection/IntrSegment2Segment2.h b/Engine/Source/Runtime/GeometryCore/Public/Intersection/IntrSegment2Segment2.h
--- a/Engine/Source/Runtime/GeometryCore/Public/Intersection/IntrSegment2Segment2.h
+++ b/Engine/Source/Runtime/GeometryCore/Public/Intersection/IntrSegment2Segment2.h
@@ -138,7 +138,10 @@ public:
 		return *this;
 	}
 
+	// This is implemented in TSegment2::Intersects
+	// bool Test()
 
+	// Note: This implementation is identical to TSegment2::Intersects but also computes the intersection geometry
 	bool Find()
 	{
 		if (Result != EIntersectionResult::NotComputed)
@@ -179,7 +182,7 @@ public:
 		}
 		else if (Type == EIntersectionType::Line)
 		{
-			// Compute the location of Segment1 endpoints relative to segment0.
+			// Compute the location of Segment2 endpoints relative to Segment1.
 			TVector2<RealType> diff = Segment2.Center - Segment1.Center;
 			RealType t1 = Segment1.Direction.Dot(diff);
 			RealType tmin = t1 - Segment2.Extent;
diff --git a/Engine/Source/Runtime/GeometryCore/Public/PlaneTypes.h b/Engine/Source/Runtime/GeometryCore/Public/PlaneTypes.h
--- a/Engine/Source/Runtime/GeometryCore/Public/PlaneTypes.h
+++ b/Engine/Source/Runtime/GeometryCore/Public/PlaneTypes.h
@@ -42,6 +42,15 @@ struct TPlane3
 		Normal = VectorUtil::Normal(P0, P1, P2);
 		Constant = Normal.Dot(P0);
 	}
+	
+	explicit TPlane3(const FPlane& Plane) : TPlane3(Plane.GetNormal(), Plane.W)
+	{
+	}
+	
+	explicit operator FPlane() const
+	{
+		return FPlane(Normal.X, Normal.Y, Normal.Z, Constant);
+	}
 
 
 	/**
@@ -106,25 +115,35 @@ struct TPlane3
 	}
 
 
+	enum EClipSegmentType
+	{
+		FullyClipped,
+		FirstClipped,
+		SecondClipped,
+		NotClipped
+	};
 
 	/**
 	 * Clip line segment defined by two points against plane. Region of Segment on positive side of Plane is kept.
-	 * Note that the line may be fully clipped, in that case 0 is returned
+	 * 
 	 * @param Point0 first point of segment
 	 * @param Point1 second point of segment
-	 * @return 0 if line is fully clipped, 1 if line is partially clipped, 2 if line is not clipped
+	 * @return FullyClipped if the segment lies fully behind or exactly in the plane
+	 *         FirstClipped if Point0 became the intersection point on the plane
+	 *         SecondClipped if Point1 became the intersection point on the plane
+	 *         NotClipped if the segment lies fully in front of the plane
 	 */
-	int ClipSegment(UE::Math::TVector<RealType>& Point0, UE::Math::TVector<RealType>& Point1)
+	EClipSegmentType ClipSegment(UE::Math::TVector<RealType>& Point0, UE::Math::TVector<RealType>& Point1) const
 	{
 		RealType Dist0 = DistanceTo(Point0);
 		RealType Dist1 = DistanceTo(Point1);
 		if (Dist0 <= 0 && Dist1 <= 0)
 		{
-			return 0;
+			return EClipSegmentType::FullyClipped;
 		}
-		else if (Dist0 * Dist1 >= 0)
+		else if (Dist0 * Dist1 > 0)
 		{
-			return 2;
+			return EClipSegmentType::NotClipped;
 		}
 
 		TVector<RealType> DirectionVec = Point1 - Point0;
@@ -135,7 +154,7 @@ struct TPlane3
 		RealType NormalDot = Direction.Dot(Normal);
 		if ( TMathUtil<RealType>::Abs(NormalDot) < TMathUtil<RealType>::ZeroTolerance )
 		{
-			return 2;
+			return EClipSegmentType::NotClipped;
 		}
 
 		RealType LineT = -Dist0 / NormalDot;  // calculate line parameter for line/plane intersection
@@ -144,14 +163,15 @@ struct TPlane3
 			if (NormalDot < 0)
 			{
 				Point1 = Point0 + LineT * Direction;
+				return EClipSegmentType::SecondClipped;
 			}
 			else
 			{
 				Point0 += LineT * Direction;
+				return EClipSegmentType::FirstClipped;
 			}
-			return 1;
 		}
-		return 2;
+		return EClipSegmentType::NotClipped;
 	}
 
 };
