ModelingTools: add tooltips for UV layout types
#rnx

diff --git a/Engine/Plugins/Experimental/MeshModelingToolsetExp/Source/ModelingOperators/Public/Properties/UVLayoutProperties.h b/Engine/Plugins/Experimental/MeshModelingToolsetExp/Source/ModelingOperators/Public/Properties/UVLayoutProperties.h
--- a/Engine/Plugins/Experimental/MeshModelingToolsetExp/Source/ModelingOperators/Public/Properties/UVLayoutProperties.h
+++ b/Engine/Plugins/Experimental/MeshModelingToolsetExp/Source/ModelingOperators/Public/Properties/UVLayoutProperties.h
@@ -8,17 +8,23 @@
 
 #include "UVLayoutProperties.generated.h"
 
+/**
+ * UV Layout Strategies for the UV Layout Tool
+ */
 UENUM()
 enum class EUVLayoutType
 {
+	/** Apply Translation and Scale to existing UVs */
 	Transform,
+	/** Uniformly scale/translate UV Islands into the standard UV Unit Square */
 	Stack,
+	/** Repack all UV Islands to fit inside standard UV unit square */
 	Repack
 };
 
 
 /**
- * Standard properties
+ * UV Layout Settings
  */
 UCLASS()
 class MODELINGOPERATORS_API UUVLayoutProperties : public UInteractiveToolPropertySet
