ModelingTools: add UIMin clamps to group index in GroupPaintTool
#rnx

diff --git a/Engine/Plugins/Experimental/MeshModelingToolsetExp/Source/MeshModelingToolsExp/Public/MeshGroupPaintTool.h b/Engine/Plugins/Experimental/MeshModelingToolsetExp/Source/MeshModelingToolsExp/Public/MeshGroupPaintTool.h
--- a/Engine/Plugins/Experimental/MeshModelingToolsetExp/Source/MeshModelingToolsExp/Public/MeshGroupPaintTool.h
+++ b/Engine/Plugins/Experimental/MeshModelingToolsetExp/Source/MeshModelingToolsExp/Public/MeshGroupPaintTool.h
@@ -140,7 +140,7 @@ public:
 	bool bHitBackFaces = true;
 
 	/** The group that will be assigned to triangles */
-	UPROPERTY(EditAnywhere, Category = ActionType)
+	UPROPERTY(EditAnywhere, Category = ActionType, meta = (UIMin = 0, ClampMin = 0) )
 	int32 SetGroup = 1;
 
 	/** If true, only triangles with no group assigned will be painted */
@@ -148,7 +148,7 @@ public:
 	bool bOnlySetUngrouped = false;
 
 	/** Group to set as Erased value */
-	UPROPERTY(EditAnywhere, Category = ActionType)
+	UPROPERTY(EditAnywhere, Category = ActionType, meta = (UIMin = 0, ClampMin = 0))
 	int32 EraseGroup = 0;
 
 	/** When enabled, only the current group configured in the Paint brush is erased */
