ModelingTools: optional path tracing support inside Modeling Tools
- Add new UModelingComponentsSettings which has settings for in-tool preview mesh rendering and emitted new mesh objects
- Add ToolSetupUtil::ApplyRenderingConfigurationToPreview() utility functions, most Tools now call this function to configure UPreviewMesh/DynamicMeshComponents they create for mesh editing preview rendering (uses Settings to make setup decisions like RT on/off)
- add PreviewMesh::NotifyWorldPathTracedOutputInvalidated(), forces invalidation of current path tracing result, call from various PreviewMesh internal functions
- UCreateMeshObjectTypeProperties::ConfigureCreateMeshObjectParams() now configures new-mesh-object parameters with RT and Collision from UModelingComponentsSettings
- DynamicMeshComponent now defaults to raytracing-enabled, usage in Editor is controlled by above settings

#rnx

diff --git a/Engine/Plugins/Experimental/MeshLODToolset/Source/MeshLODToolset/Private/Tools/GenerateStaticMeshLODAssetTool.cpp b/Engine/Plugins/Experimental/MeshLODToolset/Source/MeshLODToolset/Private/Tools/GenerateStaticMeshLODAssetTool.cpp
--- a/Engine/Plugins/Experimental/MeshLODToolset/Source/MeshLODToolset/Private/Tools/GenerateStaticMeshLODAssetTool.cpp
+++ b/Engine/Plugins/Experimental/MeshLODToolset/Source/MeshLODToolset/Private/Tools/GenerateStaticMeshLODAssetTool.cpp
@@ -367,6 +367,7 @@ void UGenerateStaticMeshLODAssetTool::Setup()
 	this->OpFactory = MakeUnique<FGenerateStaticMeshLODAssetOperatorFactory>(this, PreviewTransform);
 	PreviewWithBackgroundCompute = NewObject<UMeshOpPreviewWithBackgroundCompute>(this, "Preview");
 	PreviewWithBackgroundCompute->Setup(this->TargetWorld, this->OpFactory.Get());
+	ToolSetupUtil::ApplyRenderingConfigurationToPreview(PreviewWithBackgroundCompute->PreviewMesh, nullptr);
 	PreviewWithBackgroundCompute->PreviewMesh->SetTangentsMode(EDynamicMeshComponentTangentsMode::ExternallyProvided);
 
 
diff --git a/Engine/Plugins/Experimental/MeshModelingToolsetExp/Source/MeshModelingToolsEditorOnly/Private/MergeMeshesTool.cpp b/Engine/Plugins/Experimental/MeshModelingToolsetExp/Source/MeshModelingToolsEditorOnly/Private/MergeMeshesTool.cpp
--- a/Engine/Plugins/Experimental/MeshModelingToolsetExp/Source/MeshModelingToolsEditorOnly/Private/MergeMeshesTool.cpp
+++ b/Engine/Plugins/Experimental/MeshModelingToolsetExp/Source/MeshModelingToolsEditorOnly/Private/MergeMeshesTool.cpp
@@ -101,6 +101,7 @@ void UMergeMeshesTool::Setup()
 		MeshStatisticsProperties->Update(*Compute->PreviewMesh->GetPreviewDynamicMesh());
 		UpdateAcceptWarnings(Compute->HaveEmptyResult() ? EAcceptWarning::EmptyForbidden : EAcceptWarning::NoWarning);
 	});
+	ToolSetupUtil::ApplyRenderingConfigurationToPreview(Preview->PreviewMesh, nullptr);
 
 	CreateLowQualityPreview(); // update the preview with a low-quality result
 	
diff --git a/Engine/Plugins/Experimental/MeshModelingToolsetExp/Source/MeshModelingToolsEditorOnly/Private/MeshTangentsTool.cpp b/Engine/Plugins/Experimental/MeshModelingToolsetExp/Source/MeshModelingToolsEditorOnly/Private/MeshTangentsTool.cpp
--- a/Engine/Plugins/Experimental/MeshModelingToolsetExp/Source/MeshModelingToolsEditorOnly/Private/MeshTangentsTool.cpp
+++ b/Engine/Plugins/Experimental/MeshModelingToolsetExp/Source/MeshModelingToolsEditorOnly/Private/MeshTangentsTool.cpp
@@ -69,6 +69,7 @@ void UMeshTangentsTool::Setup()
 	PreviewMesh->bBuildSpatialDataStructure = false;
 	PreviewMesh->CreateInWorld(TargetActor->GetWorld(), FTransform::Identity);
 	PreviewMesh->SetTransform((FTransform)UE::ToolTarget::GetLocalToWorldTransform(Target));
+	ToolSetupUtil::ApplyRenderingConfigurationToPreview(PreviewMesh, nullptr);
 	// configure materials
 	FComponentMaterialSet MaterialSet = UE::ToolTarget::GetMaterialSet(Target);
 	PreviewMesh->SetMaterials(MaterialSet.Materials);
diff --git a/Engine/Plugins/Experimental/MeshModelingToolsetExp/Source/MeshModelingToolsEditorOnly/Private/MeshToVolumeTool.cpp b/Engine/Plugins/Experimental/MeshModelingToolsetExp/Source/MeshModelingToolsEditorOnly/Private/MeshToVolumeTool.cpp
--- a/Engine/Plugins/Experimental/MeshModelingToolsetExp/Source/MeshModelingToolsEditorOnly/Private/MeshToVolumeTool.cpp
+++ b/Engine/Plugins/Experimental/MeshModelingToolsetExp/Source/MeshModelingToolsEditorOnly/Private/MeshToVolumeTool.cpp
@@ -59,6 +59,7 @@ void UMeshToVolumeTool::Setup()
 	PreviewMesh->bBuildSpatialDataStructure = false;
 	PreviewMesh->CreateInWorld(UE::ToolTarget::GetTargetActor(Target)->GetWorld(), FTransform::Identity);
 	PreviewMesh->SetTransform((FTransform)UE::ToolTarget::GetLocalToWorldTransform(Target));
+	ToolSetupUtil::ApplyRenderingConfigurationToPreview(PreviewMesh, nullptr);
 
 	PreviewMesh->SetTangentsMode(EDynamicMeshComponentTangentsMode::AutoCalculated);
 	PreviewMesh->ReplaceMesh(UE::ToolTarget::GetDynamicMeshCopy(Target));
diff --git a/Engine/Plugins/Experimental/MeshModelingToolsetExp/Source/MeshModelingToolsEditorOnly/Private/ParameterizeMeshTool.cpp b/Engine/Plugins/Experimental/MeshModelingToolsetExp/Source/MeshModelingToolsEditorOnly/Private/ParameterizeMeshTool.cpp
--- a/Engine/Plugins/Experimental/MeshModelingToolsetExp/Source/MeshModelingToolsEditorOnly/Private/ParameterizeMeshTool.cpp
+++ b/Engine/Plugins/Experimental/MeshModelingToolsetExp/Source/MeshModelingToolsEditorOnly/Private/ParameterizeMeshTool.cpp
@@ -41,6 +41,7 @@ void UParameterizeMeshTool::Setup()
 
 	Preview = NewObject<UMeshOpPreviewWithBackgroundCompute>(this);
 	Preview->Setup(this->TargetWorld, this);
+	ToolSetupUtil::ApplyRenderingConfigurationToPreview(Preview->PreviewMesh, nullptr);
 	Preview->PreviewMesh->SetTangentsMode(EDynamicMeshComponentTangentsMode::AutoCalculated);
 	Preview->PreviewMesh->ReplaceMesh(*InputMesh);
 	Preview->ConfigureMaterials(MaterialSet.Materials, ToolSetupUtil::GetDefaultWorkingMaterial(GetToolManager()));
diff --git a/Engine/Plugins/Experimental/MeshModelingToolsetExp/Source/MeshModelingToolsEditorOnly/Private/PolygonOnMeshTool.cpp b/Engine/Plugins/Experimental/MeshModelingToolsetExp/Source/MeshModelingToolsEditorOnly/Private/PolygonOnMeshTool.cpp
--- a/Engine/Plugins/Experimental/MeshModelingToolsetExp/Source/MeshModelingToolsEditorOnly/Private/PolygonOnMeshTool.cpp
+++ b/Engine/Plugins/Experimental/MeshModelingToolsetExp/Source/MeshModelingToolsEditorOnly/Private/PolygonOnMeshTool.cpp
@@ -217,6 +217,7 @@ void UPolygonOnMeshTool::SetupPreview()
 {
 	Preview = NewObject<UMeshOpPreviewWithBackgroundCompute>(this);
 	Preview->Setup(this->TargetWorld, this);
+	ToolSetupUtil::ApplyRenderingConfigurationToPreview(Preview->PreviewMesh, nullptr);
 	Preview->PreviewMesh->SetTangentsMode(EDynamicMeshComponentTangentsMode::AutoCalculated);
 
 	FComponentMaterialSet MaterialSet;
diff --git a/Engine/Plugins/Experimental/MeshModelingToolsetExp/Source/MeshModelingToolsEditorOnly/Private/SimplifyMeshTool.cpp b/Engine/Plugins/Experimental/MeshModelingToolsetExp/Source/MeshModelingToolsEditorOnly/Private/SimplifyMeshTool.cpp
--- a/Engine/Plugins/Experimental/MeshModelingToolsetExp/Source/MeshModelingToolsEditorOnly/Private/SimplifyMeshTool.cpp
+++ b/Engine/Plugins/Experimental/MeshModelingToolsetExp/Source/MeshModelingToolsEditorOnly/Private/SimplifyMeshTool.cpp
@@ -69,6 +69,7 @@ void USimplifyMeshTool::Setup()
 
 	Preview = NewObject<UMeshOpPreviewWithBackgroundCompute>(this);
 	Preview->Setup(this->TargetWorld, this);
+	ToolSetupUtil::ApplyRenderingConfigurationToPreview(Preview->PreviewMesh, nullptr);
 	FComponentMaterialSet MaterialSet = UE::ToolTarget::GetMaterialSet(Target);
 	Preview->ConfigureMaterials( MaterialSet.Materials,
 		ToolSetupUtil::GetDefaultWorkingMaterial(GetToolManager())
diff --git a/Engine/Plugins/Experimental/MeshModelingToolsetExp/Source/MeshModelingToolsEditorOnly/Private/SubdividePolyTool.cpp b/Engine/Plugins/Experimental/MeshModelingToolsetExp/Source/MeshModelingToolsEditorOnly/Private/SubdividePolyTool.cpp
--- a/Engine/Plugins/Experimental/MeshModelingToolsetExp/Source/MeshModelingToolsEditorOnly/Private/SubdividePolyTool.cpp
+++ b/Engine/Plugins/Experimental/MeshModelingToolsetExp/Source/MeshModelingToolsEditorOnly/Private/SubdividePolyTool.cpp
@@ -193,6 +193,7 @@ void USubdividePolyTool::Setup()
 		return;
 	}
 	PreviewMesh->CreateInWorld(TargetWorld, FTransform::Identity);
+	ToolSetupUtil::ApplyRenderingConfigurationToPreview(PreviewMesh, nullptr);
 
 	PreviewMesh->SetTransform(TargetComponent->GetWorldTransform());
 	PreviewMesh->UpdatePreview(OriginalMesh.Get());
diff --git a/Engine/Plugins/Experimental/MeshModelingToolsetExp/Source/MeshModelingToolsEditorOnly/Private/VoxelCSGMeshesTool.cpp b/Engine/Plugins/Experimental/MeshModelingToolsetExp/Source/MeshModelingToolsEditorOnly/Private/VoxelCSGMeshesTool.cpp
--- a/Engine/Plugins/Experimental/MeshModelingToolsetExp/Source/MeshModelingToolsEditorOnly/Private/VoxelCSGMeshesTool.cpp
+++ b/Engine/Plugins/Experimental/MeshModelingToolsetExp/Source/MeshModelingToolsEditorOnly/Private/VoxelCSGMeshesTool.cpp
@@ -93,6 +93,7 @@ void UVoxelCSGMeshesTool::Setup()
 	// initialize the PreviewMesh+BackgroundCompute object
 	Preview = NewObject<UMeshOpPreviewWithBackgroundCompute>(this, "Preview");
 	Preview->Setup(this->TargetWorld, this);
+	ToolSetupUtil::ApplyRenderingConfigurationToPreview(Preview->PreviewMesh, nullptr);
 	Preview->OnMeshUpdated.AddLambda([this](UMeshOpPreviewWithBackgroundCompute* Compute) {
 		MeshStatisticsProperties->Update(*Compute->PreviewMesh->GetPreviewDynamicMesh());
 		UpdateAcceptWarnings(Compute->HaveEmptyResult() ? EAcceptWarning::EmptyForbidden : EAcceptWarning::NoWarning);
diff --git a/Engine/Plugins/Experimental/MeshModelingToolsetExp/Source/MeshModelingToolsExp/Private/AddPatchTool.cpp b/Engine/Plugins/Experimental/MeshModelingToolsetExp/Source/MeshModelingToolsExp/Private/AddPatchTool.cpp
--- a/Engine/Plugins/Experimental/MeshModelingToolsetExp/Source/MeshModelingToolsExp/Private/AddPatchTool.cpp
+++ b/Engine/Plugins/Experimental/MeshModelingToolsetExp/Source/MeshModelingToolsExp/Private/AddPatchTool.cpp
@@ -2,6 +2,7 @@
 
 #include "AddPatchTool.h"
 #include "ToolBuilderUtil.h"
+#include "ToolSetupUtil.h"
 #include "InteractiveToolManager.h"
 #include "BaseBehaviors/MouseHoverBehavior.h"
 #include "Selection/ToolSelectionUtil.h"
@@ -73,6 +74,7 @@ void UAddPatchTool::Setup()
 	// create preview mesh object
 	PreviewMesh = NewObject<UPreviewMesh>(this);
 	PreviewMesh->CreateInWorld(TargetWorld, FTransform::Identity);
+	ToolSetupUtil::ApplyRenderingConfigurationToPreview(PreviewMesh, nullptr);
 	PreviewMesh->SetVisible(false);
 	PreviewMesh->SetMaterial(MaterialProperties->Material.Get());
 	GeneratePreviewBaseMesh();
diff --git a/Engine/Plugins/Experimental/MeshModelingToolsetExp/Source/MeshModelingToolsExp/Private/BakeMeshAttributeMapsToolBase.cpp b/Engine/Plugins/Experimental/MeshModelingToolsetExp/Source/MeshModelingToolsExp/Private/BakeMeshAttributeMapsToolBase.cpp
--- a/Engine/Plugins/Experimental/MeshModelingToolsetExp/Source/MeshModelingToolsExp/Private/BakeMeshAttributeMapsToolBase.cpp
+++ b/Engine/Plugins/Experimental/MeshModelingToolsetExp/Source/MeshModelingToolsExp/Private/BakeMeshAttributeMapsToolBase.cpp
@@ -3,6 +3,7 @@
 #include "BakeMeshAttributeMapsToolBase.h"
 #include "InteractiveToolManager.h"
 #include "ToolBuilderUtil.h"
+#include "ToolSetupUtil.h"
 
 #include "DynamicMesh/DynamicMesh3.h"
 #include "DynamicMesh/DynamicMeshAttributeSet.h"
@@ -61,6 +62,7 @@ void UBakeMeshAttributeMapsToolBase::Setup()
 	const UE::Geometry::FTransform3d BaseToWorld = UE::ToolTarget::GetLocalToWorldTransform(Targets[0]);
 	PreviewMesh = NewObject<UPreviewMesh>(this);
 	PreviewMesh->CreateInWorld(TargetWorld, FTransform::Identity);
+	ToolSetupUtil::ApplyRenderingConfigurationToPreview(PreviewMesh, nullptr);
 	PreviewMesh->SetTransform(static_cast<FTransform>(BaseToWorld));
 	PreviewMesh->SetTangentsMode(EDynamicMeshComponentTangentsMode::ExternallyProvided);
 	PreviewMesh->ReplaceMesh(InputMesh);
diff --git a/Engine/Plugins/Experimental/MeshModelingToolsetExp/Source/MeshModelingToolsExp/Private/BakeMeshAttributeVertexTool.cpp b/Engine/Plugins/Experimental/MeshModelingToolsetExp/Source/MeshModelingToolsExp/Private/BakeMeshAttributeVertexTool.cpp
--- a/Engine/Plugins/Experimental/MeshModelingToolsetExp/Source/MeshModelingToolsExp/Private/BakeMeshAttributeVertexTool.cpp
+++ b/Engine/Plugins/Experimental/MeshModelingToolsetExp/Source/MeshModelingToolsExp/Private/BakeMeshAttributeVertexTool.cpp
@@ -3,6 +3,7 @@
 #include "BakeMeshAttributeVertexTool.h"
 #include "InteractiveToolManager.h"
 #include "ToolBuilderUtil.h"
+#include "ToolSetupUtil.h"
 
 #include "Materials/Material.h"
 #include "Materials/MaterialInstanceDynamic.h"
@@ -291,6 +292,7 @@ void UBakeMeshAttributeVertexTool::Setup()
 	const FDynamicMesh3 InputMeshWithTangents = UE::ToolTarget::GetDynamicMeshCopy(Targets[0], true);
 	PreviewMesh = NewObject<UPreviewMesh>(this);
 	PreviewMesh->CreateInWorld(TargetWorld, FTransform::Identity);
+	ToolSetupUtil::ApplyRenderingConfigurationToPreview(PreviewMesh, nullptr);
 	PreviewMesh->SetTransform(static_cast<FTransform>(UE::ToolTarget::GetLocalToWorldTransform(Targets[0])));
 	PreviewMesh->SetTangentsMode(EDynamicMeshComponentTangentsMode::ExternallyProvided);
 	PreviewMesh->ReplaceMesh(InputMeshWithTangents);
diff --git a/Engine/Plugins/Experimental/MeshModelingToolsetExp/Source/MeshModelingToolsExp/Private/ConvertToPolygonsTool.cpp b/Engine/Plugins/Experimental/MeshModelingToolsetExp/Source/MeshModelingToolsExp/Private/ConvertToPolygonsTool.cpp
--- a/Engine/Plugins/Experimental/MeshModelingToolsetExp/Source/MeshModelingToolsExp/Private/ConvertToPolygonsTool.cpp
+++ b/Engine/Plugins/Experimental/MeshModelingToolsetExp/Source/MeshModelingToolsExp/Private/ConvertToPolygonsTool.cpp
@@ -194,6 +194,7 @@ void UConvertToPolygonsTool::Setup()
 
 		PreviewCompute = NewObject<UMeshOpPreviewWithBackgroundCompute>(ConvertToPolygonsOperatorFactory);
 		PreviewCompute->Setup(this->TargetWorld, ConvertToPolygonsOperatorFactory);
+		ToolSetupUtil::ApplyRenderingConfigurationToPreview(PreviewCompute->PreviewMesh, Target);
 		PreviewCompute->SetIsMeshTopologyConstant(true, EMeshRenderAttributeFlags::Positions | EMeshRenderAttributeFlags::VertexNormals);
 
 		// Give the preview something to display
diff --git a/Engine/Plugins/Experimental/MeshModelingToolsetExp/Source/MeshModelingToolsExp/Private/CubeGridTool.cpp b/Engine/Plugins/Experimental/MeshModelingToolsetExp/Source/MeshModelingToolsExp/Private/CubeGridTool.cpp
--- a/Engine/Plugins/Experimental/MeshModelingToolsetExp/Source/MeshModelingToolsExp/Private/CubeGridTool.cpp
+++ b/Engine/Plugins/Experimental/MeshModelingToolsetExp/Source/MeshModelingToolsExp/Private/CubeGridTool.cpp
@@ -667,6 +667,7 @@ void UCubeGridTool::Setup()
 
 	Preview = NewObject<UMeshOpPreviewWithBackgroundCompute>(this);
 	Preview->Setup(TargetWorld, this);
+	ToolSetupUtil::ApplyRenderingConfigurationToPreview(Preview->PreviewMesh, Target);
 	Preview->PreviewMesh->SetTangentsMode(EDynamicMeshComponentTangentsMode::AutoCalculated);
 	if (Target)
 	{
diff --git a/Engine/Plugins/Experimental/MeshModelingToolsetExp/Source/MeshModelingToolsExp/Private/DeformMeshPolygonsTool.cpp b/Engine/Plugins/Experimental/MeshModelingToolsetExp/Source/MeshModelingToolsExp/Private/DeformMeshPolygonsTool.cpp
--- a/Engine/Plugins/Experimental/MeshModelingToolsetExp/Source/MeshModelingToolsExp/Private/DeformMeshPolygonsTool.cpp
+++ b/Engine/Plugins/Experimental/MeshModelingToolsetExp/Source/MeshModelingToolsExp/Private/DeformMeshPolygonsTool.cpp
@@ -12,6 +12,7 @@
 #include "SegmentTypes.h"
 #include "Solvers/ConstrainedMeshDeformer.h"
 #include "ToolBuilderUtil.h"
+#include "ToolSetupUtil.h"
 #include "ToolSceneQueriesUtil.h"
 #include "ModelingToolTargetUtil.h"
 
@@ -852,6 +853,7 @@ void UDeformMeshPolygonsTool::Setup()
 	DynamicMeshComponent->RegisterComponent();
 	WorldTransform = UE::ToolTarget::GetLocalToWorldTransform(Target);
 	DynamicMeshComponent->SetWorldTransform((FTransform)WorldTransform);
+	ToolSetupUtil::ApplyRenderingConfigurationToPreview(DynamicMeshComponent, Target);
 
 	// set materials
 	FComponentMaterialSet MaterialSet = UE::ToolTarget::GetMaterialSet(Target);
diff --git a/Engine/Plugins/Experimental/MeshModelingToolsetExp/Source/MeshModelingToolsExp/Private/DisplaceMeshTool.cpp b/Engine/Plugins/Experimental/MeshModelingToolsetExp/Source/MeshModelingToolsExp/Private/DisplaceMeshTool.cpp
--- a/Engine/Plugins/Experimental/MeshModelingToolsetExp/Source/MeshModelingToolsExp/Private/DisplaceMeshTool.cpp
+++ b/Engine/Plugins/Experimental/MeshModelingToolsetExp/Source/MeshModelingToolsExp/Private/DisplaceMeshTool.cpp
@@ -828,6 +828,7 @@ void UDisplaceMeshTool::Setup()
 	DynamicMeshComponent->RegisterComponent();
 	DynamicMeshComponent->SetWorldTransform((FTransform)UE::ToolTarget::GetLocalToWorldTransform(Target));
 	DynamicMeshComponent->bExplicitShowWireframe = CommonProperties->bShowWireframe;
+	ToolSetupUtil::ApplyRenderingConfigurationToPreview(DynamicMeshComponent, Target);
 
 	// transfer materials
 	FComponentMaterialSet MaterialSet = UE::ToolTarget::GetMaterialSet(Target);
diff --git a/Engine/Plugins/Experimental/MeshModelingToolsetExp/Source/MeshModelingToolsExp/Private/DrawPolyPathTool.cpp b/Engine/Plugins/Experimental/MeshModelingToolsetExp/Source/MeshModelingToolsExp/Private/DrawPolyPathTool.cpp
--- a/Engine/Plugins/Experimental/MeshModelingToolsetExp/Source/MeshModelingToolsExp/Private/DrawPolyPathTool.cpp
+++ b/Engine/Plugins/Experimental/MeshModelingToolsetExp/Source/MeshModelingToolsExp/Private/DrawPolyPathTool.cpp
@@ -687,6 +687,7 @@ void UDrawPolyPathTool::InitializePreviewMesh()
 	{
 		EditPreview = NewObject<UPolyEditPreviewMesh>(this);
 		EditPreview->CreateInWorld(TargetWorld, FTransform::Identity);
+		ToolSetupUtil::ApplyRenderingConfigurationToPreview(EditPreview, nullptr);
 		if ( MaterialProperties->Material == nullptr )
 		{
 			EditPreview->SetMaterial(
diff --git a/Engine/Plugins/Experimental/MeshModelingToolsetExp/Source/MeshModelingToolsExp/Private/DynamicMeshBrushTool.cpp b/Engine/Plugins/Experimental/MeshModelingToolsetExp/Source/MeshModelingToolsExp/Private/DynamicMeshBrushTool.cpp
--- a/Engine/Plugins/Experimental/MeshModelingToolsetExp/Source/MeshModelingToolsExp/Private/DynamicMeshBrushTool.cpp
+++ b/Engine/Plugins/Experimental/MeshModelingToolsetExp/Source/MeshModelingToolsExp/Private/DynamicMeshBrushTool.cpp
@@ -3,6 +3,7 @@
 #include "DynamicMeshBrushTool.h"
 #include "InteractiveToolManager.h"
 #include "ToolBuilderUtil.h"
+#include "ToolSetupUtil.h"
 #include "ModelingToolTargetUtil.h"
 
 #include "ExplicitUseGeometryMathTypes.h"		// using UE::Geometry::(math types)
@@ -27,6 +28,7 @@ void UDynamicMeshBrushTool::Setup()
 	PreviewMesh->CreateInWorld(UE::ToolTarget::GetTargetActor(Target)->GetWorld(), FTransform::Identity);
 	UE::Geometry::FTransform3d LocalToWorldTransform = UE::ToolTarget::GetLocalToWorldTransform(Target);
 	PreviewMesh->SetTransform((FTransform)LocalToWorldTransform);
+	ToolSetupUtil::ApplyRenderingConfigurationToPreview(PreviewMesh, Target);
 
 	FComponentMaterialSet MaterialSet = UE::ToolTarget::GetMaterialSet(Target);
 	PreviewMesh->SetMaterials(MaterialSet.Materials);
diff --git a/Engine/Plugins/Experimental/MeshModelingToolsetExp/Source/MeshModelingToolsExp/Private/DynamicMeshSculptTool.cpp b/Engine/Plugins/Experimental/MeshModelingToolsetExp/Source/MeshModelingToolsExp/Private/DynamicMeshSculptTool.cpp
--- a/Engine/Plugins/Experimental/MeshModelingToolsetExp/Source/MeshModelingToolsExp/Private/DynamicMeshSculptTool.cpp
+++ b/Engine/Plugins/Experimental/MeshModelingToolsetExp/Source/MeshModelingToolsExp/Private/DynamicMeshSculptTool.cpp
@@ -104,6 +104,7 @@ void UDynamicMeshSculptTool::Setup()
 	DynamicMeshComponent->SetShadowsEnabled(false);
 	DynamicMeshComponent->SetupAttachment(UE::ToolTarget::GetTargetActor(Target)->GetRootComponent());
 	DynamicMeshComponent->RegisterComponent();
+	ToolSetupUtil::ApplyRenderingConfigurationToPreview(DynamicMeshComponent, Target); 
 
 	// initialize from LOD-0 MeshDescription
 	DynamicMeshComponent->SetMesh(UE::ToolTarget::GetDynamicMeshCopy(Target));
diff --git a/Engine/Plugins/Experimental/MeshModelingToolsetExp/Source/MeshModelingToolsExp/Private/EditNormalsTool.cpp b/Engine/Plugins/Experimental/MeshModelingToolsetExp/Source/MeshModelingToolsExp/Private/EditNormalsTool.cpp
--- a/Engine/Plugins/Experimental/MeshModelingToolsetExp/Source/MeshModelingToolsExp/Private/EditNormalsTool.cpp
+++ b/Engine/Plugins/Experimental/MeshModelingToolsetExp/Source/MeshModelingToolsExp/Private/EditNormalsTool.cpp
@@ -155,6 +155,7 @@ void UEditNormalsTool::UpdateNumPreviews()
 
 			UMeshOpPreviewWithBackgroundCompute* Preview = Previews.Add_GetRef(NewObject<UMeshOpPreviewWithBackgroundCompute>(OpFactory, "Preview"));
 			Preview->Setup(this->TargetWorld, OpFactory);
+			ToolSetupUtil::ApplyRenderingConfigurationToPreview(Preview->PreviewMesh, Targets[PreviewIdx]);
 			Preview->PreviewMesh->SetTangentsMode(EDynamicMeshComponentTangentsMode::AutoCalculated);
 
 			FComponentMaterialSet MaterialSet;
diff --git a/Engine/Plugins/Experimental/MeshModelingToolsetExp/Source/MeshModelingToolsExp/Private/EditUVIslandsTool.cpp b/Engine/Plugins/Experimental/MeshModelingToolsetExp/Source/MeshModelingToolsExp/Private/EditUVIslandsTool.cpp
--- a/Engine/Plugins/Experimental/MeshModelingToolsetExp/Source/MeshModelingToolsExp/Private/EditUVIslandsTool.cpp
+++ b/Engine/Plugins/Experimental/MeshModelingToolsetExp/Source/MeshModelingToolsExp/Private/EditUVIslandsTool.cpp
@@ -47,6 +47,7 @@ void UEditUVIslandsTool::Setup()
 	DynamicMeshComponent->SetupAttachment(UE::ToolTarget::GetTargetActor(Target)->GetRootComponent());
 	DynamicMeshComponent->RegisterComponent();
 	DynamicMeshComponent->SetWorldTransform((FTransform)UE::ToolTarget::GetLocalToWorldTransform(Target));
+	ToolSetupUtil::ApplyRenderingConfigurationToPreview(DynamicMeshComponent, Target);
 	WorldTransform = UE::Geometry::FTransform3d(DynamicMeshComponent->GetComponentTransform());
 
 	// set materials
diff --git a/Engine/Plugins/Experimental/MeshModelingToolsetExp/Source/MeshModelingToolsExp/Private/HoleFillTool.cpp b/Engine/Plugins/Experimental/MeshModelingToolsetExp/Source/MeshModelingToolsExp/Private/HoleFillTool.cpp
--- a/Engine/Plugins/Experimental/MeshModelingToolsetExp/Source/MeshModelingToolsExp/Private/HoleFillTool.cpp
+++ b/Engine/Plugins/Experimental/MeshModelingToolsetExp/Source/MeshModelingToolsExp/Private/HoleFillTool.cpp
@@ -316,6 +316,7 @@ void UHoleFillTool::SetupPreview()
 
 	Preview = NewObject<UMeshOpPreviewWithBackgroundCompute>(OpFactory, "Preview");
 	Preview->Setup(this->TargetWorld, OpFactory);
+	ToolSetupUtil::ApplyRenderingConfigurationToPreview(Preview->PreviewMesh, Target);
 
 	FComponentMaterialSet MaterialSet;
 	Cast<IMaterialProvider>(Target)->GetMaterialSet(MaterialSet);
diff --git a/Engine/Plugins/Experimental/MeshModelingToolsetExp/Source/MeshModelingToolsExp/Private/LatticeDeformerTool.cpp b/Engine/Plugins/Experimental/MeshModelingToolsetExp/Source/MeshModelingToolsExp/Private/LatticeDeformerTool.cpp
--- a/Engine/Plugins/Experimental/MeshModelingToolsetExp/Source/MeshModelingToolsExp/Private/LatticeDeformerTool.cpp
+++ b/Engine/Plugins/Experimental/MeshModelingToolsetExp/Source/MeshModelingToolsExp/Private/LatticeDeformerTool.cpp
@@ -372,8 +372,8 @@ void ULatticeDeformerTool::StartPreview()
 	LatticeDeformOpCreator->LatticeDeformerTool = this;
 
 	Preview = NewObject<UMeshOpPreviewWithBackgroundCompute>(LatticeDeformOpCreator);
-
 	Preview->Setup(TargetWorld, LatticeDeformOpCreator);
+	ToolSetupUtil::ApplyRenderingConfigurationToPreview(Preview->PreviewMesh, Target);
 
 	Preview->SetIsMeshTopologyConstant(true, EMeshRenderAttributeFlags::Positions | EMeshRenderAttributeFlags::VertexNormals);
 
diff --git a/Engine/Plugins/Experimental/MeshModelingToolsetExp/Source/MeshModelingToolsExp/Private/MeshInspectorTool.cpp b/Engine/Plugins/Experimental/MeshModelingToolsetExp/Source/MeshModelingToolsExp/Private/MeshInspectorTool.cpp
--- a/Engine/Plugins/Experimental/MeshModelingToolsetExp/Source/MeshModelingToolsExp/Private/MeshInspectorTool.cpp
+++ b/Engine/Plugins/Experimental/MeshModelingToolsetExp/Source/MeshModelingToolsExp/Private/MeshInspectorTool.cpp
@@ -52,6 +52,7 @@ void UMeshInspectorTool::Setup()
 	PreviewMesh = NewObject<UPreviewMesh>(this);
 	PreviewMesh->bBuildSpatialDataStructure = false;
 	PreviewMesh->CreateInWorld(TargetWorld, FTransform::Identity);
+	ToolSetupUtil::ApplyRenderingConfigurationToPreview(PreviewMesh, Target);
 	PreviewMesh->SetTransform((FTransform)UE::ToolTarget::GetLocalToWorldTransform(Target));
 
 	FComponentMaterialSet MaterialSet = UE::ToolTarget::GetMaterialSet(Target);
diff --git a/Engine/Plugins/Experimental/MeshModelingToolsetExp/Source/MeshModelingToolsExp/Private/MeshSpaceDeformerTool.cpp b/Engine/Plugins/Experimental/MeshModelingToolsetExp/Source/MeshModelingToolsExp/Private/MeshSpaceDeformerTool.cpp
--- a/Engine/Plugins/Experimental/MeshModelingToolsetExp/Source/MeshModelingToolsExp/Private/MeshSpaceDeformerTool.cpp
+++ b/Engine/Plugins/Experimental/MeshModelingToolsetExp/Source/MeshModelingToolsExp/Private/MeshSpaceDeformerTool.cpp
@@ -153,6 +153,7 @@ void UMeshSpaceDeformerTool::Setup()
 
 	OriginalMeshPreview = NewObject<UPreviewMesh>();
 	OriginalMeshPreview->CreateInWorld(TargetWorld, MeshTransform);
+	ToolSetupUtil::ApplyRenderingConfigurationToPreview(OriginalMeshPreview, Target);
 	OriginalMeshPreview->UpdatePreview(OriginalDynamicMesh.Get());
 	OriginalMeshPreview->SetMaterial(0, ToolSetupUtil::GetCustomDepthOffsetMaterial(GetToolManager(), FLinearColor::White,
 		-0.5,  // depth offset, 0.5% inward
@@ -278,6 +279,7 @@ void UMeshSpaceDeformerTool::Setup()
 
 		Preview = NewObject<UMeshOpPreviewWithBackgroundCompute>(DeformerOperatorFactory, "Preview");
 		Preview->Setup(this->TargetWorld, DeformerOperatorFactory);
+		ToolSetupUtil::ApplyRenderingConfigurationToPreview(Preview->PreviewMesh, Target);
 		Preview->PreviewMesh->SetTangentsMode(EDynamicMeshComponentTangentsMode::AutoCalculated);
 
 		Preview->SetIsMeshTopologyConstant(true, EMeshRenderAttributeFlags::Positions | EMeshRenderAttributeFlags::VertexNormals);
diff --git a/Engine/Plugins/Experimental/MeshModelingToolsetExp/Source/MeshModelingToolsExp/Private/MirrorTool.cpp b/Engine/Plugins/Experimental/MeshModelingToolsetExp/Source/MeshModelingToolsExp/Private/MirrorTool.cpp
--- a/Engine/Plugins/Experimental/MeshModelingToolsetExp/Source/MeshModelingToolsExp/Private/MirrorTool.cpp
+++ b/Engine/Plugins/Experimental/MeshModelingToolsetExp/Source/MeshModelingToolsExp/Private/MirrorTool.cpp
@@ -293,6 +293,7 @@ void UMirrorTool::SetupPreviews()
 		UMeshOpPreviewWithBackgroundCompute* Preview = Previews.Add_GetRef(
 			NewObject<UMeshOpPreviewWithBackgroundCompute>(MirrorOpCreator, "Preview"));
 		Preview->Setup(TargetWorld, MirrorOpCreator);
+		ToolSetupUtil::ApplyRenderingConfigurationToPreview(Preview->PreviewMesh, nullptr);
 		Preview->PreviewMesh->SetTangentsMode(EDynamicMeshComponentTangentsMode::AutoCalculated);
 
 		FComponentMaterialSet MaterialSet;
diff --git a/Engine/Plugins/Experimental/MeshModelingToolsetExp/Source/MeshModelingToolsExp/Private/Physics/ExtractCollisionGeometryTool.cpp b/Engine/Plugins/Experimental/MeshModelingToolsetExp/Source/MeshModelingToolsExp/Private/Physics/ExtractCollisionGeometryTool.cpp
--- a/Engine/Plugins/Experimental/MeshModelingToolsetExp/Source/MeshModelingToolsExp/Private/Physics/ExtractCollisionGeometryTool.cpp
+++ b/Engine/Plugins/Experimental/MeshModelingToolsetExp/Source/MeshModelingToolsExp/Private/Physics/ExtractCollisionGeometryTool.cpp
@@ -70,6 +70,7 @@ void UExtractCollisionGeometryTool::Setup()
 	{
 		return LinearColors::SelectFColor(Mesh->GetTriangleGroup(TriangleID));
 	});
+	ToolSetupUtil::ApplyRenderingConfigurationToPreview(PreviewMesh, nullptr);
 
 	OutputTypeProperties = NewObject<UCreateMeshObjectTypeProperties>(this);
 	OutputTypeProperties->OutputType = UCreateMeshObjectTypeProperties::VolumeIdentifier;		// prefer volumes for extracting simple collision
diff --git a/Engine/Plugins/Experimental/MeshModelingToolsetExp/Source/MeshModelingToolsExp/Private/PlaneCutTool.cpp b/Engine/Plugins/Experimental/MeshModelingToolsetExp/Source/MeshModelingToolsExp/Private/PlaneCutTool.cpp
--- a/Engine/Plugins/Experimental/MeshModelingToolsetExp/Source/MeshModelingToolsExp/Private/PlaneCutTool.cpp
+++ b/Engine/Plugins/Experimental/MeshModelingToolsetExp/Source/MeshModelingToolsExp/Private/PlaneCutTool.cpp
@@ -211,6 +211,7 @@ void UPlaneCutTool::SetupPreviews()
 		CutSide->ComponentIndex = PreviewIdx;
 		UMeshOpPreviewWithBackgroundCompute* Preview = Previews.Add_GetRef(NewObject<UMeshOpPreviewWithBackgroundCompute>(CutSide, "Preview"));
 		Preview->Setup(this->TargetWorld, CutSide);
+		ToolSetupUtil::ApplyRenderingConfigurationToPreview(Preview->PreviewMesh, nullptr);
 		Preview->PreviewMesh->SetTangentsMode(EDynamicMeshComponentTangentsMode::AutoCalculated);
 
 		FComponentMaterialSet MaterialSet;
diff --git a/Engine/Plugins/Experimental/MeshModelingToolsetExp/Source/MeshModelingToolsExp/Private/RecomputeUVsTool.cpp b/Engine/Plugins/Experimental/MeshModelingToolsetExp/Source/MeshModelingToolsExp/Private/RecomputeUVsTool.cpp
--- a/Engine/Plugins/Experimental/MeshModelingToolsetExp/Source/MeshModelingToolsExp/Private/RecomputeUVsTool.cpp
+++ b/Engine/Plugins/Experimental/MeshModelingToolsetExp/Source/MeshModelingToolsExp/Private/RecomputeUVsTool.cpp
@@ -44,6 +44,7 @@ void URecomputeUVsTool::Setup()
 
 	Preview = NewObject<UMeshOpPreviewWithBackgroundCompute>(this);
 	Preview->Setup(this->TargetWorld, this);
+	ToolSetupUtil::ApplyRenderingConfigurationToPreview(Preview->PreviewMesh, Target);
 	Preview->PreviewMesh->SetTangentsMode(EDynamicMeshComponentTangentsMode::AutoCalculated);
 	Preview->PreviewMesh->ReplaceMesh(*InputMesh);
 	Preview->ConfigureMaterials(MaterialSet.Materials, ToolSetupUtil::GetDefaultWorkingMaterial(GetToolManager()) );
diff --git a/Engine/Plugins/Experimental/MeshModelingToolsetExp/Source/MeshModelingToolsExp/Private/RemeshMeshTool.cpp b/Engine/Plugins/Experimental/MeshModelingToolsetExp/Source/MeshModelingToolsExp/Private/RemeshMeshTool.cpp
--- a/Engine/Plugins/Experimental/MeshModelingToolsetExp/Source/MeshModelingToolsExp/Private/RemeshMeshTool.cpp
+++ b/Engine/Plugins/Experimental/MeshModelingToolsetExp/Source/MeshModelingToolsExp/Private/RemeshMeshTool.cpp
@@ -115,6 +115,7 @@ void URemeshMeshTool::Setup()
 	TargetComponent->SetOwnerVisibility(false);
 	Preview = NewObject<UMeshOpPreviewWithBackgroundCompute>(this);
 	Preview->Setup(this->TargetWorld, this);
+	ToolSetupUtil::ApplyRenderingConfigurationToPreview(Preview->PreviewMesh, Targets[0]);
 
 	FComponentMaterialSet MaterialSet = UE::ToolTarget::GetMaterialSet(Targets[0]);
 	Preview->ConfigureMaterials( MaterialSet.Materials,
diff --git a/Engine/Plugins/Experimental/MeshModelingToolsetExp/Source/MeshModelingToolsExp/Private/RemoveOccludedTrianglesTool.cpp b/Engine/Plugins/Experimental/MeshModelingToolsetExp/Source/MeshModelingToolsExp/Private/RemoveOccludedTrianglesTool.cpp
--- a/Engine/Plugins/Experimental/MeshModelingToolsetExp/Source/MeshModelingToolsExp/Private/RemoveOccludedTrianglesTool.cpp
+++ b/Engine/Plugins/Experimental/MeshModelingToolsetExp/Source/MeshModelingToolsExp/Private/RemoveOccludedTrianglesTool.cpp
@@ -256,6 +256,7 @@ void URemoveOccludedTrianglesTool::SetupPreviews()
 
 			UMeshOpPreviewWithBackgroundCompute* Preview = Previews.Add_GetRef(NewObject<UMeshOpPreviewWithBackgroundCompute>(OpFactory, "Preview"));
 			Preview->Setup(this->TargetWorld, OpFactory);
+			ToolSetupUtil::ApplyRenderingConfigurationToPreview(Preview->PreviewMesh, nullptr);
 			Preview->PreviewMesh->SetTangentsMode(EDynamicMeshComponentTangentsMode::AutoCalculated);
 
 			FComponentMaterialSet MaterialSet;
diff --git a/Engine/Plugins/Experimental/MeshModelingToolsetExp/Source/MeshModelingToolsExp/Private/RevolveBoundaryTool.cpp b/Engine/Plugins/Experimental/MeshModelingToolsetExp/Source/MeshModelingToolsExp/Private/RevolveBoundaryTool.cpp
--- a/Engine/Plugins/Experimental/MeshModelingToolsetExp/Source/MeshModelingToolsExp/Private/RevolveBoundaryTool.cpp
+++ b/Engine/Plugins/Experimental/MeshModelingToolsetExp/Source/MeshModelingToolsExp/Private/RevolveBoundaryTool.cpp
@@ -229,6 +229,7 @@ void URevolveBoundaryTool::StartPreview()
 
 	Preview = NewObject<UMeshOpPreviewWithBackgroundCompute>(RevolveBoundaryOpCreator);
 	Preview->Setup(TargetWorld, RevolveBoundaryOpCreator);
+	ToolSetupUtil::ApplyRenderingConfigurationToPreview(Preview->PreviewMesh, nullptr);
 	Preview->PreviewMesh->SetTangentsMode(EDynamicMeshComponentTangentsMode::AutoCalculated);
 
 	Preview->ConfigureMaterials(MaterialProperties->Material.Get(),
diff --git a/Engine/Plugins/Experimental/MeshModelingToolsetExp/Source/MeshModelingToolsExp/Private/Sculpting/MeshSculptToolBase.cpp b/Engine/Plugins/Experimental/MeshModelingToolsetExp/Source/MeshModelingToolsExp/Private/Sculpting/MeshSculptToolBase.cpp
--- a/Engine/Plugins/Experimental/MeshModelingToolsetExp/Source/MeshModelingToolsExp/Private/Sculpting/MeshSculptToolBase.cpp
+++ b/Engine/Plugins/Experimental/MeshModelingToolsetExp/Source/MeshModelingToolsExp/Private/Sculpting/MeshSculptToolBase.cpp
@@ -235,6 +235,8 @@ void UMeshSculptToolBase::Render(IToolsContextRenderAPI* RenderAPI)
 
 void UMeshSculptToolBase::InitializeSculptMeshComponent(UBaseDynamicMeshComponent* Component)
 {
+	ToolSetupUtil::ApplyRenderingConfigurationToPreview(Component, nullptr);
+
 	// disable shadows initially, as changing shadow settings invalidates the SceneProxy
 	Component->SetShadowsEnabled(false);
 
diff --git a/Engine/Plugins/Experimental/MeshModelingToolsetExp/Source/MeshModelingToolsExp/Private/UVLayoutTool.cpp b/Engine/Plugins/Experimental/MeshModelingToolsetExp/Source/MeshModelingToolsExp/Private/UVLayoutTool.cpp
--- a/Engine/Plugins/Experimental/MeshModelingToolsetExp/Source/MeshModelingToolsExp/Private/UVLayoutTool.cpp
+++ b/Engine/Plugins/Experimental/MeshModelingToolsetExp/Source/MeshModelingToolsExp/Private/UVLayoutTool.cpp
@@ -163,6 +163,7 @@ void UUVLayoutTool::UpdateNumPreviews()
 
 			UMeshOpPreviewWithBackgroundCompute* Preview = Previews.Add_GetRef(NewObject<UMeshOpPreviewWithBackgroundCompute>(OpFactory, "Preview"));
 			Preview->Setup(this->TargetWorld, OpFactory);
+			ToolSetupUtil::ApplyRenderingConfigurationToPreview(Preview->PreviewMesh, Targets[PreviewIdx]); 
 
 			FComponentMaterialSet MaterialSet;
 			TargetMaterialInterface(PreviewIdx)->GetMaterialSet(MaterialSet);
diff --git a/Engine/Plugins/Experimental/MeshModelingToolsetExp/Source/MeshModelingToolsExp/Private/UVProjectionTool.cpp b/Engine/Plugins/Experimental/MeshModelingToolsetExp/Source/MeshModelingToolsExp/Private/UVProjectionTool.cpp
--- a/Engine/Plugins/Experimental/MeshModelingToolsetExp/Source/MeshModelingToolsExp/Private/UVProjectionTool.cpp
+++ b/Engine/Plugins/Experimental/MeshModelingToolsetExp/Source/MeshModelingToolsExp/Private/UVProjectionTool.cpp
@@ -164,6 +164,7 @@ void UUVProjectionTool::UpdateNumPreviews()
 
 	Preview = NewObject<UMeshOpPreviewWithBackgroundCompute>(OperatorFactory);
 	Preview->Setup(this->TargetWorld, OperatorFactory);
+	ToolSetupUtil::ApplyRenderingConfigurationToPreview(Preview->PreviewMesh, Target); 
 	Preview->OnMeshUpdated.AddUObject(this, &UUVProjectionTool::OnMeshUpdated);
 	Preview->PreviewMesh->SetTangentsMode(EDynamicMeshComponentTangentsMode::AutoCalculated);
 
diff --git a/Engine/Plugins/Experimental/MeshModelingToolsetExp/Source/MeshModelingToolsExp/Private/VolumeToMeshTool.cpp b/Engine/Plugins/Experimental/MeshModelingToolsetExp/Source/MeshModelingToolsExp/Private/VolumeToMeshTool.cpp
--- a/Engine/Plugins/Experimental/MeshModelingToolsetExp/Source/MeshModelingToolsExp/Private/VolumeToMeshTool.cpp
+++ b/Engine/Plugins/Experimental/MeshModelingToolsetExp/Source/MeshModelingToolsExp/Private/VolumeToMeshTool.cpp
@@ -73,6 +73,7 @@ void UVolumeToMeshTool::Setup()
 	PreviewMesh->bBuildSpatialDataStructure = false;
 	PreviewMesh->CreateInWorld(TargetVolume->GetWorld(), FTransform::Identity);
 	PreviewMesh->SetTransform(TargetVolume->GetActorTransform());
+	ToolSetupUtil::ApplyRenderingConfigurationToPreview(PreviewMesh, nullptr); 
 
 	PreviewMesh->SetMaterial( ToolSetupUtil::GetDefaultSculptMaterial(GetToolManager()) );
 
diff --git a/Engine/Plugins/Experimental/MeshModelingToolsetExp/Source/MeshModelingToolsExp/Private/WeldMeshEdgesTool.cpp b/Engine/Plugins/Experimental/MeshModelingToolsetExp/Source/MeshModelingToolsExp/Private/WeldMeshEdgesTool.cpp
--- a/Engine/Plugins/Experimental/MeshModelingToolsetExp/Source/MeshModelingToolsExp/Private/WeldMeshEdgesTool.cpp
+++ b/Engine/Plugins/Experimental/MeshModelingToolsetExp/Source/MeshModelingToolsExp/Private/WeldMeshEdgesTool.cpp
@@ -110,6 +110,7 @@ void UWeldMeshEdgesTool::Setup()
 
 	PreviewCompute = NewObject<UMeshOpPreviewWithBackgroundCompute>(OperatorFactory);
 	PreviewCompute->Setup(this->TargetWorld, OperatorFactory);
+	ToolSetupUtil::ApplyRenderingConfigurationToPreview(PreviewCompute->PreviewMesh, Target); 
 
 	// Give the preview something to display
 	PreviewCompute->PreviewMesh->SetTransform(MeshTransform);
diff --git a/Engine/Plugins/Experimental/MeshModelingToolsetExp/Source/ModelingComponents/ModelingComponents.Build.cs b/Engine/Plugins/Experimental/MeshModelingToolsetExp/Source/ModelingComponents/ModelingComponents.Build.cs
--- a/Engine/Plugins/Experimental/MeshModelingToolsetExp/Source/ModelingComponents/ModelingComponents.Build.cs
+++ b/Engine/Plugins/Experimental/MeshModelingToolsetExp/Source/ModelingComponents/ModelingComponents.Build.cs
@@ -35,7 +35,8 @@ public class ModelingComponents : ModuleRules
 				"GeometryAlgorithms",
 				"DynamicMesh",
 				"MeshConversion",
-				"ModelingOperators"
+				"ModelingOperators",
+				"DeveloperSettings"
 				// ... add other public dependencies that you statically link with here ...
 			}
             );
diff --git a/Engine/Plugins/Experimental/MeshModelingToolsetExp/Source/ModelingComponents/Private/BaseTools/BaseCreateFromSelectedTool.cpp b/Engine/Plugins/Experimental/MeshModelingToolsetExp/Source/ModelingComponents/Private/BaseTools/BaseCreateFromSelectedTool.cpp
--- a/Engine/Plugins/Experimental/MeshModelingToolsetExp/Source/ModelingComponents/Private/BaseTools/BaseCreateFromSelectedTool.cpp
+++ b/Engine/Plugins/Experimental/MeshModelingToolsetExp/Source/ModelingComponents/Private/BaseTools/BaseCreateFromSelectedTool.cpp
@@ -98,6 +98,7 @@ void UBaseCreateFromSelectedTool::Setup()
 
 	Preview = NewObject<UMeshOpPreviewWithBackgroundCompute>(this);
 	Preview->Setup(this->TargetWorld, this);
+	ToolSetupUtil::ApplyRenderingConfigurationToPreview(Preview->PreviewMesh, nullptr);
 
 	SetPreviewCallbacks();
 	Preview->OnMeshUpdated.AddLambda(
diff --git a/Engine/Plugins/Experimental/MeshModelingToolsetExp/Source/ModelingComponents/Private/BaseTools/BaseMeshProcessingTool.cpp b/Engine/Plugins/Experimental/MeshModelingToolsetExp/Source/ModelingComponents/Private/BaseTools/BaseMeshProcessingTool.cpp
--- a/Engine/Plugins/Experimental/MeshModelingToolsetExp/Source/ModelingComponents/Private/BaseTools/BaseMeshProcessingTool.cpp
+++ b/Engine/Plugins/Experimental/MeshModelingToolsetExp/Source/ModelingComponents/Private/BaseTools/BaseMeshProcessingTool.cpp
@@ -143,6 +143,7 @@ void UBaseMeshProcessingTool::Setup()
 	Preview = NewObject<UMeshOpPreviewWithBackgroundCompute>(this, "Preview");
 	Preview->Setup(this->TargetWorld, this); // Adds the actual functional tool in the Preview object
 	Preview->PreviewMesh->SetTangentsMode(EDynamicMeshComponentTangentsMode::AutoCalculated);
+	ToolSetupUtil::ApplyRenderingConfigurationToPreview(Preview->PreviewMesh, Target);
 
 	FComponentMaterialSet MaterialSet;
 	Cast<IMaterialProvider>(Target)->GetMaterialSet(MaterialSet);
diff --git a/Engine/Plugins/Experimental/MeshModelingToolsetExp/Source/ModelingComponents/Private/ModelingComponentsSettings.cpp b/Engine/Plugins/Experimental/MeshModelingToolsetExp/Source/ModelingComponents/Private/ModelingComponentsSettings.cpp
new file mode 100644
--- /dev/null
+++ b/Engine/Plugins/Experimental/MeshModelingToolsetExp/Source/ModelingComponents/Private/ModelingComponentsSettings.cpp
@@ -0,0 +1,21 @@
+// Copyright Epic Games, Inc. All Rights Reserved.
+
+#include "ModelingComponentsSettings.h"
+
+
+#define LOCTEXT_NAMESPACE "ModelingComponentsSettings"
+
+#if WITH_EDITOR
+FText UModelingComponentsSettings::GetSectionText() const 
+{ 
+	return LOCTEXT("ModelingComponentsSettingsName", "Modeling Mode Tools"); 
+}
+
+FText UModelingComponentsSettings::GetSectionDescription() const
+{
+	return LOCTEXT("ModelingComponentsSettingsDescription", "Configure Tool-level Settings for the Modeling Tools Editor Mode plugin");
+}
+#endif
+
+
+#undef LOCTEXT_NAMESPACE
\ No newline at end of file
diff --git a/Engine/Plugins/Experimental/MeshModelingToolsetExp/Source/ModelingComponents/Private/ModelingObjectsCreationAPI.cpp b/Engine/Plugins/Experimental/MeshModelingToolsetExp/Source/ModelingComponents/Private/ModelingObjectsCreationAPI.cpp
--- a/Engine/Plugins/Experimental/MeshModelingToolsetExp/Source/ModelingComponents/Private/ModelingObjectsCreationAPI.cpp
+++ b/Engine/Plugins/Experimental/MeshModelingToolsetExp/Source/ModelingComponents/Private/ModelingObjectsCreationAPI.cpp
@@ -152,3 +152,44 @@ FString UE::Modeling::StripGeneratedAssetSuffixFromName(FString InputName)
 
 	return InputName.Left(Index);
 }
+
+
+
+
+
+FString UE::Modeling::GenerateRandomShortHexString(int32 NumChars)
+{
+	int32 FailCount = 0;
+	while (FailCount++ < 10)
+	{
+		FGuid Guid = FGuid::NewGuid();
+		FString GuidString = Guid.ToString(EGuidFormats::UniqueObjectGuid).ToUpper();
+		FString Result;
+		int32 Digits = 0, Letters = 0;
+		for (int32 k = 0; k < GuidString.Len(); ++k)
+		{
+			TCHAR Character = GuidString[k];
+			if (FChar::IsHexDigit(Character))
+			{
+				Result.AppendChar(Character);
+				if (FChar::IsDigit(Character))
+				{
+					Digits++;
+				}
+				else
+				{
+					Letters++;
+				}
+			}
+			if (Result.Len() == NumChars)
+			{
+				if (Digits > 0 && Letters > 0)
+				{
+					return Result;
+				}
+				break;		// exit loop
+			}
+		}
+	}
+	return TEXT("BADGUID1");
+}
\ No newline at end of file
diff --git a/Engine/Plugins/Experimental/MeshModelingToolsetExp/Source/ModelingComponents/Private/PreviewMesh.cpp b/Engine/Plugins/Experimental/MeshModelingToolsetExp/Source/ModelingComponents/Private/PreviewMesh.cpp
--- a/Engine/Plugins/Experimental/MeshModelingToolsetExp/Source/ModelingComponents/Private/PreviewMesh.cpp
+++ b/Engine/Plugins/Experimental/MeshModelingToolsetExp/Source/ModelingComponents/Private/PreviewMesh.cpp
@@ -218,16 +218,20 @@ void UPreviewMesh::SetTransform(const FTransform& UseTransform)
 {
 	if (TemporaryParentActor != nullptr)
 	{
-		TemporaryParentActor->SetActorTransform(UseTransform);
+		if (!TemporaryParentActor->GetActorTransform().Equals(UseTransform))
+		{
+			TemporaryParentActor->SetActorTransform(UseTransform);
+			NotifyWorldPathTracedOutputInvalidated();
+		}
 	}
 }
 
-
 void UPreviewMesh::SetVisible(bool bVisible)
 {
-	if (DynamicMeshComponent != nullptr)
+	if (DynamicMeshComponent != nullptr && IsVisible() != bVisible )
 	{
 		DynamicMeshComponent->SetVisibility(bVisible, true);
+		NotifyWorldPathTracedOutputInvalidated();
 	}
 }
 
@@ -399,6 +403,8 @@ void UPreviewMesh::ReplaceMesh(FDynamicMesh3&& NewMesh)
 	{
 		UpdateRenderMeshDecomposition();
 	}
+
+	NotifyWorldPathTracedOutputInvalidated();
 }
 
 
@@ -417,6 +423,8 @@ void UPreviewMesh::EditMesh(TFunctionRef<void(FDynamicMesh3&)> EditFunc)
 	{
 		UpdateRenderMeshDecomposition();
 	}
+
+	NotifyWorldPathTracedOutputInvalidated();
 }
 
 
@@ -471,6 +479,8 @@ void UPreviewMesh::NotifyDeferredEditCompleted(ERenderUpdateMode UpdateMode, EMe
 			DynamicMeshComponent->FastNotifyVertexAttributesUpdated(bNormals, bColors, bUVs);
 		}
 	}
+
+	NotifyWorldPathTracedOutputInvalidated();
 }
 
 
@@ -493,6 +503,8 @@ TUniquePtr<FMeshChange> UPreviewMesh::TrackedEditMesh(TFunctionRef<void(FDynamic
 		UpdateRenderMeshDecomposition();
 	}
 
+	NotifyWorldPathTracedOutputInvalidated();
+
 	return MoveTemp(Change);
 }
 
@@ -519,6 +531,8 @@ void UPreviewMesh::ApplyChange(const FMeshChange* Change, bool bRevert)
 	{
 		UpdateRenderMeshDecomposition();
 	}
+
+	NotifyWorldPathTracedOutputInvalidated();
 }
 void UPreviewMesh::ApplyChange(const FMeshReplacementChange* Change, bool bRevert)
 {
@@ -532,6 +546,8 @@ void UPreviewMesh::ApplyChange(const FMeshReplacementChange* Change, bool bRever
 	{
 		UpdateRenderMeshDecomposition();
 	}
+
+	NotifyWorldPathTracedOutputInvalidated();
 }
 
 FSimpleMulticastDelegate& UPreviewMesh::GetOnMeshChanged()
@@ -589,9 +605,24 @@ void UPreviewMesh::UpdateRenderMeshDecomposition()
 void UPreviewMesh::NotifyRegionDeferredEditCompleted(const TArray<int32>& Triangles, EMeshRenderAttributeFlags ModifiedAttribs)
 {
 	DynamicMeshComponent->FastNotifyTriangleVerticesUpdated(Triangles, ModifiedAttribs);
+	NotifyWorldPathTracedOutputInvalidated();
 }
 
 void UPreviewMesh::NotifyRegionDeferredEditCompleted(const TSet<int32>& Triangles, EMeshRenderAttributeFlags ModifiedAttribs)
 {
 	DynamicMeshComponent->FastNotifyTriangleVerticesUpdated(Triangles, ModifiedAttribs);
-}
\ No newline at end of file
+	NotifyWorldPathTracedOutputInvalidated();
+}
+
+
+void UPreviewMesh::NotifyWorldPathTracedOutputInvalidated()
+{
+	if (TemporaryParentActor != nullptr)
+	{
+		UWorld* World = TemporaryParentActor->GetWorld();
+		if (World && World->Scene && FApp::CanEverRender())
+		{
+			World->Scene->InvalidatePathTracedOutput();
+		}
+	}
+}
diff --git a/Engine/Plugins/Experimental/MeshModelingToolsetExp/Source/ModelingComponents/Private/PropertySets/CreateMeshObjectTypeProperties.cpp b/Engine/Plugins/Experimental/MeshModelingToolsetExp/Source/ModelingComponents/Private/PropertySets/CreateMeshObjectTypeProperties.cpp
--- a/Engine/Plugins/Experimental/MeshModelingToolsetExp/Source/ModelingComponents/Private/PropertySets/CreateMeshObjectTypeProperties.cpp
+++ b/Engine/Plugins/Experimental/MeshModelingToolsetExp/Source/ModelingComponents/Private/PropertySets/CreateMeshObjectTypeProperties.cpp
@@ -2,6 +2,7 @@
 
 #include "PropertySets/CreateMeshObjectTypeProperties.h"
 #include "ModelingObjectsCreationAPI.h"
+#include "ModelingComponentsSettings.h"
 
 
 const FString UCreateMeshObjectTypeProperties::StaticMeshIdentifier = TEXT("Static Mesh");
@@ -106,6 +107,15 @@ bool UCreateMeshObjectTypeProperties::ConfigureCreateMeshObjectParams(FCreateMes
 	// client has to handle this case
 	ensure(OutputType != AutoIdentifier);
 
+	const UModelingComponentsSettings* Settings = GetDefault<UModelingComponentsSettings>();
+	if (Settings)
+	{
+		ParamsOut.bEnableCollision = Settings->bEnableCollision;
+		ParamsOut.CollisionMode = Settings->CollisionMode;
+
+		ParamsOut.bEnableRaytracingSupport = Settings->bEnableRayTracing;
+	}
+
 	if (OutputType == StaticMeshIdentifier)
 	{
 		ParamsOut.TypeHint = ECreateObjectTypeHint::StaticMesh;
diff --git a/Engine/Plugins/Experimental/MeshModelingToolsetExp/Source/ModelingComponents/Private/ToolSetupUtil.cpp b/Engine/Plugins/Experimental/MeshModelingToolsetExp/Source/ModelingComponents/Private/ToolSetupUtil.cpp
--- a/Engine/Plugins/Experimental/MeshModelingToolsetExp/Source/ModelingComponents/Private/ToolSetupUtil.cpp
+++ b/Engine/Plugins/Experimental/MeshModelingToolsetExp/Source/ModelingComponents/Private/ToolSetupUtil.cpp
@@ -2,13 +2,17 @@
 
 
 #include "ToolSetupUtil.h"
-
+#include "ModelingComponentsSettings.h"
 #include "Curves/CurveFloat.h"
 #include "InteractiveTool.h"
 #include "InteractiveToolManager.h"
 #include "Materials/Material.h"
 #include "Materials/MaterialInstanceDynamic.h"
 
+#include "ModelingToolTargetUtil.h"
+#include "Components/BaseDynamicMeshComponent.h"
+#include "PreviewMesh.h"
+
 
 UMaterialInterface* ToolSetupUtil::GetDefaultMaterial()
 {
@@ -318,4 +322,29 @@ UCurveFloat* ToolSetupUtil::GetContrastAdjustmentCurve(UInteractiveToolManager*
 	UCurveFloat* CurveCopy = DuplicateObject<UCurveFloat>(Curve, GetTransientPackage());
 
 	return CurveCopy;
+}
+
+
+
+
+void ToolSetupUtil::ApplyRenderingConfigurationToPreview(UBaseDynamicMeshComponent* Component, UToolTarget* SourceTarget)
+{
+	UModelingComponentsSettings* Settings = GetMutableDefault<UModelingComponentsSettings>();
+	bool bEnableRaytracingSupport = (Settings != nullptr) ? Settings->bEnableRayTracingWhileEditing : false;
+
+
+	if (bEnableRaytracingSupport)
+	{
+		Component->SetEnableRaytracing(bEnableRaytracingSupport);
+	}
+}
+
+void ToolSetupUtil::ApplyRenderingConfigurationToPreview(UPreviewMesh* PreviewMesh, UToolTarget* SourceTarget)
+{
+	UBaseDynamicMeshComponent* Component = Cast<UBaseDynamicMeshComponent>(PreviewMesh->GetRootComponent());
+	if (!Component)
+	{
+		return;
+	}
+	ApplyRenderingConfigurationToPreview(Component, SourceTarget);
 }
\ No newline at end of file
diff --git a/Engine/Plugins/Experimental/MeshModelingToolsetExp/Source/ModelingComponents/Public/ModelingComponentsSettings.h b/Engine/Plugins/Experimental/MeshModelingToolsetExp/Source/ModelingComponents/Public/ModelingComponentsSettings.h
new file mode 100644
--- /dev/null
+++ b/Engine/Plugins/Experimental/MeshModelingToolsetExp/Source/ModelingComponents/Public/ModelingComponentsSettings.h
@@ -0,0 +1,59 @@
+// Copyright Epic Games, Inc. All Rights Reserved.
+
+#pragma once
+
+#include "Engine/DeveloperSettings.h"
+#include "ModelingComponentsSettings.generated.h"
+
+
+
+/**
+ * Settings for the Modeling Components plug-in. These settings are primarily used to configure two things:
+ *   - Behavior of things like optional Rendering features inside Modeling Tools, eg for edit-preview rendering
+ *   - Setup of New Mesh Objects emitted by Modeling Tools (eg their default collision settings, etc)
+ */
+UCLASS(config=EditorPerProjectUserSettings)
+class MODELINGCOMPONENTS_API UModelingComponentsSettings : public UDeveloperSettings
+{
+	GENERATED_BODY()
+
+public:
+
+	// UDeveloperSettings overrides
+	virtual FName GetContainerName() const { return FName("Project"); }
+	virtual FName GetCategoryName() const { return FName("Plugins"); }
+	virtual FName GetSectionName() const { return FName("Modeling Mode Tools"); }
+#if WITH_EDITOR
+	virtual FText GetSectionText() const override;
+	virtual FText GetSectionDescription() const override;
+#endif
+
+public:
+
+	//
+	// Rendering settings applied to Preview Components used during live mesh editing
+	//
+
+	/** Enable Realtime Raytracing support for Mesh Editing Tools. This will impact performance of Tools with Real-Time feedback like 3D Sculpting. */
+	UPROPERTY(config, EditAnywhere, Category = "Modeling Tools|Rendering")
+	bool bEnableRayTracingWhileEditing = false;
+
+
+	//
+	// New Mesh Object settings, for new meshes/assets emitted by Modeling Tools
+	//
+
+	/** Enable Raytracing Support for new Mesh Objects created by Modeling Tools, if support is optional (eg DynamicMeshActors) */
+	UPROPERTY(config, EditAnywhere, Category = "Modeling Tools|New Mesh Objects")
+	bool bEnableRayTracing = false;
+
+	/** Enable Collision Support for new Mesh Objects created by Modeling Tools */
+	UPROPERTY(config, EditAnywhere, Category = "Modeling Tools|New Mesh Objects")
+	bool bEnableCollision = true;
+
+	/** Default Collision Mode set on new Mesh Objects created by Modeling Tools */
+	UPROPERTY(config, EditAnywhere, Category = "Modeling Tools|New Mesh Objects")
+	TEnumAsByte<enum ECollisionTraceFlag> CollisionMode = ECollisionTraceFlag::CTF_UseComplexAsSimple;
+
+
+};
\ No newline at end of file
diff --git a/Engine/Plugins/Experimental/MeshModelingToolsetExp/Source/ModelingComponents/Public/ModelingObjectsCreationAPI.h b/Engine/Plugins/Experimental/MeshModelingToolsetExp/Source/ModelingComponents/Public/ModelingObjectsCreationAPI.h
--- a/Engine/Plugins/Experimental/MeshModelingToolsetExp/Source/ModelingComponents/Public/ModelingObjectsCreationAPI.h
+++ b/Engine/Plugins/Experimental/MeshModelingToolsetExp/Source/ModelingComponents/Public/ModelingObjectsCreationAPI.h
@@ -133,6 +133,14 @@ struct MODELINGCOMPONENTS_API FCreateMeshObjectParams
 	TEnumAsByte<enum ECollisionTraceFlag> CollisionMode = ECollisionTraceFlag::CTF_UseComplexAsSimple;
 
 
+	//
+	// Rendering Configuration Options, if this is applicable for the given mesh object
+	//
+
+	/** Specify whether normals should be automatically recomputed for this new mesh object */
+	UPROPERTY(Category = "CreateMeshObjectParams", EditAnywhere)
+	bool bEnableRaytracingSupport = true;
+
 	//
 	// Mesh Build Options, if this is applicable for the given mesh object
 	// (Currently somewhat specific to Assets in the Editor)
@@ -375,6 +383,15 @@ MODELINGCOMPONENTS_API FCreateTextureObjectResult CreateTextureObject(UInteracti
  */
 MODELINGCOMPONENTS_API FString StripGeneratedAssetSuffixFromName( FString InputName );
 
+
+/**
+* Generate a N-letter GUID string that contains only hex digits,
+* and contains at least one letter and one number. Used to create _UUID suffixes
+* for making asset names unique, etc.
+*/
+MODELINGCOMPONENTS_API FString GenerateRandomShortHexString(int32 NumChars = 8);
+
+
 /**
  * Look up the name for the Asset underlying the given Component, if there is one.
  * Optionally Strip off the appended auto-generated UUID string if we can detect one.
diff --git a/Engine/Plugins/Experimental/MeshModelingToolsetExp/Source/ModelingComponents/Public/PreviewMesh.h b/Engine/Plugins/Experimental/MeshModelingToolsetExp/Source/ModelingComponents/Public/PreviewMesh.h
--- a/Engine/Plugins/Experimental/MeshModelingToolsetExp/Source/ModelingComponents/Public/PreviewMesh.h
+++ b/Engine/Plugins/Experimental/MeshModelingToolsetExp/Source/ModelingComponents/Public/PreviewMesh.h
@@ -446,6 +446,10 @@ protected:
 
 	/** Update chunk decomposition */
 	void UpdateRenderMeshDecomposition();
+
+	// This function is called internally on some changes, to let the path tracer know that this mesha/actor
+	// has been modified in a way that will require invalidating the current path tracing result
+	void NotifyWorldPathTracedOutputInvalidated();
 };
 
 
diff --git a/Engine/Plugins/Experimental/MeshModelingToolsetExp/Source/ModelingComponents/Public/ToolSetupUtil.h b/Engine/Plugins/Experimental/MeshModelingToolsetExp/Source/ModelingComponents/Public/ToolSetupUtil.h
--- a/Engine/Plugins/Experimental/MeshModelingToolsetExp/Source/ModelingComponents/Public/ToolSetupUtil.h
+++ b/Engine/Plugins/Experimental/MeshModelingToolsetExp/Source/ModelingComponents/Public/ToolSetupUtil.h
@@ -10,6 +10,10 @@ class UCurveFloat;
 class UTexture;
 class UInteractiveToolManager;
 
+class UToolTarget;
+class UBaseDynamicMeshComponent;
+class UPreviewMesh;
+
 /**
  * Utility functions for Tool implementations to use when doing configuration/setup
  */
@@ -176,4 +180,18 @@ namespace ToolSetupUtil
 	 * @return a curve asset used for contrast adjustments when using a texture map for displacements.
 	 */
 	MODELINGCOMPONENTS_API UCurveFloat* GetContrastAdjustmentCurve(UInteractiveToolManager* ToolManager);
+
+
+
+
+
+	//
+	// Rendering Configuration/Setup Functions
+	// These utility functions are used to configure rendering settings on Preview Meshes created internally by Modeling Tools.
+	// 
+	//
+
+	MODELINGCOMPONENTS_API void ApplyRenderingConfigurationToPreview(UBaseDynamicMeshComponent* Component, UToolTarget* SourceTarget = nullptr);
+	MODELINGCOMPONENTS_API void ApplyRenderingConfigurationToPreview(UPreviewMesh* PreviewMesh, UToolTarget* SourceTarget = nullptr);
+
 }
\ No newline at end of file
diff --git a/Engine/Plugins/Experimental/MeshModelingToolsetExp/Source/ModelingComponentsEditorOnly/Private/EditorModelingObjectsCreationAPI.cpp b/Engine/Plugins/Experimental/MeshModelingToolsetExp/Source/ModelingComponentsEditorOnly/Private/EditorModelingObjectsCreationAPI.cpp
--- a/Engine/Plugins/Experimental/MeshModelingToolsetExp/Source/ModelingComponentsEditorOnly/Private/EditorModelingObjectsCreationAPI.cpp
+++ b/Engine/Plugins/Experimental/MeshModelingToolsetExp/Source/ModelingComponentsEditorOnly/Private/EditorModelingObjectsCreationAPI.cpp
@@ -245,6 +245,9 @@ FCreateMeshObjectResult UEditorModelingObjectsCreationAPI::CreateDynamicMeshActo
 		NewComponent->UpdateCollision(false);
 	}
 
+	// configure raytracing
+	NewComponent->SetEnableRaytracing(CreateMeshParams.bEnableRaytracingSupport);
+
 	NewActor->PostEditChange();
 
 	// emit result
diff --git a/Engine/Plugins/Runtime/MeshModelingToolset/Source/MeshModelingTools/Private/AddPrimitiveTool.cpp b/Engine/Plugins/Runtime/MeshModelingToolset/Source/MeshModelingTools/Private/AddPrimitiveTool.cpp
--- a/Engine/Plugins/Runtime/MeshModelingToolset/Source/MeshModelingTools/Private/AddPrimitiveTool.cpp
+++ b/Engine/Plugins/Runtime/MeshModelingToolset/Source/MeshModelingTools/Private/AddPrimitiveTool.cpp
@@ -7,6 +7,7 @@
 #include "Selection/ToolSelectionUtil.h"
 #include "ModelingObjectsCreationAPI.h"
 #include "ToolSceneQueriesUtil.h"
+#include "ToolSetupUtil.h"
 
 #include "Components/StaticMeshComponent.h"
 #include "Engine/StaticMesh.h"
@@ -139,6 +140,7 @@ void UAddPrimitiveTool::Setup()
 	// create preview mesh object
 	PreviewMesh = NewObject<UPreviewMesh>(this);
 	PreviewMesh->CreateInWorld(TargetWorld, FTransform::Identity);
+	ToolSetupUtil::ApplyRenderingConfigurationToPreview(PreviewMesh, nullptr); 
 	PreviewMesh->SetVisible(false);
 	PreviewMesh->SetMaterial(MaterialProperties->Material.Get());
 	PreviewMesh->EnableWireframe(MaterialProperties->bWireframe);
diff --git a/Engine/Plugins/Runtime/MeshModelingToolset/Source/MeshModelingTools/Private/CutMeshWithMeshTool.cpp b/Engine/Plugins/Runtime/MeshModelingToolset/Source/MeshModelingTools/Private/CutMeshWithMeshTool.cpp
--- a/Engine/Plugins/Runtime/MeshModelingToolset/Source/MeshModelingTools/Private/CutMeshWithMeshTool.cpp
+++ b/Engine/Plugins/Runtime/MeshModelingToolset/Source/MeshModelingTools/Private/CutMeshWithMeshTool.cpp
@@ -53,6 +53,7 @@ void UCutMeshWithMeshTool::SetupProperties()
 	// create intersection preview mesh object
 	IntersectPreviewMesh = NewObject<UPreviewMesh>(this);
 	IntersectPreviewMesh->CreateInWorld(TargetWorld, FTransform::Identity);
+	ToolSetupUtil::ApplyRenderingConfigurationToPreview(IntersectPreviewMesh, nullptr); 
 	IntersectPreviewMesh->SetVisible(true);
 	//IntersectPreviewMesh->SetMaterial(MaterialProperties->Material.Get());
 	//IntersectPreviewMesh->SetMaterial(ToolSetupUtil::GetDefaultMaterial(nullptr));
diff --git a/Engine/Plugins/Runtime/MeshModelingToolset/Source/MeshModelingTools/Private/DrawAndRevolveTool.cpp b/Engine/Plugins/Runtime/MeshModelingToolset/Source/MeshModelingTools/Private/DrawAndRevolveTool.cpp
--- a/Engine/Plugins/Runtime/MeshModelingToolset/Source/MeshModelingTools/Private/DrawAndRevolveTool.cpp
+++ b/Engine/Plugins/Runtime/MeshModelingToolset/Source/MeshModelingTools/Private/DrawAndRevolveTool.cpp
@@ -269,6 +269,7 @@ void UDrawAndRevolveTool::StartPreview()
 	Preview = NewObject<UMeshOpPreviewWithBackgroundCompute>(RevolveOpCreator, "RevolveToolPreview");
 
 	Preview->Setup(TargetWorld, RevolveOpCreator);
+	ToolSetupUtil::ApplyRenderingConfigurationToPreview(Preview->PreviewMesh, nullptr); 
 	Preview->PreviewMesh->SetTangentsMode(EDynamicMeshComponentTangentsMode::AutoCalculated);
 
 	Preview->ConfigureMaterials(MaterialProperties->Material.Get(),
diff --git a/Engine/Plugins/Runtime/MeshModelingToolset/Source/MeshModelingTools/Private/DrawPolygonTool.cpp b/Engine/Plugins/Runtime/MeshModelingToolset/Source/MeshModelingTools/Private/DrawPolygonTool.cpp
--- a/Engine/Plugins/Runtime/MeshModelingToolset/Source/MeshModelingTools/Private/DrawPolygonTool.cpp
+++ b/Engine/Plugins/Runtime/MeshModelingToolset/Source/MeshModelingTools/Private/DrawPolygonTool.cpp
@@ -4,6 +4,7 @@
 #include "InteractiveToolManager.h"
 #include "InteractiveGizmoManager.h"
 #include "ToolBuilderUtil.h"
+#include "ToolSetupUtil.h"
 #include "BaseBehaviors/MultiClickSequenceInputBehavior.h"
 #include "BaseBehaviors/KeyAsModifierInputBehavior.h"
 
@@ -138,6 +139,7 @@ void UDrawPolygonTool::Setup()
 	// create preview mesh object
 	PreviewMesh = NewObject<UPreviewMesh>(this);
 	PreviewMesh->CreateInWorld(this->TargetWorld, FTransform::Identity);
+	ToolSetupUtil::ApplyRenderingConfigurationToPreview(PreviewMesh, nullptr); 
 	PreviewMesh->SetVisible(false);
 	{
 		UMaterialInterface* Material = nullptr;
diff --git a/Engine/Plugins/Runtime/MeshModelingToolset/Source/MeshModelingTools/Private/EditMeshPolygonsTool.cpp b/Engine/Plugins/Runtime/MeshModelingToolset/Source/MeshModelingTools/Private/EditMeshPolygonsTool.cpp
--- a/Engine/Plugins/Runtime/MeshModelingToolset/Source/MeshModelingTools/Private/EditMeshPolygonsTool.cpp
+++ b/Engine/Plugins/Runtime/MeshModelingToolset/Source/MeshModelingTools/Private/EditMeshPolygonsTool.cpp
@@ -245,6 +245,7 @@ void UEditMeshPolygonsTool::Setup()
 	// Create the preview object
 	Preview = NewObject<UMeshOpPreviewWithBackgroundCompute>();
 	Preview->Setup(TargetWorld);
+	ToolSetupUtil::ApplyRenderingConfigurationToPreview(Preview->PreviewMesh, Target); 
 	WorldTransform = UE::ToolTarget::GetLocalToWorldTransform(Target);
 	Preview->PreviewMesh->SetTransform((FTransform)WorldTransform);
 
diff --git a/Engine/Plugins/Runtime/MeshModelingToolset/Source/MeshModelingTools/Private/ToolActivities/PolyEditActivityUtil.cpp b/Engine/Plugins/Runtime/MeshModelingToolset/Source/MeshModelingTools/Private/ToolActivities/PolyEditActivityUtil.cpp
--- a/Engine/Plugins/Runtime/MeshModelingToolset/Source/MeshModelingTools/Private/ToolActivities/PolyEditActivityUtil.cpp
+++ b/Engine/Plugins/Runtime/MeshModelingToolset/Source/MeshModelingTools/Private/ToolActivities/PolyEditActivityUtil.cpp
@@ -16,6 +16,7 @@ UPolyEditPreviewMesh* PolyEditActivityUtil::CreatePolyEditPreviewMesh(UInteracti
 {
 	UPolyEditPreviewMesh* EditPreview = NewObject<UPolyEditPreviewMesh>(&Tool);
 	EditPreview->CreateInWorld(Tool.GetWorld(), FTransform::Identity);
+	ToolSetupUtil::ApplyRenderingConfigurationToPreview(EditPreview, nullptr); 
 	UpdatePolyEditPreviewMaterials(Tool, ActivityContext, *EditPreview, EPreviewMaterialType::PreviewMaterial);
 	EditPreview->EnableWireframe(true);
 
diff --git a/Engine/Source/Runtime/GeometryFramework/Public/Components/BaseDynamicMeshComponent.h b/Engine/Source/Runtime/GeometryFramework/Public/Components/BaseDynamicMeshComponent.h
--- a/Engine/Source/Runtime/GeometryFramework/Public/Components/BaseDynamicMeshComponent.h
+++ b/Engine/Source/Runtime/GeometryFramework/Public/Components/BaseDynamicMeshComponent.h
@@ -329,7 +329,7 @@ protected:
 public:
 
 	UPROPERTY(EditAnywhere, BlueprintReadWrite, Category = "Dynamic Mesh Component|Rendering")
-	bool bEnableRaytracing = false;
+	bool bEnableRaytracing = true;
 
 	/**
 	 * Enable/Disable raytracing support. This is an expensive call as it flushes
