ModelingTools: add optional UMeshSculptBrushOpProps::SetStrength(), implement in BrushOpProps where it is applicable, and use to implement IncreaseBrushSpeedAction/DecreaseBrushSpeedAction in VertexSculptTool
#rnx

diff --git a/Engine/Plugins/Experimental/MeshModelingToolsetExp/Source/MeshModelingToolsExp/Private/MeshVertexSculptTool.cpp b/Engine/Plugins/Experimental/MeshModelingToolsetExp/Source/MeshModelingToolsExp/Private/MeshVertexSculptTool.cpp
--- a/Engine/Plugins/Experimental/MeshModelingToolsetExp/Source/MeshModelingToolsExp/Private/MeshVertexSculptTool.cpp
+++ b/Engine/Plugins/Experimental/MeshModelingToolsetExp/Source/MeshModelingToolsExp/Private/MeshVertexSculptTool.cpp
@@ -936,12 +936,18 @@ bool UMeshVertexSculptTool::GetBaseMeshNearest(int32 VertexID, const FVector3d&
 
 void UMeshVertexSculptTool::IncreaseBrushSpeedAction()
 {
-	//SculptProperties->PrimaryBrushSpeed = FMath::Clamp(SculptProperties->PrimaryBrushSpeed + 0.05f, 0.0f, 1.0f);
+	TUniquePtr<FMeshSculptBrushOp>& UseBrushOp = GetActiveBrushOp();
+	float CurStrength = UseBrushOp->PropertySet->GetStrength();
+	float NewStrength = FMath::Clamp(CurStrength + 0.05f, 0.0f, 1.0f);
+	UseBrushOp->PropertySet->SetStrength(NewStrength);
 }
 
 void UMeshVertexSculptTool::DecreaseBrushSpeedAction()
 {
-	//SculptProperties->PrimaryBrushSpeed = FMath::Clamp(SculptProperties->PrimaryBrushSpeed - 0.05f, 0.0f, 1.0f);
+	TUniquePtr<FMeshSculptBrushOp>& UseBrushOp = GetActiveBrushOp();
+	float CurStrength = UseBrushOp->PropertySet->GetStrength();
+	float NewStrength = FMath::Clamp(CurStrength - 0.05f, 0.0f, 1.0f);
+	UseBrushOp->PropertySet->SetStrength(NewStrength);
 }
 
 void UMeshVertexSculptTool::UpdateBrushAlpha(UTexture2D* NewAlpha)
diff --git a/Engine/Plugins/Experimental/MeshModelingToolsetExp/Source/MeshModelingToolsExp/Public/Sculpting/MeshBrushOpBase.h b/Engine/Plugins/Experimental/MeshModelingToolsetExp/Source/MeshModelingToolsExp/Public/Sculpting/MeshBrushOpBase.h
--- a/Engine/Plugins/Experimental/MeshModelingToolsetExp/Source/MeshModelingToolsExp/Public/Sculpting/MeshBrushOpBase.h
+++ b/Engine/Plugins/Experimental/MeshModelingToolsetExp/Source/MeshModelingToolsExp/Public/Sculpting/MeshBrushOpBase.h
@@ -89,6 +89,9 @@ public:
 	virtual float GetStrength() { return 1.0f; }
 	virtual float GetDepth() { return 0.0f; }
 	virtual float GetFalloff() { return 0.5f; }
+
+	// support for this is optional, used by UI level to edit brush props via hotkeys/etc
+	virtual void SetStrength(float NewStrength) { }
 };
 
 
diff --git a/Engine/Plugins/Experimental/MeshModelingToolsetExp/Source/MeshModelingToolsExp/Public/Sculpting/MeshInflateBrushOps.h b/Engine/Plugins/Experimental/MeshModelingToolsetExp/Source/MeshModelingToolsExp/Public/Sculpting/MeshInflateBrushOps.h
--- a/Engine/Plugins/Experimental/MeshModelingToolsetExp/Source/MeshModelingToolsExp/Public/Sculpting/MeshInflateBrushOps.h
+++ b/Engine/Plugins/Experimental/MeshModelingToolsetExp/Source/MeshModelingToolsExp/Public/Sculpting/MeshInflateBrushOps.h
@@ -24,6 +24,7 @@ public:
 	float Falloff = 0.5;
 
 	virtual float GetStrength() override { return Strength; }
+	virtual void SetStrength(float NewStrength) override { Strength = FMathf::Clamp(NewStrength, 0.0f, 1.0f); }
 	virtual float GetFalloff() override { return Falloff; }
 };
 
diff --git a/Engine/Plugins/Experimental/MeshModelingToolsetExp/Source/MeshModelingToolsExp/Public/Sculpting/MeshMoveBrushOps.h b/Engine/Plugins/Experimental/MeshModelingToolsetExp/Source/MeshModelingToolsExp/Public/Sculpting/MeshMoveBrushOps.h
--- a/Engine/Plugins/Experimental/MeshModelingToolsetExp/Source/MeshModelingToolsExp/Public/Sculpting/MeshMoveBrushOps.h
+++ b/Engine/Plugins/Experimental/MeshModelingToolsetExp/Source/MeshModelingToolsExp/Public/Sculpting/MeshMoveBrushOps.h
@@ -30,6 +30,7 @@ public:
 	float Depth = 0;
 
 	virtual float GetStrength() override { return Strength; }
+	virtual void SetStrength(float NewStrength) override { Strength = FMathf::Clamp(NewStrength, 0.0f, 1.0f); }
 	virtual float GetFalloff() override { return Falloff; }
 	virtual float GetDepth() override { return Depth; }
 
diff --git a/Engine/Plugins/Experimental/MeshModelingToolsetExp/Source/MeshModelingToolsExp/Public/Sculpting/MeshPinchBrushOps.h b/Engine/Plugins/Experimental/MeshModelingToolsetExp/Source/MeshModelingToolsExp/Public/Sculpting/MeshPinchBrushOps.h
--- a/Engine/Plugins/Experimental/MeshModelingToolsetExp/Source/MeshModelingToolsExp/Public/Sculpting/MeshPinchBrushOps.h
+++ b/Engine/Plugins/Experimental/MeshModelingToolsetExp/Source/MeshModelingToolsExp/Public/Sculpting/MeshPinchBrushOps.h
@@ -31,6 +31,7 @@ public:
 	bool bPerpDamping = true;
 
 	virtual float GetStrength() override { return Strength; }
+	virtual void SetStrength(float NewStrength) override { Strength = FMathf::Clamp(NewStrength, 0.0f, 1.0f); }
 	virtual float GetFalloff() override { return Falloff; }
 	virtual float GetDepth() override { return Depth; }
 };
diff --git a/Engine/Plugins/Experimental/MeshModelingToolsetExp/Source/MeshModelingToolsExp/Public/Sculpting/MeshPlaneBrushOps.h b/Engine/Plugins/Experimental/MeshModelingToolsetExp/Source/MeshModelingToolsExp/Public/Sculpting/MeshPlaneBrushOps.h
--- a/Engine/Plugins/Experimental/MeshModelingToolsetExp/Source/MeshModelingToolsExp/Public/Sculpting/MeshPlaneBrushOps.h
+++ b/Engine/Plugins/Experimental/MeshModelingToolsetExp/Source/MeshModelingToolsExp/Public/Sculpting/MeshPlaneBrushOps.h
@@ -41,6 +41,7 @@ public:
 	EPlaneBrushSideMode WhichSide = EPlaneBrushSideMode::BothSides;
 
 	virtual float GetStrength() override { return Strength; }
+	virtual void SetStrength(float NewStrength) override { Strength = FMathf::Clamp(NewStrength, 0.0f, 1.0f); }
 	virtual float GetFalloff() override { return Falloff; }
 	virtual float GetDepth() override { return Depth; }
 	virtual EPlaneBrushSideMode GetWhichSide() { return WhichSide; }
@@ -69,6 +70,7 @@ public:
 	EPlaneBrushSideMode WhichSide = EPlaneBrushSideMode::BothSides;
 
 	virtual float GetStrength() override { return Strength; }
+	virtual void SetStrength(float NewStrength) override { Strength = FMathf::Clamp(NewStrength, 0.0f, 1.0f); }
 	virtual float GetFalloff() override { return Falloff; }
 	virtual float GetDepth() override { return Depth; }
 	virtual EPlaneBrushSideMode GetWhichSide() { return WhichSide; }
@@ -97,6 +99,7 @@ public:
 	EPlaneBrushSideMode WhichSide = EPlaneBrushSideMode::BothSides;
 
 	virtual float GetStrength() override { return Strength; }
+	virtual void SetStrength(float NewStrength) override { Strength = FMathf::Clamp(NewStrength, 0.0f, 1.0f); }
 	virtual float GetFalloff() override { return Falloff; }
 	virtual float GetDepth() override { return Depth; }
 	virtual EPlaneBrushSideMode GetWhichSide() { return WhichSide; }
diff --git a/Engine/Plugins/Experimental/MeshModelingToolsetExp/Source/MeshModelingToolsExp/Public/Sculpting/MeshSculptBrushOps.h b/Engine/Plugins/Experimental/MeshModelingToolsetExp/Source/MeshModelingToolsExp/Public/Sculpting/MeshSculptBrushOps.h
--- a/Engine/Plugins/Experimental/MeshModelingToolsetExp/Source/MeshModelingToolsExp/Public/Sculpting/MeshSculptBrushOps.h
+++ b/Engine/Plugins/Experimental/MeshModelingToolsetExp/Source/MeshModelingToolsExp/Public/Sculpting/MeshSculptBrushOps.h
@@ -24,6 +24,7 @@ public:
 	float Falloff = 1.0;
 
 	virtual float GetStrength() override { return Strength; }
+	virtual void SetStrength(float NewStrength) override { Strength = FMathf::Clamp(NewStrength, 0.0f, 1.0f); }
 	virtual float GetFalloff() override { return Falloff; }
 };
 
@@ -100,6 +101,7 @@ public:
 	float Falloff = 1.0;
 
 	virtual float GetStrength() override { return Strength; }
+	virtual void SetStrength(float NewStrength) override { Strength = FMathf::Clamp(NewStrength, 0.0f, 1.0f); }
 	virtual float GetFalloff() override { return Falloff; }
 };
 
@@ -197,6 +199,7 @@ public:
 
 
 	virtual float GetStrength() override { return Strength; }
+	virtual void SetStrength(float NewStrength) override { Strength = FMathf::Clamp(NewStrength, 0.0f, 1.0f); }
 	virtual float GetFalloff() override { return Falloff; }
 };
 
diff --git a/Engine/Plugins/Experimental/MeshModelingToolsetExp/Source/MeshModelingToolsExp/Public/Sculpting/MeshSmoothingBrushOps.h b/Engine/Plugins/Experimental/MeshModelingToolsetExp/Source/MeshModelingToolsExp/Public/Sculpting/MeshSmoothingBrushOps.h
--- a/Engine/Plugins/Experimental/MeshModelingToolsetExp/Source/MeshModelingToolsExp/Public/Sculpting/MeshSmoothingBrushOps.h
+++ b/Engine/Plugins/Experimental/MeshModelingToolsetExp/Source/MeshModelingToolsExp/Public/Sculpting/MeshSmoothingBrushOps.h
@@ -37,6 +37,7 @@ public:
 	bool bPreserveUVFlow = false;
 
 	virtual float GetStrength() override { return Strength; }
+	virtual void SetStrength(float NewStrength) override { Strength = FMathf::Clamp(NewStrength, 0.0f, 1.0f); }
 	virtual float GetFalloff() override { return Falloff; }
 	virtual bool GetPreserveUVFlow() override { return bPreserveUVFlow; }
 };
@@ -61,6 +62,7 @@ public:
 	bool bPreserveUVFlow = false;
 
 	virtual float GetStrength() override { return Strength; }
+	virtual void SetStrength(float NewStrength) override { Strength = FMathf::Clamp(NewStrength, 0.0f, 1.0f); }
 	virtual float GetFalloff() override { return Falloff; }
 	virtual bool GetPreserveUVFlow() override { return bPreserveUVFlow; }
 };
@@ -124,6 +126,7 @@ public:
 	bool bPreserveUVFlow = false;
 
 	virtual float GetStrength() override { return Strength; }
+	virtual void SetStrength(float NewStrength) override { Strength = FMathf::Clamp(NewStrength, 0.0f, 1.0f); }
 	virtual float GetFalloff() override { return Falloff; }
 	virtual bool GetPreserveUVFlow() override { return bPreserveUVFlow; }
 };
@@ -200,6 +203,7 @@ public:
 	EPlaneBrushSideMode WhichSide = EPlaneBrushSideMode::BothSides;
 
 	virtual float GetStrength() override { return Strength; }
+	virtual void SetStrength(float NewStrength) override { Strength = FMathf::Clamp(NewStrength, 0.0f, 1.0f); }
 	virtual float GetFalloff() override { return Falloff; }
 	virtual float GetDepth() override { return Depth; }
 };
@@ -271,6 +275,7 @@ public:
 	float Falloff = 0.5;
 
 	virtual float GetStrength() override { return Strength; }
+	virtual void SetStrength(float NewStrength) override { Strength = FMathf::Clamp(NewStrength, 0.0f, 1.0f); }
 	virtual float GetFalloff() override { return Falloff; }
 };
 
