ModelingTools: set restored wireframe state on DynamicMeshSculptTool startup
#rnx

diff --git a/Engine/Plugins/Experimental/MeshModelingToolsetExp/Source/MeshModelingToolsExp/Private/DynamicMeshSculptTool.cpp b/Engine/Plugins/Experimental/MeshModelingToolsetExp/Source/MeshModelingToolsExp/Private/DynamicMeshSculptTool.cpp
--- a/Engine/Plugins/Experimental/MeshModelingToolsetExp/Source/MeshModelingToolsExp/Private/DynamicMeshSculptTool.cpp
+++ b/Engine/Plugins/Experimental/MeshModelingToolsetExp/Source/MeshModelingToolsExp/Private/DynamicMeshSculptTool.cpp
@@ -259,6 +259,7 @@ void UDynamicMeshSculptTool::Setup()
 		[this]() { return GizmoProperties->Rotation; },
 		[this](FQuat NewRotation) { UpdateGizmoFromProperties(); }, GizmoProperties->Rotation);
 
+	DynamicMeshComponent->bExplicitShowWireframe = ViewProperties->bShowWireframe;
 
 	// create proxy for plane gizmo, but not gizmo itself, as it only appears in FixedPlane brush mode
 	// listen for changes to the proxy and update the plane when that happens
