ModelingComponents: add missing header for nonunity build
#rnx

diff --git a/Engine/Plugins/Experimental/MeshModelingToolsetExp/Source/ModelingComponents/Public/ModelingComponentsSettings.h b/Engine/Plugins/Experimental/MeshModelingToolsetExp/Source/ModelingComponents/Public/ModelingComponentsSettings.h
--- a/Engine/Plugins/Experimental/MeshModelingToolsetExp/Source/ModelingComponents/Public/ModelingComponentsSettings.h
+++ b/Engine/Plugins/Experimental/MeshModelingToolsetExp/Source/ModelingComponents/Public/ModelingComponentsSettings.h
@@ -3,6 +3,7 @@
 #pragma once
 
 #include "Engine/DeveloperSettings.h"
+#include "BodySetupEnums.h"
 #include "ModelingComponentsSettings.generated.h"
 
 
