ModelingMode: fix incorrect color on startup material in VertexSculptTool
#rnx

diff --git a/Engine/Plugins/Experimental/MeshModelingToolsetExp/Source/MeshModelingToolsExp/Private/Sculpting/MeshSculptToolBase.cpp b/Engine/Plugins/Experimental/MeshModelingToolsetExp/Source/MeshModelingToolsExp/Private/Sculpting/MeshSculptToolBase.cpp
--- a/Engine/Plugins/Experimental/MeshModelingToolsetExp/Source/MeshModelingToolsExp/Private/Sculpting/MeshSculptToolBase.cpp
+++ b/Engine/Plugins/Experimental/MeshModelingToolsetExp/Source/MeshModelingToolsExp/Private/Sculpting/MeshSculptToolBase.cpp
@@ -86,7 +86,7 @@ void UMeshSculptToolBase::Setup()
 		[this](FLinearColor NewColor) { UpdateColorSetting(NewColor); }); 
 	// This can actually use the same function since the parameter names for the material are the same
 	ViewProperties->WatchProperty(ViewProperties->TransparentMaterialColor,
-		[this](FLinearColor NewColor) { UpdateColorSetting(NewColor); });
+		[this](FLinearColor NewColor) { UpdateTransparentColorSetting(NewColor); });
 	ViewProperties->WatchProperty(ViewProperties->Opacity,
 		[this](double NewValue) { UpdateOpacitySetting(NewValue); });
 	ViewProperties->WatchProperty(ViewProperties->bTwoSided,
@@ -904,12 +904,29 @@ void UMeshSculptToolBase::UpdateFlatShadingSetting(bool bNewValue)
 
 void UMeshSculptToolBase::UpdateColorSetting(FLinearColor NewColor)
 {
-	if (ActiveOverrideMaterial != nullptr)
+	if (ViewProperties->MaterialMode != EMeshEditingMaterialModes::Transparent)
 	{
-		ActiveOverrideMaterial->SetVectorParameterValue(TEXT("Color"), NewColor);
+		if (ActiveOverrideMaterial != nullptr)
+		{
+			ActiveOverrideMaterial->SetVectorParameterValue(TEXT("Color"), NewColor);
+		}
 	}
 }
 
+void UMeshSculptToolBase::UpdateTransparentColorSetting(FLinearColor NewColor)
+{
+	// only want to update the active material if it is the transparent one...
+	if (ViewProperties->MaterialMode == EMeshEditingMaterialModes::Transparent)
+	{
+		if (ActiveOverrideMaterial != nullptr)
+		{
+			ActiveOverrideMaterial->SetVectorParameterValue(TEXT("Color"), NewColor);
+		}
+	}
+}
+
+
+
 void UMeshSculptToolBase::UpdateImageSetting(UTexture2D* NewImage)
 {
 	if (ActiveOverrideMaterial != nullptr)
diff --git a/Engine/Plugins/Experimental/MeshModelingToolsetExp/Source/MeshModelingToolsExp/Public/Properties/MeshMaterialProperties.h b/Engine/Plugins/Experimental/MeshModelingToolsetExp/Source/MeshModelingToolsExp/Public/Properties/MeshMaterialProperties.h
--- a/Engine/Plugins/Experimental/MeshModelingToolsetExp/Source/MeshModelingToolsExp/Public/Properties/MeshMaterialProperties.h
+++ b/Engine/Plugins/Experimental/MeshModelingToolsetExp/Source/MeshModelingToolsExp/Public/Properties/MeshMaterialProperties.h
@@ -132,13 +132,14 @@ public:
 	UPROPERTY(EditAnywhere, Category = Rendering, meta = (EditConditionHides, EditCondition = "MaterialMode == EMeshEditingMaterialModes::CustomImage", TransientToolProperty) )
 	TObjectPtr<UTexture2D> Image;
 
-	//~ Could have used the same property as Color, above, but the user may want different saved values for the two
-	UPROPERTY(EditAnywhere, Category = Rendering, AdvancedDisplay, meta = (EditConditionHides, EditCondition = "MaterialMode == EMeshEditingMaterialModes::Transparent", DisplayName = "Color"))
-	FLinearColor TransparentMaterialColor = FLinearColor(0.0606, 0.309, 0.842);
-
+	/** Opacity of transparent material */
 	UPROPERTY(EditAnywhere, Category = Rendering, meta = (EditConditionHides, EditCondition = "MaterialMode == EMeshEditingMaterialModes::Transparent", ClampMin = "0", ClampMax = "1.0"))
 	double Opacity = 0.65;
 
+	//~ Could have used the same property as Color, above, but the user may want different saved values for the two
+	UPROPERTY(EditAnywhere, Category = Rendering, meta = (EditConditionHides, EditCondition = "MaterialMode == EMeshEditingMaterialModes::Transparent", DisplayName = "Color"))
+	FLinearColor TransparentMaterialColor = FLinearColor(0.0606, 0.309, 0.842);
+
 	/** Although a two-sided transparent material causes rendering issues with overlapping faces, it is still frequently useful to see the shape when sculpting around other objects. */
 	UPROPERTY(EditAnywhere, Category = Rendering, meta = (EditConditionHides, EditCondition = "MaterialMode == EMeshEditingMaterialModes::Transparent"))
 	bool bTwoSided = true;
diff --git a/Engine/Plugins/Experimental/MeshModelingToolsetExp/Source/MeshModelingToolsExp/Public/Sculpting/MeshSculptToolBase.h b/Engine/Plugins/Experimental/MeshModelingToolsetExp/Source/MeshModelingToolsExp/Public/Sculpting/MeshSculptToolBase.h
--- a/Engine/Plugins/Experimental/MeshModelingToolsetExp/Source/MeshModelingToolsExp/Public/Sculpting/MeshSculptToolBase.h
+++ b/Engine/Plugins/Experimental/MeshModelingToolsetExp/Source/MeshModelingToolsExp/Public/Sculpting/MeshSculptToolBase.h
@@ -424,6 +424,7 @@ protected:
 	virtual void UpdateMaterialMode(EMeshEditingMaterialModes NewMode);
 	virtual void UpdateFlatShadingSetting(bool bNewValue);
 	virtual void UpdateColorSetting(FLinearColor NewColor);
+	virtual void UpdateTransparentColorSetting(FLinearColor NewColor);
 	virtual void UpdateOpacitySetting(double Opacity);
 	virtual void UpdateTwoSidedSetting(bool bOn);
 	virtual void UpdateCustomMaterial(TWeakObjectPtr<UMaterialInterface> NewMaterial);
