make pivot tool preserve scales when possible

diff --git a/Engine/Plugins/Experimental/MeshModelingToolsetExp/Source/MeshModelingToolsExp/Private/EditPivotTool.cpp b/Engine/Plugins/Experimental/MeshModelingToolsetExp/Source/MeshModelingToolsExp/Private/EditPivotTool.cpp
--- a/Engine/Plugins/Experimental/MeshModelingToolsetExp/Source/MeshModelingToolsExp/Private/EditPivotTool.cpp
+++ b/Engine/Plugins/Experimental/MeshModelingToolsetExp/Source/MeshModelingToolsExp/Private/EditPivotTool.cpp
@@ -469,20 +469,75 @@ void UEditPivotTool::UpdateAssets(const FFrame3d& NewPivotWorldFrame)
 		else if (MapToFirstOccurrences[ComponentIdx] == ComponentIdx)
 		{
 			UE::Geometry::FTransform3d ToBake(OriginalTransforms[ComponentIdx] * NewWorldInverse);
+			// to preserve scale, after baking the first transform, we will also need to bake an inverse scale transform
+			// this bake will need to be applied separately only in cases where FTransform3d cannot correctly represent the combination of the two transforms
+			// TODO: we could skip the extra bake step if the bake functions could take a general matrix
+			UE::Geometry::FTransform3d SeparateBakeScale = UE::Geometry::FTransform3d::Identity();
+			bool bNeedSeparateScale = false;
+			// ScaledNewWorldTransform is the pivot widget's transform with the scale of the original component transform
+			FTransform ScaledNewWorldTransform = NewWorldTransform;
+			// Note: this section only needed if we want to preserve the original mesh scale
+			// Basically our goal is, given the original actor transform A = Ta Ra Sa, and gizmo transform G = Tg Rg:
+			//   Try to keep the mesh in the same place w/ new actor transform, Tg Rg Sa (the gizmo transform w/ original actor scale)
+			//   To keep the mesh unmoved, we then have:
+			//    (New actor transform: Tg Rg Sa) * (Baked transform: Sa^-1 Rg^-1 Tg^-1 Ta Ra Sa) = (Original actor transform: Ta Ra Sa)
+			//   This cannot in general be represented as a single transform, so it requires us to bake the Sa^-1 separately ...
+			//   but in special cases where Rg^-1 Ra Sa == Sa Rg^-1 Ra, the baked scales cancel and we can instead bake a single transform with no scale
+			{
+				FVector3d Scale = (FVector3d)OriginalTransforms[ComponentIdx].GetScale3D();
+				FQuaterniond Rotation = ToBake.GetRotation();
+				bool AxisZero[3]{ FMath::IsNearlyZero(Rotation.X, DOUBLE_KINDA_SMALL_NUMBER),
+								  FMath::IsNearlyZero(Rotation.Y, DOUBLE_KINDA_SMALL_NUMBER), 
+								  FMath::IsNearlyZero(Rotation.Z, DOUBLE_KINDA_SMALL_NUMBER) };
+				int Zeros = (int)AxisZero[0] + (int)AxisZero[1] + (int)AxisZero[2];
+				bool bEqXY = FMath::IsNearlyEqual(Scale.X, Scale.Y);
+				bool bEqYZ = FMath::IsNearlyEqual(Scale.Y, Scale.Z);
+				bool bCanCombinedScales = 
+					Zeros == 3 || // no rotation (quaternion x,y,z == axis scaled by sin(angle/2) == 0 when angle is 0)
+					(bEqXY && bEqYZ) || // uniform scale
+					(Zeros == 2 && (// rotation is around a major axis && 
+						(!AxisZero[0] && bEqYZ) || // (scales on dimensions-moved-by-rotation are equal)
+						(!AxisZero[1] && FMath::IsNearlyEqual(Scale.X, Scale.Z)) ||
+						(!AxisZero[2] && bEqXY)
+					));
+				bNeedSeparateScale = !bCanCombinedScales;
+				FVector3d InvScale = UE::Geometry::FTransform3d::GetSafeScaleReciprocal(OriginalTransforms[ComponentIdx].GetScale3D());
+				if (!bNeedSeparateScale)
+				{
+					ToBake.SetScale(FVector3d::One()); // clear scale; it will be fully captured by the new transform
+					ToBake.SetTranslation(ToBake.GetTranslation() * InvScale);
+					ScaledNewWorldTransform.SetScale3D(OriginalTransforms[ComponentIdx].GetScale3D());
+				}
+				else if (InvScale.X != 0 && InvScale.Y != 0 && InvScale.Z != 0) // Scale was inverted correctly
+				{
+					// non-uniform scale is incompatible with new pivot orientation; need to bake additional counter-scale into the mesh
+					SeparateBakeScale.SetScale(InvScale);
+					ScaledNewWorldTransform.SetScale3D(OriginalTransforms[ComponentIdx].GetScale3D());
+				}
+				// else do nothing -- scale is not invertible and must be baked
+			}
 
 			// transform simple collision geometry
 			if (UE::Geometry::ComponentTypeSupportsCollision(Component))
 			{
 				UE::Geometry::TransformSimpleCollision(Component, ToBake);
+				if (bNeedSeparateScale)
+				{
+					UE::Geometry::TransformSimpleCollision(Component, SeparateBakeScale);
+				}
 			}
 
 			FMeshDescription SourceMesh(UE::ToolTarget::GetMeshDescriptionCopy(Targets[ComponentIdx], false));
 			FMeshDescriptionEditableTriangleMeshAdapter EditableMeshDescAdapter(&SourceMesh);
 			MeshAdapterTransforms::ApplyTransform(EditableMeshDescAdapter, ToBake);
+			if (bNeedSeparateScale)
+			{
+				MeshAdapterTransforms::ApplyTransform(EditableMeshDescAdapter, SeparateBakeScale);
+			}
 			// todo: support vertex-only update
 			UE::ToolTarget::CommitMeshDescriptionUpdate(Targets[ComponentIdx], &SourceMesh);
 
-			Component->SetWorldTransform(NewWorldTransform);
+			Component->SetWorldTransform(ScaledNewWorldTransform);
 		}
 		else
 		{
diff --git a/Engine/Plugins/Experimental/MeshModelingToolsetExp/Source/MeshModelingToolsExp/Public/EditPivotTool.h b/Engine/Plugins/Experimental/MeshModelingToolsetExp/Source/MeshModelingToolsExp/Public/EditPivotTool.h
--- a/Engine/Plugins/Experimental/MeshModelingToolsetExp/Source/MeshModelingToolsExp/Public/EditPivotTool.h
+++ b/Engine/Plugins/Experimental/MeshModelingToolsetExp/Source/MeshModelingToolsExp/Public/EditPivotTool.h
@@ -55,7 +55,7 @@ enum class EEditPivotSnapDragRotationMode : uint8
 
 
 /**
- * Standard properties of the Transform Meshes operation
+ * Standard properties of the Edit Pivot operation
  */
 UCLASS()
 class MESHMODELINGTOOLSEXP_API UEditPivotToolProperties : public UInteractiveToolPropertySet
