ModelingTools: draw borders of planar polygon in DrawPolyPathTool, so that use can see what is happening when the path is obscured by scene geometry.
#rnx

diff --git a/Engine/Plugins/Experimental/MeshModelingToolsetExp/Source/MeshModelingToolsExp/Private/DrawPolyPathTool.cpp b/Engine/Plugins/Experimental/MeshModelingToolsetExp/Source/MeshModelingToolsExp/Private/DrawPolyPathTool.cpp
--- a/Engine/Plugins/Experimental/MeshModelingToolsetExp/Source/MeshModelingToolsExp/Private/DrawPolyPathTool.cpp
+++ b/Engine/Plugins/Experimental/MeshModelingToolsetExp/Source/MeshModelingToolsExp/Private/DrawPolyPathTool.cpp
@@ -19,6 +19,8 @@
 #include "Selection/ToolSelectionUtil.h"
 #include "Operations/ExtrudeMesh.h"
 #include "DynamicMesh/MeshNormals.h"
+#include "MeshBoundaryLoops.h"
+#include "ToolDataVisualizer.h"
 
 #include "ExplicitUseGeometryMathTypes.h"		// using UE::Geometry::(math types)
 using namespace UE::Geometry;
@@ -407,6 +409,33 @@ void UDrawPolyPathTool::Render(IToolsContextRenderAPI* RenderAPI)
 	{
 		SurfacePathMechanic->Render(RenderAPI);
 	}
+
+	if (CurPolyLoop.Num() > 0)
+	{
+		FToolDataVisualizer LineRenderer;
+		LineRenderer.LineColor = LinearColors::DarkOrange3f();
+		LineRenderer.LineThickness = 4.0f;
+		LineRenderer.bDepthTested = false;
+
+		LineRenderer.BeginFrame(RenderAPI);
+
+		int32 NumPoints = CurPolyLoop.Num();
+		for (int32 k = 0; k < NumPoints; ++k)
+		{
+			LineRenderer.DrawLine( CurPolyLoop[k], CurPolyLoop[ (k+1) % NumPoints ] );
+		}
+		if (SecondPolyLoop.Num() > 0)
+		{
+			NumPoints = SecondPolyLoop.Num();
+			for (int32 k = 0; k < NumPoints; ++k)
+			{
+				LineRenderer.DrawLine( SecondPolyLoop[k], SecondPolyLoop[ (k+1) % NumPoints ] );
+			}
+		}
+
+		LineRenderer.EndFrame();
+	}
+
 }
 
 
@@ -578,6 +607,9 @@ void UDrawPolyPathTool::UpdatePathPreview()
 
 void UDrawPolyPathTool::GeneratePathMesh(FDynamicMesh3& Mesh)
 {
+	CurPolyLoop.Reset();
+	SecondPolyLoop.Reset();
+
 	Mesh.Clear();
 	int32 NumPoints = CurPathPoints.Num();
 	if (NumPoints > 1)
@@ -631,6 +663,16 @@ void UDrawPolyPathTool::GeneratePathMesh(FDynamicMesh3& Mesh)
 		}
 
 		FMeshNormals::QuickRecomputeOverlayNormals(Mesh);
+
+		FMeshBoundaryLoops Loops(&Mesh, true);
+		if (Loops.Loops.Num() > 0)
+		{
+			Loops.Loops[0].GetVertices<FVector3d>(CurPolyLoop);
+			if (Loops.Loops.Num() > 1)
+			{
+				Loops.Loops[1].GetVertices<FVector3d>(SecondPolyLoop);
+			}
+		}
 	}
 }
 
@@ -707,6 +749,9 @@ void UDrawPolyPathTool::ClearPreview()
 		EditPreview->Disconnect();
 		EditPreview = nullptr;
 	}
+
+	CurPolyLoop.Reset();
+	SecondPolyLoop.Reset();
 }
 
 
@@ -789,6 +834,9 @@ void UDrawPolyPathTool::EmitNewObject(EDrawPolyPathOutputMode OutputMode)
 	}
 
 	GetToolManager()->EndUndoTransaction();
+
+	CurPolyLoop.Reset();
+	SecondPolyLoop.Reset();
 }
 
 
diff --git a/Engine/Plugins/Experimental/MeshModelingToolsetExp/Source/MeshModelingToolsExp/Public/DrawPolyPathTool.h b/Engine/Plugins/Experimental/MeshModelingToolsetExp/Source/MeshModelingToolsExp/Public/DrawPolyPathTool.h
--- a/Engine/Plugins/Experimental/MeshModelingToolsetExp/Source/MeshModelingToolsExp/Public/DrawPolyPathTool.h
+++ b/Engine/Plugins/Experimental/MeshModelingToolsetExp/Source/MeshModelingToolsExp/Public/DrawPolyPathTool.h
@@ -196,6 +196,9 @@ protected:
 	double CurHeight;
 	bool bPathIsClosed = false;		// If true, CurPathPoints are assumed to define a closed path
 
+	TArray<FVector3d> CurPolyLoop;
+	TArray<FVector3d> SecondPolyLoop;
+
 	UPROPERTY()
 	TObjectPtr<UPolyEditPreviewMesh> EditPreview;
 
