ModelingTools: fix various UV issues
- FDynamicMeshUVEditor::SetTriangleUVsFromFreeBoundaryConforma no longer worked w/ bUseExistingUVTopology flag, as it was clearing UVs on triangles before trying to update them
- RecomputeUVsOp was not preservinge existing UV topology when using existing UV islands w/ Conformal unwrap (not supported for ExpMap)
- UE::ToolTarget::CommitDynamicMeshUVUpdate was not storing a mesh change, resulting in empty transactions (eg undo/redo broken)
- UE::MeshUVTransforms::MakeSeamsDisjoint was failing if elements were unreferenced, and jitter introduced was too large
#rnx

diff --git a/Engine/Plugins/Experimental/MeshModelingToolsetExp/Source/ModelingComponents/Private/ModelingToolTargetUtil.cpp b/Engine/Plugins/Experimental/MeshModelingToolsetExp/Source/ModelingComponents/Private/ModelingToolTargetUtil.cpp
--- a/Engine/Plugins/Experimental/MeshModelingToolsetExp/Source/ModelingComponents/Private/ModelingToolTargetUtil.cpp
+++ b/Engine/Plugins/Experimental/MeshModelingToolsetExp/Source/ModelingComponents/Private/ModelingToolTargetUtil.cpp
@@ -332,10 +332,9 @@ UE::ToolTarget::EDynamicMeshUpdateResult UE::ToolTarget::CommitDynamicMeshUVUpda
 	IPersistentDynamicMeshSource* DynamicMeshSource = Cast<IPersistentDynamicMeshSource>(Target);
 	if (DynamicMeshSource)
 	{
+		// just do a full mesh update for now
 		// todo actually only update UVs? 
-		UDynamicMesh* DynamicMesh = DynamicMeshSource->GetDynamicMeshContainer();
-		DynamicMesh->EditMesh([&](FDynamicMesh3& EditMesh) { EditMesh = *UpdatedMesh; });
-		return EDynamicMeshUpdateResult::Ok;
+		return CommitDynamicMeshUpdate(Target, *UpdatedMesh, true);
 	}
 
 	IDynamicMeshCommitter* DynamicMeshCommitter = Cast<IDynamicMeshCommitter>(Target);
diff --git a/Engine/Plugins/Experimental/MeshModelingToolsetExp/Source/ModelingOperators/Private/ParameterizationOps/RecomputeUVsOp.cpp b/Engine/Plugins/Experimental/MeshModelingToolsetExp/Source/ModelingOperators/Private/ParameterizationOps/RecomputeUVsOp.cpp
--- a/Engine/Plugins/Experimental/MeshModelingToolsetExp/Source/ModelingOperators/Private/ParameterizationOps/RecomputeUVsOp.cpp
+++ b/Engine/Plugins/Experimental/MeshModelingToolsetExp/Source/ModelingOperators/Private/ParameterizationOps/RecomputeUVsOp.cpp
@@ -118,6 +118,7 @@ bool FRecomputeUVsOp::CalculateResult_Basic(FProgressCancel* Progress)
 	}
 
 	// find group-connected-components
+	bool bUseExisingUVTopology = false;
 	FMeshConnectedComponents ConnectedComponents(ResultMesh.Get());
 	if (IslandMode == ERecomputeUVsIslandMode::PolyGroups)
 	{
@@ -139,6 +140,7 @@ bool FRecomputeUVsOp::CalculateResult_Basic(FProgressCancel* Progress)
 		ConnectedComponents.FindConnectedTriangles([&](int32 Triangle0, int32 Triangle1) {
 			return UseOverlay->AreTrianglesConnected(Triangle0, Triangle1);
 		});
+		bUseExisingUVTopology = true;
 	}
 
 	if (Progress && Progress->Cancelled())
@@ -169,7 +171,7 @@ bool FRecomputeUVsOp::CalculateResult_Basic(FProgressCancel* Progress)
 		break;
 
 		case ERecomputeUVsUnwrapType::ConformalFreeBoundary:
-			bComponentSolved[k] = UVEditor.SetTriangleUVsFromFreeBoundaryConformal(ComponentTris);
+			bComponentSolved[k] = UVEditor.SetTriangleUVsFromFreeBoundaryConformal(ComponentTris, bUseExisingUVTopology);
 			if (bComponentSolved[k])
 			{
 				UVEditor.ScaleUVAreaTo3DArea(ComponentTris, true);
diff --git a/Engine/Plugins/Runtime/GeometryProcessing/Source/DynamicMesh/Private/Parameterization/DynamicMeshUVEditor.cpp b/Engine/Plugins/Runtime/GeometryProcessing/Source/DynamicMesh/Private/Parameterization/DynamicMeshUVEditor.cpp
--- a/Engine/Plugins/Runtime/GeometryProcessing/Source/DynamicMesh/Private/Parameterization/DynamicMeshUVEditor.cpp
+++ b/Engine/Plugins/Runtime/GeometryProcessing/Source/DynamicMesh/Private/Parameterization/DynamicMeshUVEditor.cpp
@@ -527,8 +527,6 @@ bool FDynamicMeshUVEditor::SetTriangleUVsFromFreeBoundaryConformal(const TArray<
 	if (ensure(UVOverlay) == false) return false;
 	if (Triangles.Num() == 0) return false;
 
-	ResetUVs(Triangles);
-
 	FDynamicMesh3 Submesh(EMeshComponents::None);
 	TMap<int32, int32> BaseToSubmeshV;
 	TArray<int32> SubmeshToBaseV;
@@ -622,6 +620,9 @@ bool FDynamicMeshUVEditor::SetTriangleUVsFromFreeBoundaryConformal(const TArray<
 	}
 	else
 	{
+		// clear existing triangles
+		ResetUVs(Triangles);
+
 		// copy back to target UVOverlay
 		TArray<int32> VtxElementIDs;
 		TArray<int32> NewElementIDs;
diff --git a/Engine/Plugins/Runtime/GeometryProcessing/Source/DynamicMesh/Private/Parameterization/MeshUVTransforms.cpp b/Engine/Plugins/Runtime/GeometryProcessing/Source/DynamicMesh/Private/Parameterization/MeshUVTransforms.cpp
--- a/Engine/Plugins/Runtime/GeometryProcessing/Source/DynamicMesh/Private/Parameterization/MeshUVTransforms.cpp
+++ b/Engine/Plugins/Runtime/GeometryProcessing/Source/DynamicMesh/Private/Parameterization/MeshUVTransforms.cpp
@@ -142,7 +142,11 @@ void UE::MeshUVTransforms::MakeSeamsDisjoint(FDynamicMeshUVOverlay* UVOverlay)
 	for (int32 elemid : UVOverlay->ElementIndicesItr())
 	{
 		int32 ParentVtxID = UVOverlay->GetParentVertex(elemid);
-		checkSlow(ParentVtxID != FDynamicMesh3::InvalidID);
+		if (ParentVtxID == FDynamicMesh3::InvalidID)
+		{
+			checkSlow(false);
+			continue;
+		}
 		int32 SeamListIndex = SeamVerticesMap[ParentVtxID];
 		if (SeamListIndex == -1)
 		{
@@ -153,11 +157,11 @@ void UE::MeshUVTransforms::MakeSeamsDisjoint(FDynamicMeshUVOverlay* UVOverlay)
 		// the position becomes unique
 		FVector2f UVPos = UVOverlay->GetElement(elemid);
 		int32 TryCounter = 0;
-		float JitterScale = 0.005;
+		float JitterScale = 0.0001f;
 		while (SeamUVPositions[SeamListIndex].Contains(UVPos) && TryCounter++ < 25)
 		{
 			UVPos = UVOverlay->GetElement(elemid) + GetRandomUVJitter(JitterScale);
-			JitterScale *= 2.0;
+			JitterScale *= 2.0f;
 		}
 
 		SeamUVPositions[SeamListIndex].Add(UVPos);
