Implement brute force Triangle, Island and Mesh marquee selection in UV Editor

diff --git a/Engine/Plugins/Experimental/MeshModelingToolsetExp/Source/ModelingComponents/Private/Mechanics/RectangleMarqueeMechanic.cpp b/Engine/Plugins/Experimental/MeshModelingToolsetExp/Source/ModelingComponents/Private/Mechanics/RectangleMarqueeMechanic.cpp
--- a/Engine/Plugins/Experimental/MeshModelingToolsetExp/Source/ModelingComponents/Private/Mechanics/RectangleMarqueeMechanic.cpp
+++ b/Engine/Plugins/Experimental/MeshModelingToolsetExp/Source/ModelingComponents/Private/Mechanics/RectangleMarqueeMechanic.cpp
@@ -6,8 +6,8 @@
 #include "CanvasItem.h"
 #include "ToolSceneQueriesUtil.h"
 
+#include "Intersection/IntersectionQueries2.h"
 #include "ExplicitUseGeometryMathTypes.h"		// using UE::Geometry::(math types)
-#include "Intersection/IntrSegment2Segment2.h"
 using namespace UE::Geometry;
 
 #define LOCTEXT_NAMESPACE "URectangleMarqueeMechanic"
@@ -15,27 +15,27 @@ using namespace UE::Geometry;
 void FCameraRectangle::Initialize()
 {
 	bIsInitialized = true;
-	CameraPlane = FPlane3(CameraState.Forward(), CameraState.Position + CameraState.Forward());
-	RectangleInCameraPlane = GetSelectionRectangleUV();
+	SelectionDomain.Plane = FPlane3(CameraState.Forward(), CameraState.Position + CameraState.Forward());
+	SelectionDomain.Rectangle = ProjectSelectionDomain().Rectangle;
 }
 
-FCameraRectangle::FAxisAlignedBox2 FCameraRectangle::GetSelectionRectangleUV(double OffsetFromCameraPlane) const
+FCameraRectangle::FRectangleInPlane FCameraRectangle::ProjectSelectionDomain(double OffsetFromSelectionDomain) const
 {
 	ensure(bIsInitialized);
 	
-	FAxisAlignedBox2 Result;
+	FRectangleInPlane Result;
 
-	FPlane3 SelectionPlane = CameraPlane;
-	SelectionPlane.Constant += OffsetFromCameraPlane;
+	Result.Plane = SelectionDomain.Plane;
+	Result.Plane.Constant += OffsetFromSelectionDomain;
 	
-	FVector ProjectedStart = FMath::RayPlaneIntersection(RectangleStartRay.WorldRay.Origin, RectangleStartRay.WorldRay.Direction, (FPlane)SelectionPlane);
-	FVector ProjectedEnd = FMath::RayPlaneIntersection(RectangleEndRay.WorldRay.Origin, RectangleEndRay.WorldRay.Direction, (FPlane)SelectionPlane);
-	FVector2 StartUV = PlaneCoordinates(SelectionPlane, ProjectedStart);
-	FVector2 EndUV = PlaneCoordinates(SelectionPlane, ProjectedEnd);
+	FVector ProjectedStart = FMath::RayPlaneIntersection(RectangleStartRay.WorldRay.Origin, RectangleStartRay.WorldRay.Direction, (FPlane)Result.Plane);
+	FVector ProjectedEnd = FMath::RayPlaneIntersection(RectangleEndRay.WorldRay.Origin, RectangleEndRay.WorldRay.Direction, (FPlane)Result.Plane);
+	FVector2 StartUV = Point3DToPointUV(Result.Plane, ProjectedStart);
+	FVector2 EndUV = Point3DToPointUV(Result.Plane, ProjectedEnd);
 
 	// Initialize this way so we don't have to care about min/max
-	Result.Contain(StartUV);
-	Result.Contain(EndUV);
+	Result.Rectangle.Contain(StartUV);
+	Result.Rectangle.Contain(EndUV);
 
 	return Result;
 }
@@ -47,20 +47,20 @@ bool FCameraRectangle::IsProjectedPointInRectangle(const FVector& Point) const
 	FVector ProjectedPoint;
 	if (CameraState.bIsOrthographic)
 	{
-		ProjectedPoint = OrthographicProjection(CameraPlane, Point);
+		ProjectedPoint = OrthographicProjection(SelectionDomain.Plane, Point);
 	}
 	else
 	{
 		// If it's not in front of the camera plane, its not contained in the camera rectangle
-		if (CameraPlane.DistanceTo(Point) <= 0)
+		if (SelectionDomain.Plane.DistanceTo(Point) <= 0)
 		{
 			return false;
 		}
-		ProjectedPoint = PerspectiveProjection(CameraPlane, Point);
+		ProjectedPoint = PerspectiveProjection(SelectionDomain.Plane, Point);
 	}
 
-	FVector2 PointUV = PlaneCoordinates(CameraPlane, ProjectedPoint);
-	return RectangleInCameraPlane.Contains(PointUV);
+	FVector2 PointUV = Point3DToPointUV(SelectionDomain.Plane, ProjectedPoint);
+	return SelectionDomain.Rectangle.Contains(PointUV);
 }
 
 bool FCameraRectangle::IsProjectedSegmentIntersectingRectangle(const FVector& Endpoint1, const FVector& Endpoint2) const
@@ -72,8 +72,8 @@ bool FCameraRectangle::IsProjectedSegmentIntersectingRectangle(const FVector& En
 
 	if (CameraState.bIsOrthographic)
 	{
-		ProjectedEndpoint1 = OrthographicProjection(CameraPlane, Endpoint1);
-		ProjectedEndpoint2 = OrthographicProjection(CameraPlane, Endpoint2);
+		ProjectedEndpoint1 = OrthographicProjection(SelectionDomain.Plane, Endpoint1);
+		ProjectedEndpoint2 = OrthographicProjection(SelectionDomain.Plane, Endpoint2);
 	}
 	else
 	{
@@ -81,7 +81,7 @@ bool FCameraRectangle::IsProjectedSegmentIntersectingRectangle(const FVector& En
 		ProjectedEndpoint2 = Endpoint2;
 
 		// We have to crop the segment to the portion in front of the camera plane
-		FPlane3::EClipSegmentType ClipType = CameraPlane.ClipSegment(ProjectedEndpoint1, ProjectedEndpoint2);
+		FPlane3::EClipSegmentType ClipType = SelectionDomain.Plane.ClipSegment(ProjectedEndpoint1, ProjectedEndpoint2);
 		
 		// Since the selection plane is identical to the clipping plane there's no need to reproject clipped points
 		switch (ClipType)
@@ -89,47 +89,24 @@ bool FCameraRectangle::IsProjectedSegmentIntersectingRectangle(const FVector& En
 			case FPlane3::FullyClipped:
 				return false; // Segment is behind the camera plane hence not in the camera rectangle
 			case FPlane3::FirstClipped:
-				ProjectedEndpoint2 = PerspectiveProjection(CameraPlane, ProjectedEndpoint2);
+				ProjectedEndpoint2 = PerspectiveProjection(SelectionDomain.Plane, ProjectedEndpoint2);
 				break;
 			case FPlane3::SecondClipped:
-				ProjectedEndpoint1 = PerspectiveProjection(CameraPlane, ProjectedEndpoint1);
+				ProjectedEndpoint1 = PerspectiveProjection(SelectionDomain.Plane, ProjectedEndpoint1);
 				break;
 			case FPlane3::NotClipped:
 			default:
-				ProjectedEndpoint1 = PerspectiveProjection(CameraPlane, ProjectedEndpoint1);
-				ProjectedEndpoint2 = PerspectiveProjection(CameraPlane, ProjectedEndpoint2);
+				ProjectedEndpoint1 = PerspectiveProjection(SelectionDomain.Plane, ProjectedEndpoint1);
+				ProjectedEndpoint2 = PerspectiveProjection(SelectionDomain.Plane, ProjectedEndpoint2);
+				break;
 		}
 	}
 
-	// TODO Port Wild Magic IntrSegment2Box2 (which requires porting IntrLine2Box2) so we can call segment-box intersection here
-
-	FVector2 Endpoint1UV = PlaneCoordinates(CameraPlane, ProjectedEndpoint1);
-	FVector2 Endpoint2UV = PlaneCoordinates(CameraPlane, ProjectedEndpoint2);
+	FVector2 Endpoint1UV = Point3DToPointUV(SelectionDomain.Plane, ProjectedEndpoint1);
+	FVector2 Endpoint2UV = Point3DToPointUV(SelectionDomain.Plane, ProjectedEndpoint2);
 	FSegment2 ProjectedSegmentUV(Endpoint1UV, Endpoint2UV);
 
-	// If either endpoint is inside, then definitely (at least partially) contained
-	if (RectangleInCameraPlane.Contains(ProjectedSegmentUV.StartPoint()) ||
-		RectangleInCameraPlane.Contains(ProjectedSegmentUV.EndPoint()))
-	{
-		return true;
-	}
-
-	// If both outside, have to do some intersections with the box sides
-
-	if (ProjectedSegmentUV.Intersects(FSegment2(RectangleInCameraPlane.GetCorner(0), RectangleInCameraPlane.GetCorner(1))))
-	{
-		return true;
-	}
-
-	if (ProjectedSegmentUV.Intersects(FSegment2(RectangleInCameraPlane.GetCorner(1), RectangleInCameraPlane.GetCorner(2))))
-	{
-		return true;
-	}
-
-	return ProjectedSegmentUV.Intersects(FSegment2(RectangleInCameraPlane.GetCorner(3), RectangleInCameraPlane.GetCorner(2)));
-
-	// Don't need to intersect with the fourth side because segment would have to intersect two sides
-	// of box if both endpoints are outside the box.
+	return TestIntersection(ProjectedSegmentUV, SelectionDomain.Rectangle);
 }
 
 // ---------------------------------------
diff --git a/Engine/Plugins/Experimental/MeshModelingToolsetExp/Source/ModelingComponents/Private/Selection/MeshSelectionMechanic.cpp b/Engine/Plugins/Experimental/MeshModelingToolsetExp/Source/ModelingComponents/Private/Selection/MeshSelectionMechanic.cpp
--- a/Engine/Plugins/Experimental/MeshModelingToolsetExp/Source/ModelingComponents/Private/Selection/MeshSelectionMechanic.cpp
+++ b/Engine/Plugins/Experimental/MeshModelingToolsetExp/Source/ModelingComponents/Private/Selection/MeshSelectionMechanic.cpp
@@ -7,6 +7,8 @@
 #include "Drawing/LineSetComponent.h"
 #include "Drawing/PointSetComponent.h"
 #include "Drawing/PreviewGeometryActor.h"
+#include "Intersection/IntrTriangle2AxisAlignedBox2.h"
+#include "Intersection/IntersectionQueries2.h"
 #include "InteractiveToolManager.h"
 #include "Polyline3.h"
 #include "Selections/MeshConnectedComponents.h"
@@ -90,6 +92,92 @@ namespace MeshSelectionMechanicLocals
 		checkNoEntry();
 		return UE::Geometry::FDynamicMeshSelection::EType::Vertex;
 	}
+
+	template <typename Real>
+	TVector2<FVector::FReal> XY(const TVector<Real>& Point)
+	{
+		return {static_cast<FVector::FReal>(Point.X), static_cast<FVector::FReal>(Point.Y)};
+	}
+
+	// Returns the marquee selection rectangle, obtained from the given CameraRectangle, projected to the XY plane
+	FCameraRectangle::FAxisAlignedBox2 GetRectangleXY(const FCameraRectangle& CameraRectangle)
+	{
+		ensure(CameraRectangle.bIsInitialized);
+		
+		double Offset = CameraRectangle.SelectionDomain.Plane.DistanceTo(FVector::ZeroVector);
+		FCameraRectangle::FRectangleInPlane Domain = CameraRectangle.ProjectSelectionDomain(Offset);
+		FCameraRectangle::FAxisAlignedBox2 RectangleXY;
+		
+		// This works because we know the UV axes are aligned with the XY axes, see the comment in UUVEditorMode::InitializeTargets
+		RectangleXY.Contain(XY(CameraRectangle.PointUVToPoint3D(Domain.Plane, Domain.Rectangle.Min)));
+		RectangleXY.Contain(XY(CameraRectangle.PointUVToPoint3D(Domain.Plane, Domain.Rectangle.Max)));
+	
+		return RectangleXY;
+	}
+
+	// Returns the points of the given mesh that lie inside the given rectangle. Assumptions:
+	// - Mesh vertices lie in the XY plane (have zero Z coordinate)
+	// - The camera is above the mesh and looking in the -Z direction (so no clip testing is needed)
+	TArray<int32> FindProjectedVerticesInRectangleXY0(const FCameraRectangle& CameraRectangle, const FDynamicMesh3& MeshXY0)
+	{
+		FCameraRectangle::FAxisAlignedBox2 RectangleXY = GetRectangleXY(CameraRectangle);
+		
+		TArray<int32> Result;
+		for (int32 VertexIndex : MeshXY0.VertexIndicesItr())
+		{
+			if (RectangleXY.Contains(XY(MeshXY0.GetVertex(VertexIndex))))
+			{
+				Result.Add(VertexIndex);
+			}
+		}
+
+		return Result;
+	}
+	
+	// Returns the edges of the given mesh that lie inside the given rectangle. Assumptions:
+	// - Mesh vertices lie in the XY plane (have zero Z coordinate)
+	// - The camera is above the mesh and looking in the -Z direction (so no clip testing is needed)
+	TArray<int32> FindProjectedEdgesInRectangleXY0(const FCameraRectangle& CameraRectangle, const FDynamicMesh3& MeshXY0)
+	{
+		FCameraRectangle::FAxisAlignedBox2 RectangleXY = GetRectangleXY(CameraRectangle);
+		
+		TArray<int32> Result;
+		for (int32 EdgeIndex : MeshXY0.EdgeIndicesItr())
+		{
+			const FDynamicMesh3::FEdge& Edge = MeshXY0.GetEdgeRef(EdgeIndex);
+			FCameraRectangle::FSegment2 SegmentXY(XY(MeshXY0.GetVertex(Edge.Vert.A)),
+												  XY(MeshXY0.GetVertex(Edge.Vert.B)));
+			if (TestIntersection(SegmentXY, RectangleXY))
+			{
+				Result.Add(EdgeIndex);
+			}
+		}
+		return Result;
+	}
+	
+	// Returns the triangles of the given mesh that lie inside the given rectangle. Assumptions:
+	// - Mesh vertices lie in the XY plane (have zero Z coordinate)
+	// - The camera is above the mesh and looking in the -Z direction (so no clip testing is needed)
+	TArray<int32> FindProjectedTrianglesInRectangleXY0(const FCameraRectangle& CameraRectangle, const FDynamicMesh3& MeshXY0)
+	{
+		FCameraRectangle::FAxisAlignedBox2 RectangleXY = GetRectangleXY(CameraRectangle);
+		
+		TArray<int32> Result;
+		for (int32 TriangleIndex : MeshXY0.TriangleIndicesItr())
+		{
+			const FIndex3i& Triangle = MeshXY0.GetTriangleRef(TriangleIndex);
+			const TTriangle2<FVector::FReal> TriangleXY(XY(MeshXY0.GetVertex(Triangle.A)),
+														XY(MeshXY0.GetVertex(Triangle.B)),
+														XY(MeshXY0.GetVertex(Triangle.C)));
+
+			// Check with bTriangleIsOriented = false since some triangles maybe oriented away from the camera
+			if (TIntrTriangle2AxisAlignedBox2<FVector::FReal> Intersects(TriangleXY, RectangleXY, false); Intersects.Test())
+			{
+				Result.Add(TriangleIndex);
+			}
+		}
+		return Result;
+	}
 } // namespace MeshSelectionMechanicLocals
 
 void UMeshSelectionMechanic::Setup(UInteractiveTool* ParentToolIn)
@@ -459,7 +547,8 @@ void UMeshSelectionMechanic::OnClicked(const FInputDeviceRay& ClickPos)
 				TArray<int32> SeedTriangles;
 				SeedTriangles.Add(HitTid);
 				MeshSelectedComponent.FindTrianglesConnectedToSeeds(SeedTriangles);
-				ClickSelectedIDs = TSet<int32>(SeedTriangles);	
+				ensure(MeshSelectedComponent.Components.Num() == 1); // Expect each triangle to only be in a single component
+				ClickSelectedIDs.Append(MoveTemp(MeshSelectedComponent.Components[0].Indices));
 			}
 			else if (SelectionMode == EMeshSelectionMechanicMode::Edge)
 			{
@@ -558,6 +647,11 @@ void UMeshSelectionMechanic::OnDragRectangleChanged(const FCameraRectangle& Curr
 			for (int32 MeshIndex = 0; MeshIndex < MeshSpatials.Num(); ++MeshIndex)
 			{
 				const FDynamicMesh3& Mesh = *(MeshSpatials[MeshIndex]->GetMesh());
+
+				// TODO UMeshSelectionMechanic is currently assuming that the provided mesh transform is identity, which is
+				//  the case in the UV editor, this restriction should be lifted. When we do this we should also apply the
+				//  (inverse) tranform to the CurrentRectangle query rays. Search :ApplyTransformToQuery
+				ensure(MeshTransforms[MeshIndex].Identical(&FTransform::Identity, 0));
 				AddIDsToSelection(Mesh, MeshTransforms[MeshIndex], RectangleSelectedIDs);
 	
 				if (!RectangleSelectedIDs.IsEmpty())
@@ -573,8 +667,11 @@ void UMeshSelectionMechanic::OnDragRectangleChanged(const FCameraRectangle& Curr
 		{
 			ensure(CurrentSelectionIndex != IndexConstants::InvalidID);
 			ensure(CurrentSelection.Mesh == MeshSpatials[CurrentSelectionIndex]->GetMesh());
-			
+
 			const FDynamicMesh3& Mesh = *(MeshSpatials[CurrentSelectionIndex]->GetMesh());
+
+			// TODO See :ApplyTransformToQuery
+			ensure(MeshTransforms[CurrentSelectionIndex].Identical(&FTransform::Identity, 0));
 			AddIDsToSelection(Mesh, MeshTransforms[CurrentSelectionIndex], RectangleSelectedIDs);
 		}
 
@@ -587,14 +684,8 @@ void UMeshSelectionMechanic::OnDragRectangleChanged(const FCameraRectangle& Curr
 		auto AddVertexIDsToSelection = [this, &CurrentRectangle]
 			(const FDynamicMesh3& Mesh, const FTransform& Transform, TSet<int32>& OutSelectedIDs)
 		{
-			for (int32 VertexIndex : Mesh.VertexIndicesItr())
-			{
-				FVector3d Point = Transform.InverseTransformPosition(Mesh.GetVertex(VertexIndex));
-				if (CurrentRectangle.IsProjectedPointInRectangle(Point))
-				{
-					OutSelectedIDs.Add(VertexIndex);
-				}
-			}
+			TArray<int32> SelectedIDs = MeshSelectionMechanicLocals::FindProjectedVerticesInRectangleXY0(CurrentRectangle, Mesh);
+			OutSelectedIDs.Append(MoveTemp(SelectedIDs));
 		};
 	
 		RectangleSelectedIDs = FindIDsInsideCurrentRectangle(AddVertexIDsToSelection);
@@ -604,32 +695,55 @@ void UMeshSelectionMechanic::OnDragRectangleChanged(const FCameraRectangle& Curr
 		auto AddEdgeIDsToSelection = [this, &CurrentRectangle]
 			(const FDynamicMesh3& Mesh, const FTransform& Transform, TSet<int32>& OutSelectedIDs)
 		{
-			for (int32 EdgeIndex : Mesh.EdgeIndicesItr())
-			{
-				const FDynamicMesh3::FEdge& Edge = Mesh.GetEdgeRef(EdgeIndex);
-				FVector3d StartPoint = Transform.InverseTransformPosition(Mesh.GetVertex(Edge.Vert.A));
-				FVector3d EndPoint = Transform.InverseTransformPosition(Mesh.GetVertex(Edge.Vert.B));
-   				
-				if (CurrentRectangle.IsProjectedSegmentIntersectingRectangle(StartPoint, EndPoint))
-				{
-					OutSelectedIDs.Add(EdgeIndex);
-				}
-			}
+			TArray<int32> SelectedIDs = MeshSelectionMechanicLocals::FindProjectedEdgesInRectangleXY0(CurrentRectangle, Mesh);
+			OutSelectedIDs.Append(MoveTemp(SelectedIDs));
 		};
 
 		RectangleSelectedIDs = FindIDsInsideCurrentRectangle(AddEdgeIDsToSelection);
 	}
 	else if (SelectionMode == EMeshSelectionMechanicMode::Triangle)
 	{
-		// TODO Implement me
+		auto AddTriangleIDsToSelection = [this, &CurrentRectangle]
+			(const FDynamicMesh3& Mesh, const FTransform& Transform, TSet<int32>& OutSelectedIDs)
+		{
+			TArray<int32> SelectedIDs = MeshSelectionMechanicLocals::FindProjectedTrianglesInRectangleXY0(CurrentRectangle, Mesh);
+			OutSelectedIDs.Append(MoveTemp(SelectedIDs));
+		};
+	
+		RectangleSelectedIDs = FindIDsInsideCurrentRectangle(AddTriangleIDsToSelection);
 	}
 	else if (SelectionMode == EMeshSelectionMechanicMode::Component)
 	{
-		// TODO Implement me
+		auto AddTriangleIDsToSelection = [this, &CurrentRectangle]
+			(const FDynamicMesh3& Mesh, const FTransform& Transform, TSet<int32>& OutSelectedIDs)
+		{
+			TArray<int32> SeedTriangles = MeshSelectionMechanicLocals::FindProjectedTrianglesInRectangleXY0(CurrentRectangle, Mesh);
+			FMeshConnectedComponents MeshSelectedComponent(&Mesh);
+			MeshSelectedComponent.FindTrianglesConnectedToSeeds(SeedTriangles);
+			for (int ComponentIndex = 0; ComponentIndex < MeshSelectedComponent.Components.Num(); ComponentIndex++)
+			{
+				OutSelectedIDs.Append(MoveTemp(MeshSelectedComponent.Components[ComponentIndex].Indices));
+			}
+		};
+	
+		RectangleSelectedIDs = FindIDsInsideCurrentRectangle(AddTriangleIDsToSelection);
 	}
 	else if (SelectionMode == EMeshSelectionMechanicMode::Mesh)
 	{
-		// TODO Implement me
+		auto AddAllMeshTriangleIDsToSelection = [this, &CurrentRectangle]
+			(const FDynamicMesh3& Mesh, const FTransform& Transform, TSet<int32>& OutSelectedIDs)
+		{
+			TArray<int32> SelectedIDs = MeshSelectionMechanicLocals::FindProjectedTrianglesInRectangleXY0(CurrentRectangle, Mesh);
+			if (!SelectedIDs.IsEmpty())
+			{
+				for (int32 Tid : Mesh.TriangleIndicesItr())
+				{
+					OutSelectedIDs.Add(Tid);
+				}
+			}
+		};
+
+		RectangleSelectedIDs = FindIDsInsideCurrentRectangle(AddAllMeshTriangleIDsToSelection);
 	}
 	else
 	{
diff --git a/Engine/Plugins/Experimental/MeshModelingToolsetExp/Source/ModelingComponents/Public/Mechanics/RectangleMarqueeMechanic.h b/Engine/Plugins/Experimental/MeshModelingToolsetExp/Source/ModelingComponents/Public/Mechanics/RectangleMarqueeMechanic.h
--- a/Engine/Plugins/Experimental/MeshModelingToolsetExp/Source/ModelingComponents/Public/Mechanics/RectangleMarqueeMechanic.h
+++ b/Engine/Plugins/Experimental/MeshModelingToolsetExp/Source/ModelingComponents/Public/Mechanics/RectangleMarqueeMechanic.h
@@ -35,18 +35,27 @@ struct MODELINGCOMPONENTS_API FCameraRectangle
 	FViewCameraState CameraState;
 	FInputDeviceRay RectangleStartRay = FInputDeviceRay(FRay());
 	FInputDeviceRay RectangleEndRay = FInputDeviceRay(FRay());
-	FPlane3 CameraPlane;
-	FAxisAlignedBox2 RectangleInCameraPlane;
+	
+	struct FRectangleInPlane
+	{
+		// Rectangle is defined in the UV-coordinate system of Plane. See PlaneCoordinates
+		FPlane3 Plane;
+		FAxisAlignedBox2 Rectangle; // Region of marquee selection
+	};
+
+	FRectangleInPlane SelectionDomain;
 
 	// This function must be called before other member functions whenever camera state or start/end rays are updated
 	void Initialize();
 	
 	bool bIsInitialized = false;
 
-	// Return an axis aligned box (the region of a marquee selection) in the UV-coordinate system of a plane offset from
-	// CameraPlane by the given distance in the CameraPlane normal direction. This function was added in order to
-	// implement optimizations (avoid projecting points) when selecting UV editor meshes, which lie in the XY plane
-	FAxisAlignedBox2 GetSelectionRectangleUV(double OffsetFromCameraPlane = 0.) const;
+	// Returns the SelectionDomain projected onto a plane offset from SelectionDomain.Plane by the given distance in the
+	// SelectionDomain.Plane.Normal direction.
+	//
+	// Note: This function is in the public API to optimize selection in the UV editor where meshes are known to lie in
+	// the XY plane in front of the camera (so point projection/camera plane clipping is unnecessary)
+	FRectangleInPlane ProjectSelectionDomain(double OffsetFromCameraPlane = 0.) const;
 
 	// Return true if the given 3D geometry projected to the camera plane is inside or intersecting the rectangle, and
 	// false otherwise
@@ -69,17 +78,29 @@ struct MODELINGCOMPONENTS_API FCameraRectangle
 		return FMath::RayPlaneIntersection(CameraState.Position, Point - CameraState.Position, (FPlane)Plane);
 	}
 
-	// Given a 3D point lying in the given plane, return the UV coordinates of the point expressed in the following a
-	// two dimensional parameterization of the given plane:
+	// Given a 3D point lying in the given Plane, return the UV coordinates of the point expressed in the following a
+	// two dimensional parameterization of the given Plane:
 	// - the 2D origin (0,0) is located at the foot of the projection of cartesian point (0,0,0) onto the plane
 	// - the U basis vector is the camera right vector
 	// - the V basis vector is the camera up vector
-	FVector2 PlaneCoordinates(const FPlane3& Plane, const FVector& PointInPlane) const
+	FVector2 Point3DToPointUV(const FPlane3& Plane, const FVector& Point3DInPlane) const
+	{
+		ensure(bIsInitialized);
+		FVector2 PointUV;
+		PointUV[0] = FVector::DotProduct(Point3DInPlane - ((FPlane)Plane).GetOrigin(), CameraState.Right());
+		PointUV[1] = FVector::DotProduct(Point3DInPlane - ((FPlane)Plane).GetOrigin(), CameraState.Up());
+		return PointUV;
+	}
+	
+	// Given a 2D point in the UV space of the given Plane, return the coordinates of the 3D point lying in the Plane.
+	// This function is the inverse of Point3DToPointUV.
+	FVector PointUVToPoint3D(const FPlane3& Plane, const FVector2& PointUV) const
 	{
 		ensure(bIsInitialized);
-		FVector::FReal U = FVector::DotProduct(PointInPlane - ((FPlane)Plane).GetOrigin(), CameraState.Right());
-		FVector::FReal V = FVector::DotProduct(PointInPlane - ((FPlane)Plane).GetOrigin(), CameraState.Up());
-		return FVector2{U, V};
+		FVector Point3DInPlane = ((FPlane)Plane).GetOrigin();
+		Point3DInPlane += PointUV[0] * CameraState.Right();
+		Point3DInPlane += PointUV[1] * CameraState.Up();
+		return Point3DInPlane;
 	}
 };
 
diff --git a/Engine/Source/Runtime/GeometryCore/Private/Intersection/IntersectionQueries2.cpp b/Engine/Source/Runtime/GeometryCore/Private/Intersection/IntersectionQueries2.cpp
new file mode 100644
--- /dev/null
+++ b/Engine/Source/Runtime/GeometryCore/Private/Intersection/IntersectionQueries2.cpp
@@ -0,0 +1,45 @@
+﻿// Copyright Epic Games, Inc. All Rights Reserved.
+
+#include "Intersection/IntersectionQueries2.h"
+
+using namespace UE::Geometry;
+
+template<typename Real>
+bool UE::Geometry::TestIntersection(const TSegment2<Real>& Segment, const TAxisAlignedBox2<Real>& Box)
+{
+	// TODO Port Wild Magic IntrSegment2Box2 (which requires porting IntrLine2Box2) so we can call segment-box intersection here
+	
+	// If either endpoint is inside, then definitely (at least partially) contained
+	if (Box.Contains(Segment.StartPoint()) || Box.Contains(Segment.EndPoint()))
+	{
+		return true;
+	}
+
+	// If both outside, have to do some intersections with the box sides
+
+	if (Segment.Intersects(TSegment2<Real>(Box.GetCorner(0), Box.GetCorner(1))))
+	{
+		return true;
+	}
+
+	if (Segment.Intersects(TSegment2<Real>(Box.GetCorner(1), Box.GetCorner(2))))
+	{
+		return true;
+	}
+
+	return Segment.Intersects(TSegment2<Real>(Box.GetCorner(3), Box.GetCorner(2)));
+
+	// Don't need to intersect with the fourth side because segment would have to intersect two sides
+	// of box if both endpoints are outside the box.
+}
+
+namespace UE
+{
+namespace Geometry
+{
+
+template bool GEOMETRYCORE_API TestIntersection(const TSegment2<float>& Segment, const TAxisAlignedBox2<float>& Box);
+template bool GEOMETRYCORE_API TestIntersection(const TSegment2<double>& Segment, const TAxisAlignedBox2<double>& Box);
+		
+} // namespace UE::Geometry
+} // namespace UE
diff --git a/Engine/Source/Runtime/GeometryCore/Public/Intersection/IntersectionQueries2.h b/Engine/Source/Runtime/GeometryCore/Public/Intersection/IntersectionQueries2.h
new file mode 100644
--- /dev/null
+++ b/Engine/Source/Runtime/GeometryCore/Public/Intersection/IntersectionQueries2.h
@@ -0,0 +1,19 @@
+﻿// Copyright Epic Games, Inc. All Rights Reserved.
+
+#pragma once
+
+#include "CoreMinimal.h"
+#include "SegmentTypes.h"
+#include "BoxTypes.h"
+
+namespace UE
+{
+namespace Geometry
+{
+
+// Return true if Segment and Box intersect and false otherwise
+template<typename Real>
+bool TestIntersection(const TSegment2<Real>& Segment, const TAxisAlignedBox2<Real>& Box);
+		
+} // namespace UE::Geometry
+} // namespace UE
