MeshModelingTools: Add point set support to UPreviewGeometry.

#rnx

diff --git a/Engine/Plugins/Experimental/MeshModelingToolsetExp/Source/ModelingComponents/Private/Drawing/PreviewGeometryActor.cpp b/Engine/Plugins/Experimental/MeshModelingToolsetExp/Source/ModelingComponents/Private/Drawing/PreviewGeometryActor.cpp
--- a/Engine/Plugins/Experimental/MeshModelingToolsetExp/Source/ModelingComponents/Private/Drawing/PreviewGeometryActor.cpp
+++ b/Engine/Plugins/Experimental/MeshModelingToolsetExp/Source/ModelingComponents/Private/Drawing/PreviewGeometryActor.cpp
@@ -71,6 +71,10 @@ void UPreviewGeometry::SetAllVisible(bool bVisible)
 	{
 		Entry.Value->SetVisibility(bVisible);
 	}
+	for (TPair<FString, TObjectPtr<UPointSetComponent>> Entry : PointSets)
+	{
+		Entry.Value->SetVisibility(bVisible);
+	}
 }
 
 
@@ -131,9 +135,9 @@ bool UPreviewGeometry::RemoveLineSet(const FString& LineSetIdentifier, bool bDes
 
 void UPreviewGeometry::RemoveAllLineSets(bool bDestroy)
 {
-	for (TPair<FString, TObjectPtr<ULineSetComponent>> Entry : LineSets)
+	if (bDestroy)
 	{
-		if (bDestroy)
+		for (TPair<FString, TObjectPtr<ULineSetComponent>> Entry : LineSets)
 		{
 			Entry.Value->UnregisterComponent();
 			Entry.Value->DestroyComponent();
@@ -197,3 +201,96 @@ void UPreviewGeometry::CreateOrUpdateLineSet(const FString& LineSetIdentifier, i
 	LineSet->AddLines(NumIndices, LineGenFunc, LinesPerIndexHint);
 }
 
+
+UPointSetComponent* UPreviewGeometry::AddPointSet(const FString& SetIdentifier)
+{
+	if (PointSets.Contains(SetIdentifier))
+	{
+		check(false);
+		return nullptr;
+	}
+
+	UPointSetComponent* PointSet = NewObject<UPointSetComponent>(ParentActor);
+	PointSet->SetupAttachment(ParentActor->GetRootComponent());
+
+	UMaterialInterface* PointMaterial = ToolSetupUtil::GetDefaultPointComponentMaterial(nullptr);
+	if (PointMaterial != nullptr)
+	{
+		PointSet->SetPointMaterial(PointMaterial);
+	}
+
+	PointSet->RegisterComponent();
+
+	PointSets.Add(SetIdentifier, PointSet);
+	return PointSet;
+}
+
+UPointSetComponent* UPreviewGeometry::FindPointSet(const FString& PointSetIdentifier)
+{
+	TObjectPtr<UPointSetComponent>* Found = PointSets.Find(PointSetIdentifier);
+	if (Found != nullptr)
+	{
+		return *Found;
+	}
+	return nullptr;
+}
+
+bool UPreviewGeometry::RemovePointSet(const FString& PointSetIdentifier, bool bDestroy)
+{
+	TObjectPtr<UPointSetComponent>* Found = PointSets.Find(PointSetIdentifier);
+	if (Found != nullptr)
+	{
+		UPointSetComponent* PointSet = *Found;
+		PointSets.Remove(PointSetIdentifier);
+		if (bDestroy)
+		{
+			PointSet->UnregisterComponent();
+			PointSet->DestroyComponent();
+			PointSet = nullptr;
+		}
+	}
+	return false;
+}
+
+void UPreviewGeometry::RemoveAllPointSets(bool bDestroy)
+{
+	if (bDestroy)
+	{
+		for (TPair<FString, TObjectPtr<UPointSetComponent>> Entry : PointSets)
+		{
+			Entry.Value->UnregisterComponent();
+			Entry.Value->DestroyComponent();
+		}
+	}
+	PointSets.Reset();
+}
+
+bool UPreviewGeometry::SetPointSetVisibility(const FString& PointSetIdentifier, bool bVisible)
+{
+	TObjectPtr<UPointSetComponent>* Found = PointSets.Find(PointSetIdentifier);
+	if (Found != nullptr)
+	{
+		(*Found)->SetVisibility(bVisible);
+		return true;
+	}
+	return false;
+}
+
+bool UPreviewGeometry::SetPointSetMaterial(const FString& PointSetIdentifier, UMaterialInterface* NewMaterial)
+{
+	TObjectPtr<UPointSetComponent>* Found = PointSets.Find(PointSetIdentifier);
+	if (Found != nullptr)
+	{
+		(*Found)->SetPointMaterial(NewMaterial);
+		return true;
+	}
+	return false;
+}
+
+void UPreviewGeometry::SetAllPointSetsMaterial(UMaterialInterface* Material)
+{
+	for (TPair<FString, TObjectPtr<UPointSetComponent>> Entry : PointSets)
+	{
+		Entry.Value->SetPointMaterial(Material);
+	}
+}
\ No newline at end of file
diff --git a/Engine/Plugins/Experimental/MeshModelingToolsetExp/Source/ModelingComponents/Public/Drawing/PreviewGeometryActor.h b/Engine/Plugins/Experimental/MeshModelingToolsetExp/Source/ModelingComponents/Public/Drawing/PreviewGeometryActor.h
--- a/Engine/Plugins/Experimental/MeshModelingToolsetExp/Source/ModelingComponents/Public/Drawing/PreviewGeometryActor.h
+++ b/Engine/Plugins/Experimental/MeshModelingToolsetExp/Source/ModelingComponents/Public/Drawing/PreviewGeometryActor.h
@@ -5,6 +5,7 @@
 #include "CoreMinimal.h"
 #include "InteractiveToolObjects.h"
 #include "Drawing/LineSetComponent.h"
+#include "Drawing/PointSetComponent.h"
 
 #include "PreviewGeometryActor.generated.h"
 
@@ -163,6 +164,83 @@ public:
 		TFunctionRef<void(int32 Index, TArray<FRenderableLine>& LinesOut)> LineGenFunc,
 		int32 LinesPerIndexHint = -1);
 
+	//
+	// Point Sets
+	//
+
+	/** Create a new point set with the given PointSetIdentifier and return it */
+	UFUNCTION()
+	UPointSetComponent* AddPointSet(const FString& PointSetIdentifier);
+
+	/** @return the PointSetComponent with the given PointSetIdentifier, or nullptr if not found */
+	UFUNCTION()
+	UPointSetComponent* FindPointSet(const FString& PointSetIdentifier);
+
+	/** 
+	 * Remove the PointSetComponent with the given PointSetIdentifier
+	 * @param bDestroy if true, component will unregistered and destroyed. 
+	 * @return true if the PointSetComponent was found and removed
+	 */
+	UFUNCTION()
+	bool RemovePointSet(const FString& PointSetIdentifier, bool bDestroy = true);
+
+	/**
+	 * Remove all PointSetComponents
+	 * @param bDestroy if true, the components will unregistered and destroyed.
+	 */
+	UFUNCTION()
+	void RemoveAllPointSets(bool bDestroy = true);
+
+	/**
+	 * Set the visibility of the PointSetComponent with the given PointSetIdentifier
+	 * @return true if the PointSetComponent was found and updated
+	 */
+	UFUNCTION()
+	bool SetPointSetVisibility(const FString& PointSetIdentifier, bool bVisible);
+
+	/**
+	 * Set the Material of the PointSetComponent with the given PointSetIdentifier
+	 * @return true if the PointSetComponent was found and updated
+	 */
+	UFUNCTION()
+	bool SetPointSetMaterial(const FString& PointSetIdentifier, UMaterialInterface* NewMaterial);
+
+	/**
+	 * Set the Material of all PointSetComponents
+	 */
+	UFUNCTION()
+	void SetAllPointSetsMaterial(UMaterialInterface* Material);
+
+
+
+	//
+	// Point Set Utilities
+	//
+
+	/**
+	 * Find the identified point set and call UpdateFuncType(UPointSetComponent*)
+	 */
+	template<typename UpdateFuncType>
+	void UpdatePointSet(const FString& PointSetIdentifier, UpdateFuncType UpdateFunc)
+	{
+		UPointSetComponent* PointSet = FindPointSet(PointSetIdentifier);
+		if (PointSet)
+		{
+			UpdateFunc(PointSet);
+		}
+	}
+
+	/**
+	 * call UpdateFuncType(UPointSetComponent*) for all existing Point Sets
+	 */
+	template<typename UpdateFuncType>
+	void UpdateAllPointSets(UpdateFuncType UpdateFunc)
+	{
+		for (TPair<FString, TObjectPtr<UPointSetComponent>> Entry : PointSets)
+		{
+			UpdateFunc(Entry.Value);
+		}
+	}
 
 public:
 
@@ -178,7 +256,11 @@ public:
 	UPROPERTY()
 	TMap<FString, TObjectPtr<ULineSetComponent>> LineSets;
 
-
+	/**
+	 * PointSetComponents created and owned by the UPreviewGeometry, and added as child components of the ParentActor
+	 */
+	UPROPERTY()
+	TMap<FString, TObjectPtr<UPointSetComponent>> PointSets;
 
 protected:
 
