MeshModelingToolset: Generate tangents in UAddPrimitiveTool via new function to calculcate tangents in UPreviewMesh

#rnx

diff --git a/Engine/Plugins/Experimental/MeshModelingToolsetExp/Source/ModelingComponents/Private/PreviewMesh.cpp b/Engine/Plugins/Experimental/MeshModelingToolsetExp/Source/ModelingComponents/Private/PreviewMesh.cpp
--- a/Engine/Plugins/Experimental/MeshModelingToolsetExp/Source/ModelingComponents/Private/PreviewMesh.cpp
+++ b/Engine/Plugins/Experimental/MeshModelingToolsetExp/Source/ModelingComponents/Private/PreviewMesh.cpp
@@ -190,7 +190,69 @@ void UPreviewMesh::SetTangentsMode(EDynamicMeshComponentTangentsMode TangentsTyp
 	DynamicMeshComponent->SetTangentsType(TangentsType);
 }
 
+bool UPreviewMesh::CalculateTangents()
+{
+	check(DynamicMeshComponent);
+
+	UDynamicMesh *const DynamicMesh = DynamicMeshComponent->GetDynamicMesh();
+	FDynamicMesh3 *const Mesh = DynamicMesh ? DynamicMesh->GetMeshPtr() : nullptr;
+
+	if (Mesh)
+	{
+		// Holds temporary tangents in case we don't have access to existing tangents and need to compute them within this function.
+		FMeshTangentsf TempTangents;
+
+		const FMeshTangentsf* Tangents = [this, Mesh, &TempTangents]() -> const FMeshTangentsf*
+		{
+			if (DynamicMeshComponent->GetTangentsType() == EDynamicMeshComponentTangentsMode::AutoCalculated)
+			{
+				if (const FMeshTangentsf* AutoCalculatedTangents = DynamicMeshComponent->GetAutoCalculatedTangents())
+				{
+					return AutoCalculatedTangents;
+				}
+			}
+
+			const FDynamicMeshNormalOverlay* NormalOverlay = nullptr;
+			const FDynamicMeshUVOverlay* UVOverlay = nullptr;
+
+			if (const FDynamicMeshAttributeSet* Attributes = Mesh->Attributes())
+			{
+				if (Attributes->NumNormalLayers() > 0)
+				{
+					NormalOverlay = Attributes->GetNormalLayer(0);
+				}
+
+				if (Attributes->NumUVLayers() > 0)
+				{
+					UVOverlay = Attributes->GetUVLayer(0);
+				}
+			}
+
+			if (NormalOverlay && UVOverlay)
+			{
+				TempTangents.ComputeTriVertexTangents(NormalOverlay, UVOverlay, {});
+
+				return &TempTangents;
+			}
+
+			return nullptr;
+		}();
+
+		if (Tangents)
+		{
+			Mesh->Attributes()->EnableTangents();
+			if (Tangents->CopyToOverlays(*Mesh))
+			{
+				DynamicMeshComponent->FastNotifyVertexAttributesUpdated(true, false, false);
+				NotifyWorldPathTracedOutputInvalidated();
+
+				return true;
+			}
+		}
+	}
 
+	return false;
+}
 
 void UPreviewMesh::EnableWireframe(bool bEnable)
 {
diff --git a/Engine/Plugins/Experimental/MeshModelingToolsetExp/Source/ModelingComponents/Public/PreviewMesh.h b/Engine/Plugins/Experimental/MeshModelingToolsetExp/Source/ModelingComponents/Public/PreviewMesh.h
--- a/Engine/Plugins/Experimental/MeshModelingToolsetExp/Source/ModelingComponents/Public/PreviewMesh.h
+++ b/Engine/Plugins/Experimental/MeshModelingToolsetExp/Source/ModelingComponents/Public/PreviewMesh.h
@@ -187,6 +187,13 @@ public:
 	 */
 	void SetTangentsMode(EDynamicMeshComponentTangentsMode TangentsType);
 
+	/**
+	 * Calculate tangents for the underlying component.
+	 * This will calculate and assign tangents for the preview mesh independent of the tangents mode.
+	 * But if the tangents mode is set to AutoCalculated then it will try to use the auto calculated tangents.
+	 * @return true if tangents were successfully calculated and assigned to the underlying mesh
+	 */
+	bool CalculateTangents();
 
 	/**
 	 * Get the current transform on the preview mesh
diff --git a/Engine/Plugins/Runtime/MeshModelingToolset/Source/MeshModelingTools/Private/AddPrimitiveTool.cpp b/Engine/Plugins/Runtime/MeshModelingToolset/Source/MeshModelingTools/Private/AddPrimitiveTool.cpp
--- a/Engine/Plugins/Runtime/MeshModelingToolset/Source/MeshModelingTools/Private/AddPrimitiveTool.cpp
+++ b/Engine/Plugins/Runtime/MeshModelingToolset/Source/MeshModelingTools/Private/AddPrimitiveTool.cpp
@@ -313,8 +313,11 @@ void UAddPrimitiveTool::UpdatePreviewMesh()
 		NewMesh.SetVertex(vid, Pos);
 	}
 
-
 	PreviewMesh->UpdatePreview(&NewMesh);
+
+	PreviewMesh->SetTangentsMode(EDynamicMeshComponentTangentsMode::AutoCalculated);
+	const bool CalculateTangentsSuccessful = PreviewMesh->CalculateTangents();
+	checkSlow(CalculateTangentsSuccessful);
 }
 
 void UAddPrimitiveTool::OnClicked(const FInputDeviceRay& DeviceClickPos)
diff --git a/Engine/Source/Runtime/GeometryCore/Private/DynamicMesh/MeshTangents.cpp b/Engine/Source/Runtime/GeometryCore/Private/DynamicMesh/MeshTangents.cpp
--- a/Engine/Source/Runtime/GeometryCore/Private/DynamicMesh/MeshTangents.cpp
+++ b/Engine/Source/Runtime/GeometryCore/Private/DynamicMesh/MeshTangents.cpp
@@ -261,7 +261,7 @@ void TMeshTangents<RealType>::ComputeTriangleTangents(const FDynamicMeshUVOverla
 
 
 template<typename RealType>
-bool TMeshTangents<RealType>::CopyToOverlays(FDynamicMesh3& MeshToSet)
+bool TMeshTangents<RealType>::CopyToOverlays(FDynamicMesh3& MeshToSet) const
 {
 	if (!MeshToSet.HasAttributes() || MeshToSet.Attributes()->NumNormalLayers() != 3)
 	{
@@ -270,12 +270,12 @@ bool TMeshTangents<RealType>::CopyToOverlays(FDynamicMesh3& MeshToSet)
 
 	// Set aliases to make iterating over tangents and bitangents easier
 	FDynamicMeshNormalOverlay* TangentOverlays[2] = { MeshToSet.Attributes()->PrimaryTangents(), MeshToSet.Attributes()->PrimaryBiTangents() };
-	TArray<TVector<RealType>>* TangentValues[2] = { &Tangents, &Bitangents };
+	const TArray<TVector<RealType>>* TangentValues[2] = { &Tangents, &Bitangents };
 	
 	for (int Idx = 0; Idx < 2; Idx++)
 	{
 		// Create overlay topology
-		TArray<TVector<RealType>>& TV = *TangentValues[Idx];
+		const TArray<TVector<RealType>>& TV = *TangentValues[Idx];
 		TangentOverlays[Idx]->CreateFromPredicate([&MeshToSet, &TV](int ParentVertexIdx, int TriIDA, int TriIDB) -> bool
 		{
 			FIndex3i TriA = MeshToSet.GetTriangle(TriIDA);
@@ -360,6 +360,9 @@ void TMeshTangents<RealType>::ComputeSeparatePerTriangleTangents(const FDynamicM
 template<typename RealType>
 void TMeshTangents<RealType>::ComputeMikkTStyleTangents(const FDynamicMeshNormalOverlay* NormalOverlay, const FDynamicMeshUVOverlay* UVOverlay, const FComputeTangentsOptions& Options)
 {
+	checkSlow(NormalOverlay != nullptr);
+	checkSlow(UVOverlay != nullptr);
+
 	double HardcodedSplitDotThresh = -0.99;
 
 	int32 MaxTriangleID = Mesh->MaxTriangleID();
diff --git a/Engine/Source/Runtime/GeometryCore/Public/DynamicMesh/MeshTangents.h b/Engine/Source/Runtime/GeometryCore/Public/DynamicMesh/MeshTangents.h
--- a/Engine/Source/Runtime/GeometryCore/Public/DynamicMesh/MeshTangents.h
+++ b/Engine/Source/Runtime/GeometryCore/Public/DynamicMesh/MeshTangents.h
@@ -237,7 +237,7 @@ public:
 	 * Set Tangents on mesh overlays
 	 * @param MeshToSet Mesh to copy overlays to; does not need to be the same as the Mesh member of this class
 	 */
-	bool CopyToOverlays(FDynamicMesh3& MeshToSet);
+	bool CopyToOverlays(FDynamicMesh3& MeshToSet) const;
 
 protected:
 	/**
