MeshModelingToolset: Fix warning for having only a single poly group

#rnx

diff --git a/Engine/Plugins/Experimental/MeshModelingToolsetExp/Source/MeshModelingToolsExp/Private/DeformMeshPolygonsTool.cpp b/Engine/Plugins/Experimental/MeshModelingToolsetExp/Source/MeshModelingToolsExp/Private/DeformMeshPolygonsTool.cpp
--- a/Engine/Plugins/Experimental/MeshModelingToolsetExp/Source/MeshModelingToolsExp/Private/DeformMeshPolygonsTool.cpp
+++ b/Engine/Plugins/Experimental/MeshModelingToolsetExp/Source/MeshModelingToolsExp/Private/DeformMeshPolygonsTool.cpp
@@ -928,12 +928,13 @@ void UDeformMeshPolygonsTool::Setup()
 	if (Topology.Groups.Num() < 2)
 	{
 		GetToolManager()->DisplayMessage(
-		    LOCTEXT("NoGroupsWarning", "This object has a single PolyGroup. Use the PolyGroups or Select Tool to assign PolyGroups."),
-		    EToolMessageLevel::UserWarning);
+			LOCTEXT("NoGroupsWarning",
+			        "This object has only a single Polygroup. Use the GenGrps, GrpPnt or TriSel (Create Polygroup) tools to modify Polygroups."),
+			EToolMessageLevel::UserWarning);
 	}
 
 	GetToolManager()->DisplayMessage(
-		LOCTEXT("DeformMeshPolygonsToolDescription", "Deform the mesh by directly manipulating (ie click-and-drag) the PolyGroup Edges/Faces/Vertices."),
+		LOCTEXT("DeformMeshPolygonsToolDescription", "Deform the mesh by directly manipulating (i.e. click-and-drag) the Polygroup Edges/Faces/Vertices."),
 		EToolMessageLevel::UserNotification);
 }
 
diff --git a/Engine/Plugins/Runtime/MeshModelingToolset/Source/MeshModelingTools/Private/EditMeshPolygonsTool.cpp b/Engine/Plugins/Runtime/MeshModelingToolset/Source/MeshModelingTools/Private/EditMeshPolygonsTool.cpp
--- a/Engine/Plugins/Runtime/MeshModelingToolset/Source/MeshModelingTools/Private/EditMeshPolygonsTool.cpp
+++ b/Engine/Plugins/Runtime/MeshModelingToolset/Source/MeshModelingTools/Private/EditMeshPolygonsTool.cpp
@@ -402,7 +402,10 @@ void UEditMeshPolygonsTool::Setup()
 
 	if (Topology->Groups.Num() < 2)
 	{
-		GetToolManager()->DisplayMessage( LOCTEXT("NoGroupsWarning", "This object has a single PolyGroup. Use the PolyGroups or Select Tool to assign PolyGroups."), EToolMessageLevel::UserWarning);
+		GetToolManager()->DisplayMessage(
+			LOCTEXT("NoGroupsWarning",
+			        "This object has only a single Polygroup. Use the GenGrps, GrpPnt or TriSel (Create Polygroup) tools to modify Polygroups."),
+			EToolMessageLevel::UserWarning);
 	}
 
 	if (PostSetupFunction)
