ModelingTools: Add bit depth controls for output texture source data to the bake texture tools.
Add support to FTexture2DBuilder to update the source data directly.

#rnx

diff --git a/Engine/Plugins/Experimental/MeshModelingToolsetExp/Source/MeshModelingToolsExp/Private/BakeMeshAttributeMapsTool.cpp b/Engine/Plugins/Experimental/MeshModelingToolsetExp/Source/MeshModelingToolsExp/Private/BakeMeshAttributeMapsTool.cpp
--- a/Engine/Plugins/Experimental/MeshModelingToolsetExp/Source/MeshModelingToolsExp/Private/BakeMeshAttributeMapsTool.cpp
+++ b/Engine/Plugins/Experimental/MeshModelingToolsetExp/Source/MeshModelingToolsExp/Private/BakeMeshAttributeMapsTool.cpp
@@ -293,6 +293,7 @@ void UBakeMeshAttributeMapsTool::Setup()
 	Settings->WatchProperty(Settings->MapTypes, [this](int32) { bInputsDirty = true; UpdateOnModeChange(); });
 	Settings->WatchProperty(Settings->MapPreview, [this](int32) { UpdateVisualization(); GetToolManager()->PostInvalidation(); });
 	Settings->WatchProperty(Settings->Resolution, [this](EBakeTextureResolution) { bInputsDirty = true; });
+	Settings->WatchProperty(Settings->SourceFormat, [this](EBakeTextureFormat) { bInputsDirty = true; });
 	Settings->WatchProperty(Settings->UVLayer, [this](FString) { bInputsDirty = true; });
 	Settings->WatchProperty(Settings->bUseWorldSpace, [this](bool) { bDetailMeshValid = false; bInputsDirty = true; });
 	Settings->WatchProperty(Settings->Thickness, [this](float) { bInputsDirty = true; });
@@ -486,10 +487,8 @@ void UBakeMeshAttributeMapsTool::Shutdown(EToolShutdownType ShutdownType)
 		const int NumResults = Settings->Result.Num();
 		for (int ResultIdx = 0; ResultIdx < NumResults; ++ResultIdx)
 		{
-			const FTexture2DBuilder::ETextureType TexType = GetTextureType(ResultTypes[ResultIdx]);
 			FString TexName;
 			GetTextureName(ResultTypes[ResultIdx], BaseName, TexName);
-			FTexture2DBuilder::CopyPlatformDataToSourceData(Settings->Result[ResultIdx], TexType);
 			bCreatedAssetOK = bCreatedAssetOK && UE::Modeling::CreateTextureObject(GetToolManager(), FCreateTextureObjectParams{ 0, StaticMeshAsset->GetWorld(), StaticMeshAsset, TexName, Settings->Result[ResultIdx] }).IsOK();
 		}
 		ensure(bCreatedAssetOK);
@@ -585,6 +584,7 @@ void UBakeMeshAttributeMapsTool::UpdateResult()
 
 	FBakeCacheSettings BakeCacheSettings;
 	BakeCacheSettings.Dimensions = Dimensions;
+	BakeCacheSettings.SourceFormat = Settings->SourceFormat;
 	BakeCacheSettings.UVLayer = FCString::Atoi(*Settings->UVLayer);
 	BakeCacheSettings.DetailTimestamp = this->DetailMeshTimestamp;
 	BakeCacheSettings.Thickness = Settings->Thickness;
@@ -654,7 +654,7 @@ void UBakeMeshAttributeMapsTool::UpdateResult()
 	{
 		Compute = MakeUnique<TGenericDataBackgroundCompute<FMeshMapBaker>>();
 		Compute->Setup(this);
-		Compute->OnResultUpdated.AddLambda([this](const TUniquePtr<FMeshMapBaker>& NewResult) { OnMapsUpdated(NewResult); });
+		Compute->OnResultUpdated.AddLambda([this](const TUniquePtr<FMeshMapBaker>& NewResult) { OnMapsUpdated(NewResult, CachedBakeCacheSettings.SourceFormat); });
 		Compute->InvalidateResult();
 	}
 	else if (bInvalidate)
diff --git a/Engine/Plugins/Experimental/MeshModelingToolsetExp/Source/MeshModelingToolsExp/Private/BakeMeshAttributeMapsToolBase.cpp b/Engine/Plugins/Experimental/MeshModelingToolsetExp/Source/MeshModelingToolsExp/Private/BakeMeshAttributeMapsToolBase.cpp
--- a/Engine/Plugins/Experimental/MeshModelingToolsetExp/Source/MeshModelingToolsExp/Private/BakeMeshAttributeMapsToolBase.cpp
+++ b/Engine/Plugins/Experimental/MeshModelingToolsetExp/Source/MeshModelingToolsExp/Private/BakeMeshAttributeMapsToolBase.cpp
@@ -223,7 +223,7 @@ void UBakeMeshAttributeMapsToolBase::UpdatePreview(const int PreviewIdx)
 }
 
 
-void UBakeMeshAttributeMapsToolBase::OnMapsUpdated(const TUniquePtr<UE::Geometry::FMeshMapBaker>& NewResult)
+void UBakeMeshAttributeMapsToolBase::OnMapsUpdated(const TUniquePtr<UE::Geometry::FMeshMapBaker>& NewResult, const EBakeTextureFormat Format)
 {
 	FImageDimensions BakeDimensions = NewResult->GetDimensions();
 	const int32 NumEval = NewResult->NumEvaluators();
@@ -231,20 +231,31 @@ void UBakeMeshAttributeMapsToolBase::OnMapsUpdated(const TUniquePtr<UE::Geometry
 	{
 		FMeshMapEvaluator* Eval = NewResult->GetEvaluator(EvalIdx);
 
-		auto UpdateCachedMap = [this, &NewResult, &EvalIdx, &BakeDimensions](const EBakeMapType BakeMapType, const FTexture2DBuilder::ETextureType TexType, const int32 ResultIdx) -> void
+		auto UpdateCachedMap = [this, &NewResult, &EvalIdx, &BakeDimensions](const EBakeMapType MapType, const EBakeTextureFormat MapFormat, const int32 ResultIdx) -> void
 		{
 			// For 8-bit color textures, ensure that the source data is in sRGB.
+			const FTexture2DBuilder::ETextureType TexType = GetTextureType(MapType, MapFormat);
 			const bool bConvertToSRGB = TexType == FTexture2DBuilder::ETextureType::Color;
+			const ETextureSourceFormat SourceDataFormat = MapFormat == EBakeTextureFormat::ChannelBits16 ? TSF_RGBA16F : TSF_BGRA8;
+			
 			FTexture2DBuilder TextureBuilder;
 			TextureBuilder.Initialize(TexType, BakeDimensions);
 			TextureBuilder.Copy(*NewResult->GetBakeResults(EvalIdx)[ResultIdx], bConvertToSRGB);
 			TextureBuilder.Commit(false);
 
+			// Copy image to source data after commit. This will avoid incurring
+			// the cost of hitting the DDC for texture compile while iterating on
+			// bake settings. Since this dirties the texture, the next time the texture
+			// is used after accepting the final texture, the DDC will trigger and
+			// properly recompile the platform data.
+			const bool bConvertSourceToSRGB = bConvertToSRGB && SourceDataFormat == TSF_BGRA8;
+			TextureBuilder.CopyImageToSourceData(*NewResult->GetBakeResults(EvalIdx)[ResultIdx], SourceDataFormat, bConvertSourceToSRGB);
+
 			// The CachedMap & CachedMapIndices can be thrown out of sync if updated during
 			// a background compute. Validate the computed type against our cached maps.
-			if (CachedMapIndices.Contains(BakeMapType))
+			if (CachedMapIndices.Contains(MapType))
 			{
-				CachedMaps[CachedMapIndices[BakeMapType]] = TextureBuilder.GetTexture2D();
+				CachedMaps[CachedMapIndices[MapType]] = TextureBuilder.GetTexture2D();
 			}
 		};
 
@@ -253,7 +264,7 @@ void UBakeMeshAttributeMapsToolBase::OnMapsUpdated(const TUniquePtr<UE::Geometry
 		case EMeshMapEvaluatorType::Normal:
 		{
 			constexpr EBakeMapType MapType = EBakeMapType::TangentSpaceNormalMap;
-			UpdateCachedMap(MapType, GetTextureType(MapType), 0);
+			UpdateCachedMap(MapType, Format, 0);
 			break;
 		}
 		case EMeshMapEvaluatorType::Occlusion:
@@ -264,19 +275,19 @@ void UBakeMeshAttributeMapsToolBase::OnMapsUpdated(const TUniquePtr<UE::Geometry
 			if ((bool)(OcclusionEval->OcclusionType & EMeshOcclusionMapType::AmbientOcclusion))
 			{
 				constexpr EBakeMapType MapType = EBakeMapType::AmbientOcclusion;
-				UpdateCachedMap(MapType, GetTextureType(MapType), OcclusionIdx++);
+				UpdateCachedMap(MapType, Format, OcclusionIdx++);
 			}
 			if ((bool)(OcclusionEval->OcclusionType & EMeshOcclusionMapType::BentNormal))
 			{
 				constexpr EBakeMapType MapType = EBakeMapType::BentNormal;
-				UpdateCachedMap(MapType, GetTextureType(MapType), OcclusionIdx++);
+				UpdateCachedMap(MapType, Format, OcclusionIdx++);
 			}
 			break;
 		}
 		case EMeshMapEvaluatorType::Curvature:
 		{
 			constexpr EBakeMapType MapType = EBakeMapType::Curvature;
-			UpdateCachedMap(MapType, GetTextureType(MapType), 0);
+			UpdateCachedMap(MapType, Format, 0);
 			break;
 		}
 		case EMeshMapEvaluatorType::Property:
@@ -305,19 +316,19 @@ void UBakeMeshAttributeMapsToolBase::OnMapsUpdated(const TUniquePtr<UE::Geometry
 				break;
 			}
 
-			UpdateCachedMap(MapType, GetTextureType(MapType), 0);
+			UpdateCachedMap(MapType, Format, 0);
 			break;
 		}
 		case EMeshMapEvaluatorType::ResampleImage:
 		{
 			constexpr EBakeMapType MapType = EBakeMapType::Texture2DImage;
-			UpdateCachedMap(EBakeMapType::Texture2DImage, GetTextureType(MapType), 0);
+			UpdateCachedMap(EBakeMapType::Texture2DImage, Format, 0);
 			break;
 		}
 		case EMeshMapEvaluatorType::MultiResampleImage:
 		{
 			constexpr EBakeMapType MapType = EBakeMapType::MultiTexture;
-			UpdateCachedMap(MapType, GetTextureType(MapType), 0);
+			UpdateCachedMap(MapType, Format, 0);
 			break;
 		}
 		default:
@@ -357,7 +368,7 @@ TArray<EBakeMapType> UBakeMeshAttributeMapsToolBase::GetMapTypesArray(const int3
 }
 
 
-FTexture2DBuilder::ETextureType UBakeMeshAttributeMapsToolBase::GetTextureType(EBakeMapType MapType)
+FTexture2DBuilder::ETextureType UBakeMeshAttributeMapsToolBase::GetTextureType(const EBakeMapType MapType, const EBakeTextureFormat MapFormat)
 {
 	FTexture2DBuilder::ETextureType TexType = FTexture2DBuilder::ETextureType::Color;
 	switch (MapType)
@@ -382,8 +393,14 @@ FTexture2DBuilder::ETextureType UBakeMeshAttributeMapsToolBase::GetTextureType(E
 		break;
 	case EBakeMapType::MaterialID:
 	case EBakeMapType::VertexColorImage:
+		break;
 	case EBakeMapType::Texture2DImage:
 	case EBakeMapType::MultiTexture:
+		// For texture output with 16-bit source data, output HDR texture
+		if (MapFormat == EBakeTextureFormat::ChannelBits16)
+		{
+			TexType = FTexture2DBuilder::ETextureType::EmissiveHDR;
+		}
 		break;
 	}
 	return TexType;
diff --git a/Engine/Plugins/Experimental/MeshModelingToolsetExp/Source/MeshModelingToolsExp/Private/BakeMultiMeshAttributeMapsTool.cpp b/Engine/Plugins/Experimental/MeshModelingToolsetExp/Source/MeshModelingToolsExp/Private/BakeMultiMeshAttributeMapsTool.cpp
--- a/Engine/Plugins/Experimental/MeshModelingToolsetExp/Source/MeshModelingToolsExp/Private/BakeMultiMeshAttributeMapsTool.cpp
+++ b/Engine/Plugins/Experimental/MeshModelingToolsetExp/Source/MeshModelingToolsExp/Private/BakeMultiMeshAttributeMapsTool.cpp
@@ -433,6 +433,7 @@ void UBakeMultiMeshAttributeMapsTool::Setup()
 	Settings->WatchProperty(Settings->MapTypes, [this](int32) { bInputsDirty = true; UpdateOnModeChange(); });
 	Settings->WatchProperty(Settings->MapPreview, [this](int32) { UpdateVisualization(); GetToolManager()->PostInvalidation(); });
 	Settings->WatchProperty(Settings->Resolution, [this](EBakeTextureResolution) { bInputsDirty = true; });
+	Settings->WatchProperty(Settings->SourceFormat, [this](EBakeTextureFormat) { bInputsDirty = true; });
 	Settings->WatchProperty(Settings->UVLayer, [this](FString) { bInputsDirty = true; });
 	Settings->WatchProperty(Settings->Thickness, [this](float) { bInputsDirty = true; });
 	Settings->WatchProperty(Settings->Multisampling, [this](EBakeMultisampling) { bInputsDirty = true; });
@@ -560,11 +561,8 @@ void UBakeMultiMeshAttributeMapsTool::Shutdown(EToolShutdownType ShutdownType)
 		const int NumResults = Settings->Result.Num();
 		for (int ResultIdx = 0; ResultIdx < NumResults; ++ResultIdx)
 		{
-			const FTexture2DBuilder::ETextureType TexType = GetTextureType(ResultTypes[ResultIdx]);
 			FString TexName;
 			GetTextureName(ResultTypes[ResultIdx], BaseName, TexName);
-
-			FTexture2DBuilder::CopyPlatformDataToSourceData(Settings->Result[ResultIdx], TexType);
 			bCreatedAssetOK = bCreatedAssetOK && UE::Modeling::CreateTextureObject(GetToolManager(), FCreateTextureObjectParams{ 0, StaticMeshAsset->GetWorld(), StaticMeshAsset, TexName, Settings->Result[ResultIdx] }).IsOK();
 		}
 		ensure(bCreatedAssetOK);
@@ -589,6 +587,7 @@ void UBakeMultiMeshAttributeMapsTool::UpdateResult()
 
 	FBakeCacheSettings BakeCacheSettings;
 	BakeCacheSettings.Dimensions = Dimensions;
+	BakeCacheSettings.SourceFormat = Settings->SourceFormat;
 	BakeCacheSettings.UVLayer = FCString::Atoi(*Settings->UVLayer);
 	BakeCacheSettings.Thickness = Settings->Thickness;
 	BakeCacheSettings.Multisampling = (int32)Settings->Multisampling;
@@ -622,7 +621,7 @@ void UBakeMultiMeshAttributeMapsTool::UpdateResult()
 	{
 		Compute = MakeUnique<TGenericDataBackgroundCompute<FMeshMapBaker>>();
 		Compute->Setup(this);
-		Compute->OnResultUpdated.AddLambda([this](const TUniquePtr<FMeshMapBaker>& NewResult) { OnMapsUpdated(NewResult); });
+		Compute->OnResultUpdated.AddLambda([this](const TUniquePtr<FMeshMapBaker>& NewResult) { OnMapsUpdated(NewResult, CachedBakeCacheSettings.SourceFormat); });
 		Compute->InvalidateResult();
 	}
 	else if (bInvalidate)
diff --git a/Engine/Plugins/Experimental/MeshModelingToolsetExp/Source/MeshModelingToolsExp/Public/BakeMeshAttributeMapsTool.h b/Engine/Plugins/Experimental/MeshModelingToolsetExp/Source/MeshModelingToolsExp/Public/BakeMeshAttributeMapsTool.h
--- a/Engine/Plugins/Experimental/MeshModelingToolsetExp/Source/MeshModelingToolsExp/Public/BakeMeshAttributeMapsTool.h
+++ b/Engine/Plugins/Experimental/MeshModelingToolsetExp/Source/MeshModelingToolsExp/Public/BakeMeshAttributeMapsTool.h
@@ -64,6 +64,10 @@ public:
 	UPROPERTY(EditAnywhere, Category = MapSettings, meta = (TransientToolProperty))
 	EBakeTextureResolution Resolution = EBakeTextureResolution::Resolution256;
 
+	/** The channel bit depth of the source data for the generated textures */
+	UPROPERTY(EditAnywhere, Category = MapSettings, meta = (TransientToolProperty))
+	EBakeTextureFormat SourceFormat = EBakeTextureFormat::ChannelBits8;
+
 	/** The multisampling configuration per texel */
 	UPROPERTY(EditAnywhere, Category = MapSettings)
 	EBakeMultisampling Multisampling = EBakeMultisampling::None;
diff --git a/Engine/Plugins/Experimental/MeshModelingToolsetExp/Source/MeshModelingToolsExp/Public/BakeMeshAttributeMapsToolBase.h b/Engine/Plugins/Experimental/MeshModelingToolsetExp/Source/MeshModelingToolsExp/Public/BakeMeshAttributeMapsToolBase.h
--- a/Engine/Plugins/Experimental/MeshModelingToolsetExp/Source/MeshModelingToolsExp/Public/BakeMeshAttributeMapsToolBase.h
+++ b/Engine/Plugins/Experimental/MeshModelingToolsetExp/Source/MeshModelingToolsExp/Public/BakeMeshAttributeMapsToolBase.h
@@ -92,6 +92,14 @@ enum class EBakeTextureResolution
 };
 
 
+UENUM()
+enum class EBakeTextureFormat
+{
+	ChannelBits8 UMETA(DisplayName = "8 bits/channel"),
+	ChannelBits16 UMETA(DisplayName = "16 bits/channel")
+};
+
+
 UENUM()
 enum class EBakeMultisampling
 {
@@ -190,6 +198,7 @@ protected:
 	{
 		EBakeMapType BakeMapTypes = EBakeMapType::None;
 		FImageDimensions Dimensions;
+		EBakeTextureFormat SourceFormat = EBakeTextureFormat::ChannelBits8;
 		int32 UVLayer = 0;
 		int32 DetailTimestamp = 0;
 		float Thickness = 3.0;
@@ -199,7 +208,8 @@ protected:
 		{
 			return BakeMapTypes == Other.BakeMapTypes && Dimensions == Other.Dimensions &&
 				UVLayer == Other.UVLayer && DetailTimestamp == Other.DetailTimestamp &&
-				Thickness == Other.Thickness && Multisampling == Other.Multisampling;
+				Thickness == Other.Thickness && Multisampling == Other.Multisampling &&
+				SourceFormat == Other.SourceFormat;
 		}
 	};
 	FBakeCacheSettings CachedBakeCacheSettings;
@@ -221,7 +231,7 @@ protected:
 	 * 
 	 * @param NewResult the resulting FMeshMapBaker from the background Compute
 	 */
-	void OnMapsUpdated(const TUniquePtr<UE::Geometry::FMeshMapBaker>& NewResult);
+	void OnMapsUpdated(const TUniquePtr<UE::Geometry::FMeshMapBaker>& NewResult, EBakeTextureFormat Format);
 
 	/**
 	 * Update the preview material parameters for a given a result index.
@@ -251,7 +261,7 @@ protected:
 	static TArray<EBakeMapType> GetMapTypesArray(const int32& MapTypes);
 
 	/** @return the Texture2D type for a given map type */
-	static UE::Geometry::FTexture2DBuilder::ETextureType GetTextureType(EBakeMapType MapType);
+	static UE::Geometry::FTexture2DBuilder::ETextureType GetTextureType(EBakeMapType MapType, EBakeTextureFormat MapFormat);
 
 	/** @return the texture name given a base name and map type */
 	static void GetTextureName(EBakeMapType MapType, const FString& BaseName, FString& TexName);
diff --git a/Engine/Plugins/Experimental/MeshModelingToolsetExp/Source/MeshModelingToolsExp/Public/BakeMultiMeshAttributeMapsTool.h b/Engine/Plugins/Experimental/MeshModelingToolsetExp/Source/MeshModelingToolsExp/Public/BakeMultiMeshAttributeMapsTool.h
--- a/Engine/Plugins/Experimental/MeshModelingToolsetExp/Source/MeshModelingToolsExp/Public/BakeMultiMeshAttributeMapsTool.h
+++ b/Engine/Plugins/Experimental/MeshModelingToolsetExp/Source/MeshModelingToolsExp/Public/BakeMultiMeshAttributeMapsTool.h
@@ -74,6 +74,10 @@ public:
 	UPROPERTY(EditAnywhere, Category = MapSettings, meta = (TransientToolProperty))
 	EBakeTextureResolution Resolution = EBakeTextureResolution::Resolution256;
 
+	/** The channel bit depth of the source data for the generated textures */
+	UPROPERTY(EditAnywhere, Category = MapSettings, meta = (TransientToolProperty))
+	EBakeTextureFormat SourceFormat = EBakeTextureFormat::ChannelBits8;
+
 	/** The multisampling configuration per texel */
 	UPROPERTY(EditAnywhere, Category = MapSettings)
 	EBakeMultisampling Multisampling = EBakeMultisampling::None;
diff --git a/Engine/Plugins/Experimental/MeshModelingToolsetExp/Source/ModelingComponents/Private/AssetUtils/Texture2DBuilder.cpp b/Engine/Plugins/Experimental/MeshModelingToolsetExp/Source/ModelingComponents/Private/AssetUtils/Texture2DBuilder.cpp
--- a/Engine/Plugins/Experimental/MeshModelingToolsetExp/Source/ModelingComponents/Private/AssetUtils/Texture2DBuilder.cpp
+++ b/Engine/Plugins/Experimental/MeshModelingToolsetExp/Source/ModelingComponents/Private/AssetUtils/Texture2DBuilder.cpp
@@ -349,6 +349,66 @@ bool FTexture2DBuilder::Copy(const TImageBuilder<FVector4f>& SourceImage, const
 }
 
 
+bool FTexture2DBuilder::CopyImageToSourceData(const TImageBuilder<FVector4f>& SourceImage, const ETextureSourceFormat SourceDataFormat, const bool bConvertToSRGB)
+{
+	// source data only exists in Editor
+#if WITH_EDITOR
+	if(!ensure(RawTexture2D))
+	{
+		return false;
+	}
+
+	if (ensure(SourceImage.GetDimensions() == Dimensions) == false)
+	{
+		return false;
+	}
+
+	// Currently only support these source formats.
+	if (!ensure(SourceDataFormat == TSF_BGRA8 || SourceDataFormat == TSF_RGBA16F))
+	{
+		return false;
+	}
+
+	if (SourceDataFormat == TSF_RGBA16F && bConvertToSRGB)
+	{
+		ensure(false);		// not currently supported
+	}
+
+	RawTexture2D->Source.Init2DWithMipChain(Dimensions.GetWidth(), Dimensions.GetHeight(), SourceDataFormat);
+
+	const int64 Num = Dimensions.Num();
+	if (SourceDataFormat == TSF_BGRA8)
+	{
+		FColor* DestData = reinterpret_cast<FColor*>(RawTexture2D->Source.LockMip(0));
+		for (int32 i = 0; i < Num; ++i)
+		{
+			FVector4f Pixel = SourceImage.GetPixel(i);
+			Pixel.X = FMathf::Clamp(Pixel.X, 0.0, 1.0);
+			Pixel.Y = FMathf::Clamp(Pixel.Y, 0.0, 1.0);
+			Pixel.Z = FMathf::Clamp(Pixel.Z, 0.0, 1.0);
+			Pixel.W = FMathf::Clamp(Pixel.W, 0.0, 1.0);
+			const FColor Texel = ToLinearColor(Pixel).ToFColor(bConvertToSRGB);
+			DestData[i] = Texel;
+		}
+	}
+	else if (SourceDataFormat == TSF_RGBA16F)
+	{
+		FFloat16Color* DestData = reinterpret_cast<FFloat16Color*>(RawTexture2D->Source.LockMip(0));
+		for (int32 i = 0; i < Num; ++i)
+		{
+			FVector4f Pixel = SourceImage.GetPixel(i);
+			const FFloat16Color Texel(ToLinearColor(Pixel));
+			DestData[i] = Texel;
+		}
+	}
+
+	RawTexture2D->Source.UnlockMip(0);
+#endif
+	return false;
+}
+
+
+
 bool FTexture2DBuilder::CopyTo(TImageBuilder<FVector4f>& DestImage) const
 {
 	if (ensure(DestImage.GetDimensions() == Dimensions) == false)
diff --git a/Engine/Plugins/Experimental/MeshModelingToolsetExp/Source/ModelingComponents/Public/AssetUtils/Texture2DBuilder.h b/Engine/Plugins/Experimental/MeshModelingToolsetExp/Source/ModelingComponents/Public/AssetUtils/Texture2DBuilder.h
--- a/Engine/Plugins/Experimental/MeshModelingToolsetExp/Source/ModelingComponents/Public/AssetUtils/Texture2DBuilder.h
+++ b/Engine/Plugins/Experimental/MeshModelingToolsetExp/Source/ModelingComponents/Public/AssetUtils/Texture2DBuilder.h
@@ -278,6 +278,14 @@ public:
 	 */
 	bool Copy(const TImageBuilder<FVector4f>& SourceImage, const bool bConvertToSRGB = false);
 
+	/**
+	 * Populate Source Data texel values from a floating point SourceImage
+	 * @param SourceImage the image to copy to source data
+	 * @param SourceDataFormat the texture format for the Source Data
+	 * @param bConvertToSRGB if true, SourceImage is assumed to be Linear, and converted to SRGB for the Texture
+	 */
+	bool CopyImageToSourceData(const TImageBuilder<FVector4f>& SourceImage, const ETextureSourceFormat SourceDataFormat, const bool bConvertToSRGB = false);
+
 
 	/**
 	 * copy existing texel values to floating-point DestImage. Assumed to be same color space (eg values are copied without any linear/SRGB conversion)
