ModelingTools: Add SkeletalMesh and DynamicMesh support to BakeTexture

#rnx

diff --git a/Engine/Plugins/Editor/ModelingToolsEditorMode/Source/ModelingToolsEditorMode/Private/ModelingToolsEditorMode.cpp b/Engine/Plugins/Editor/ModelingToolsEditorMode/Source/ModelingToolsEditorMode/Private/ModelingToolsEditorMode.cpp
--- a/Engine/Plugins/Editor/ModelingToolsEditorMode/Source/ModelingToolsEditorMode/Private/ModelingToolsEditorMode.cpp
+++ b/Engine/Plugins/Editor/ModelingToolsEditorMode/Source/ModelingToolsEditorMode/Private/ModelingToolsEditorMode.cpp
@@ -11,6 +11,7 @@
 #include "ToolTargets/VolumeDynamicMeshToolTarget.h"
 #include "ToolTargets/VolumeMeshDescriptionToolTarget.h"
 #include "ToolTargets/DynamicMeshComponentToolTarget.h"
+#include "ToolTargets/SkeletalMeshComponentToolTarget.h"
 #include "Framework/Commands/UICommandList.h"
 #include "Framework/Application/SlateApplication.h"
 #include "Framework/MultiBox/MultiBoxBuilder.h"
@@ -317,6 +318,7 @@ void UModelingToolsEditorMode::Enter()
 	GetInteractiveToolsContext()->TargetManager->AddTargetFactory(NewObject<UVolumeDynamicMeshToolTargetFactory>(GetToolManager()));
 	GetInteractiveToolsContext()->TargetManager->AddTargetFactory(NewObject<UVolumeMeshDescriptionToolTargetFactory>(GetToolManager()));
 	GetInteractiveToolsContext()->TargetManager->AddTargetFactory(NewObject<UDynamicMeshComponentToolTargetFactory>(GetToolManager()));
+	GetInteractiveToolsContext()->TargetManager->AddTargetFactory(NewObject<USkeletalMeshComponentToolTargetFactory>(GetToolManager()));
 
 	// register stylus event handler
 	StylusStateTracker = MakeUnique<FStylusStateTracker>();
diff --git a/Engine/Plugins/Experimental/MeshModelingToolsetExp/Source/MeshModelingToolsExp/Private/BakeMeshAttributeMapsTool.cpp b/Engine/Plugins/Experimental/MeshModelingToolsetExp/Source/MeshModelingToolsExp/Private/BakeMeshAttributeMapsTool.cpp
--- a/Engine/Plugins/Experimental/MeshModelingToolsetExp/Source/MeshModelingToolsExp/Private/BakeMeshAttributeMapsTool.cpp
+++ b/Engine/Plugins/Experimental/MeshModelingToolsetExp/Source/MeshModelingToolsExp/Private/BakeMeshAttributeMapsTool.cpp
@@ -16,7 +16,6 @@
 
 #include "ImageUtils.h"
 
-#include "AssetUtils/Texture2DBuilder.h"
 #include "AssetUtils/Texture2DUtil.h"
 #include "ModelingObjectsCreationAPI.h"
 
@@ -24,12 +23,16 @@
 #include "TargetInterfaces/MeshDescriptionProvider.h"
 #include "TargetInterfaces/PrimitiveComponentBackedTarget.h"
 #include "TargetInterfaces/StaticMeshBackedTarget.h"
+#include "TargetInterfaces/SkeletalMeshBackedTarget.h"
 #include "ToolTargetManager.h"
 #include "ModelingToolTargetUtil.h"
 
 // required to pass UStaticMesh asset so we can save at same location
-#include "Engine/Classes/Components/StaticMeshComponent.h"
-#include "Engine/Classes/Engine/StaticMesh.h"
+#include "Engine/StaticMesh.h"
+#include "Engine/SkeletalMesh.h"
+#include "Components/StaticMeshComponent.h"
+#include "Components/SkeletalMeshComponent.h"
+#include "Components/DynamicMeshComponent.h"
 
 #include "ExplicitUseGeometryMathTypes.h"		// using UE::Geometry::(math types)
 using namespace UE::Geometry;
@@ -46,7 +49,6 @@ const FToolTargetTypeRequirements& UBakeMeshAttributeMapsToolBuilder::GetTargetR
 	static FToolTargetTypeRequirements TypeRequirements({
 		UMeshDescriptionProvider::StaticClass(),
 		UPrimitiveComponentBackedTarget::StaticClass(),
-		UStaticMeshBackedTarget::StaticClass(),			// currently only supports StaticMesh targets
 		UMaterialProvider::StaticClass()
 		});
 	return TypeRequirements;
@@ -55,7 +57,20 @@ const FToolTargetTypeRequirements& UBakeMeshAttributeMapsToolBuilder::GetTargetR
 bool UBakeMeshAttributeMapsToolBuilder::CanBuildTool(const FToolBuilderState& SceneState) const
 {
 	const int32 NumTargets = SceneState.TargetManager->CountSelectedAndTargetable(SceneState, GetTargetRequirements());
-	return (NumTargets == 1 || NumTargets == 2);
+	if (NumTargets == 1 || NumTargets == 2)
+	{
+		bool bValidTargets = true;
+		SceneState.TargetManager->EnumerateSelectedAndTargetableComponents(SceneState, GetTargetRequirements(),
+			[&bValidTargets](UActorComponent* Component)
+			{
+				UStaticMeshComponent* StaticMesh = Cast<UStaticMeshComponent>(Component);
+				USkeletalMeshComponent* SkeletalMesh = Cast<USkeletalMeshComponent>(Component);
+				UDynamicMeshComponent* DynMesh = Cast<UDynamicMeshComponent>(Component);
+				bValidTargets = bValidTargets && (StaticMesh || SkeletalMesh || DynMesh);
+			});
+		return bValidTargets;
+	}
+	return false;
 }
 
 UInteractiveTool* UBakeMeshAttributeMapsToolBuilder::BuildTool(const FToolBuilderState& SceneState) const
@@ -303,11 +318,14 @@ void UBakeMeshAttributeMapsTool::Setup()
 	UToolTarget* DetailTarget = Targets[bIsBakeToSelf ? 0 : 1];
 	IStaticMeshBackedTarget* DetailStaticMeshTarget = Cast<IStaticMeshBackedTarget>(DetailTarget);
 	UStaticMesh* DetailStaticMesh = DetailStaticMeshTarget ? DetailStaticMeshTarget->GetStaticMesh() : nullptr;
+	ISkeletalMeshBackedTarget* DetailSkeletalMeshTarget = Cast<ISkeletalMeshBackedTarget>(DetailTarget);
+	USkeletalMesh* DetailSkeletalMesh = DetailSkeletalMeshTarget ? DetailSkeletalMeshTarget->GetSkeletalMesh() : nullptr;
 
 	DetailMeshProps = NewObject<UDetailMeshToolProperties>(this);
 	AddToolPropertySource(DetailMeshProps);
 	SetToolPropertySourceEnabled(DetailMeshProps, true);
-	DetailMeshProps->DetailMesh = DetailStaticMesh;
+	DetailMeshProps->DetailStaticMesh = DetailStaticMesh;
+	DetailMeshProps->DetailSkeletalMesh = DetailSkeletalMesh;
 	DetailMeshProps->DetailMeshNormalMap = nullptr;
 	DetailMeshProps->WatchProperty(DetailMeshProps->DetailNormalUVLayer, [this](int) { bInputsDirty = true; });
 	DetailMeshProps->WatchProperty(DetailMeshProps->DetailMeshNormalMap, [this](UTexture2D*)
@@ -478,10 +496,17 @@ void UBakeMeshAttributeMapsTool::Shutdown(EToolShutdownType ShutdownType)
 	}
 	if (ShutdownType == EToolShutdownType::Accept)
 	{
-		UStaticMeshComponent* StaticMeshComponent = CastChecked<UStaticMeshComponent>(UE::ToolTarget::GetTargetComponent(Targets[0]));
-		UStaticMesh* StaticMeshAsset = StaticMeshComponent->GetStaticMesh();
-		check(StaticMeshAsset);
-		FString BaseName = UE::ToolTarget::GetTargetActor(Targets[0])->GetName();
+		// Check if we have a source asset to identify a location to store the texture assets.
+		IStaticMeshBackedTarget* StaticMeshTarget = Cast<IStaticMeshBackedTarget>(Targets[0]);
+		UObject* SourceAsset = StaticMeshTarget ? StaticMeshTarget->GetStaticMesh() : nullptr;
+		if (!SourceAsset)
+		{
+			// Check if our target is a Skeletal Mesh Asset
+			ISkeletalMeshBackedTarget* SkeletalMeshTarget = Cast<ISkeletalMeshBackedTarget>(Targets[0]);
+			SourceAsset = SkeletalMeshTarget ? SkeletalMeshTarget->GetSkeletalMesh() : nullptr;
+		}
+		const UPrimitiveComponent* SourceComponent = UE::ToolTarget::GetTargetComponent(Targets[0]);
+		const FString BaseName = UE::ToolTarget::GetTargetActor(Targets[0])->GetName();
 
 		bool bCreatedAssetOK = true;
 		const int NumResults = Settings->Result.Num();
@@ -489,7 +514,7 @@ void UBakeMeshAttributeMapsTool::Shutdown(EToolShutdownType ShutdownType)
 		{
 			FString TexName;
 			GetTextureName(ResultTypes[ResultIdx], BaseName, TexName);
-			bCreatedAssetOK = bCreatedAssetOK && UE::Modeling::CreateTextureObject(GetToolManager(), FCreateTextureObjectParams{ 0, StaticMeshAsset->GetWorld(), StaticMeshAsset, TexName, Settings->Result[ResultIdx] }).IsOK();
+			bCreatedAssetOK = bCreatedAssetOK && UE::Modeling::CreateTextureObject(GetToolManager(), FCreateTextureObjectParams{ 0, SourceComponent->GetWorld(), SourceAsset, TexName, Settings->Result[ResultIdx] }).IsOK();
 		}
 		ensure(bCreatedAssetOK);
 	}
diff --git a/Engine/Plugins/Experimental/MeshModelingToolsetExp/Source/MeshModelingToolsExp/Private/BakeMultiMeshAttributeMapsTool.cpp b/Engine/Plugins/Experimental/MeshModelingToolsetExp/Source/MeshModelingToolsExp/Private/BakeMultiMeshAttributeMapsTool.cpp
--- a/Engine/Plugins/Experimental/MeshModelingToolsetExp/Source/MeshModelingToolsExp/Private/BakeMultiMeshAttributeMapsTool.cpp
+++ b/Engine/Plugins/Experimental/MeshModelingToolsetExp/Source/MeshModelingToolsExp/Private/BakeMultiMeshAttributeMapsTool.cpp
@@ -45,7 +45,7 @@ const FToolTargetTypeRequirements& UBakeMultiMeshAttributeMapsToolBuilder::GetTa
 	static FToolTargetTypeRequirements TypeRequirements({
 		UMeshDescriptionProvider::StaticClass(),
 		UPrimitiveComponentBackedTarget::StaticClass(),
-		UStaticMeshBackedTarget::StaticClass(),			// currently only supports StaticMesh targets
+		UStaticMeshBackedTarget::StaticClass(),			// FMeshSceneAdapter currently only supports StaticMesh targets
 		UMaterialProvider::StaticClass()
 		});
 	return TypeRequirements;
diff --git a/Engine/Plugins/Experimental/MeshModelingToolsetExp/Source/MeshModelingToolsExp/Public/BakeMeshAttributeToolCommon.h b/Engine/Plugins/Experimental/MeshModelingToolsetExp/Source/MeshModelingToolsExp/Public/BakeMeshAttributeToolCommon.h
--- a/Engine/Plugins/Experimental/MeshModelingToolsetExp/Source/MeshModelingToolsExp/Public/BakeMeshAttributeToolCommon.h
+++ b/Engine/Plugins/Experimental/MeshModelingToolsetExp/Source/MeshModelingToolsExp/Public/BakeMeshAttributeToolCommon.h
@@ -3,7 +3,7 @@
 #pragma once
 
 #include "CoreMinimal.h"
-#include "Engine/Classes/Engine/Texture2D.h"
+#include "Engine/Texture2D.h"
 #include "MultiSelectionTool.h"
 #include "Image/ImageBuilder.h"
 #include "Image/ImageDimensions.h"
@@ -12,6 +12,8 @@
 // Pre-declarations
 using UE::Geometry::FImageDimensions;
 using UE::Geometry::TImageBuilder;
+class UStaticMesh;
+class USkeletalMesh;
 
 
 /**
@@ -25,8 +27,11 @@ class MESHMODELINGTOOLSEXP_API UDetailMeshToolProperties : public UInteractiveTo
 	GENERATED_BODY()
 public:
 	/** The detail mesh to sample */
-	UPROPERTY(VisibleAnywhere, Category = DetailMesh, meta = (TransientToolProperty))
-	TObjectPtr<UStaticMesh> DetailMesh = nullptr;
+	UPROPERTY(VisibleAnywhere, Category = DetailMesh, DisplayName = "Detail Mesh", meta = (TransientToolProperty, EditCondition = "DetailStaticMesh != nullptr", EditConditionHides))
+	TObjectPtr<UStaticMesh> DetailStaticMesh = nullptr;
+
+	UPROPERTY(VisibleAnywhere, Category = DetailMesh, DisplayName = "Detail Mesh", meta = (TransientToolProperty, EditCondition = "DetailSkeletalMesh != nullptr", EditConditionHides))
+	TObjectPtr<USkeletalMesh> DetailSkeletalMesh = nullptr;
 
 	/** The detail mesh normal map to sample. If empty, the geometric normals will be used. */
 	UPROPERTY(EditAnywhere, Category = DetailMesh, meta = (TransientToolProperty))
