ModelingTools: Removed Occlusion Distribution and Bent Normal Space properties.

#rnx

diff --git a/Engine/Plugins/Experimental/MeshModelingToolsetExp/Source/MeshModelingToolsExp/Private/BakeMeshAttributeMapsTool.cpp b/Engine/Plugins/Experimental/MeshModelingToolsetExp/Source/MeshModelingToolsExp/Private/BakeMeshAttributeMapsTool.cpp
--- a/Engine/Plugins/Experimental/MeshModelingToolsetExp/Source/MeshModelingToolsExp/Private/BakeMeshAttributeMapsTool.cpp
+++ b/Engine/Plugins/Experimental/MeshModelingToolsetExp/Source/MeshModelingToolsExp/Private/BakeMeshAttributeMapsTool.cpp
@@ -173,25 +173,6 @@ public:
 				OcclusionEval->SpreadAngle = OcclusionSettings.SpreadAngle;
 				OcclusionEval->BiasAngleDeg = OcclusionSettings.BiasAngle;
 
-				switch (OcclusionSettings.Distribution)
-				{
-				case EOcclusionMapDistribution::Cosine:
-					OcclusionEval->Distribution = FMeshOcclusionMapEvaluator::EDistribution::Cosine;
-					break;
-				case EOcclusionMapDistribution::Uniform:
-					OcclusionEval->Distribution = FMeshOcclusionMapEvaluator::EDistribution::Uniform;
-					break;
-				}
-
-				switch (OcclusionSettings.NormalSpace)
-				{
-				case ENormalMapSpace::Tangent:
-					OcclusionEval->NormalSpace = FMeshOcclusionMapEvaluator::ESpace::Tangent;
-					break;
-				case ENormalMapSpace::Object:
-					OcclusionEval->NormalSpace = FMeshOcclusionMapEvaluator::ESpace::Object;
-					break;
-				}
 				Baker->AddEvaluator(OcclusionEval);
 				break;
 			}
@@ -352,9 +333,7 @@ void UBakeMeshAttributeMapsTool::Setup()
 	OcclusionMapProps->WatchProperty(OcclusionMapProps->OcclusionRays, [this](int32) { bInputsDirty = true; });
 	OcclusionMapProps->WatchProperty(OcclusionMapProps->MaxDistance, [this](float) { bInputsDirty = true; });
 	OcclusionMapProps->WatchProperty(OcclusionMapProps->SpreadAngle, [this](float) { bInputsDirty = true; });
-	OcclusionMapProps->WatchProperty(OcclusionMapProps->Distribution, [this](EOcclusionMapDistribution) { bInputsDirty = true; });
 	OcclusionMapProps->WatchProperty(OcclusionMapProps->BiasAngle, [this](float) { bInputsDirty = true; });
-	OcclusionMapProps->WatchProperty(OcclusionMapProps->NormalSpace, [this](ENormalMapSpace) { bInputsDirty = true; });
 
 
 	CurvatureMapProps = NewObject<UBakedCurvatureMapToolProperties>(this);
@@ -759,9 +738,7 @@ EBakeOpState UBakeMeshAttributeMapsTool::UpdateResult_Occlusion()
 	OcclusionMapSettings.MaxDistance = (OcclusionMapProps->MaxDistance == 0) ? TNumericLimits<float>::Max() : OcclusionMapProps->MaxDistance;
 	OcclusionMapSettings.OcclusionRays = OcclusionMapProps->OcclusionRays;
 	OcclusionMapSettings.SpreadAngle = OcclusionMapProps->SpreadAngle;
-	OcclusionMapSettings.Distribution = OcclusionMapProps->Distribution;
 	OcclusionMapSettings.BiasAngle = OcclusionMapProps->BiasAngle;
-	OcclusionMapSettings.NormalSpace = OcclusionMapProps->NormalSpace;
 
 	if ( !(CachedOcclusionMapSettings == OcclusionMapSettings) )
 	{
diff --git a/Engine/Plugins/Experimental/MeshModelingToolsetExp/Source/MeshModelingToolsExp/Private/BakeMeshAttributeVertexTool.cpp b/Engine/Plugins/Experimental/MeshModelingToolsetExp/Source/MeshModelingToolsExp/Private/BakeMeshAttributeVertexTool.cpp
--- a/Engine/Plugins/Experimental/MeshModelingToolsetExp/Source/MeshModelingToolsExp/Private/BakeMeshAttributeVertexTool.cpp
+++ b/Engine/Plugins/Experimental/MeshModelingToolsetExp/Source/MeshModelingToolsExp/Private/BakeMeshAttributeVertexTool.cpp
@@ -112,26 +112,6 @@ public:
 			OcclusionEval->MaxDistance = OcclusionSettings.MaxDistance;
 			OcclusionEval->SpreadAngle = OcclusionSettings.SpreadAngle;
 			OcclusionEval->BiasAngleDeg = OcclusionSettings.BiasAngle;
-
-			switch (OcclusionSettings.Distribution)
-			{
-			case EOcclusionMapDistribution::Cosine:
-				OcclusionEval->Distribution = FMeshOcclusionMapEvaluator::EDistribution::Cosine;
-				break;
-			case EOcclusionMapDistribution::Uniform:
-				OcclusionEval->Distribution = FMeshOcclusionMapEvaluator::EDistribution::Uniform;
-				break;
-			}
-
-			switch (OcclusionSettings.NormalSpace)
-			{
-			case ENormalMapSpace::Tangent:
-				OcclusionEval->NormalSpace = FMeshOcclusionMapEvaluator::ESpace::Tangent;
-				break;
-			case ENormalMapSpace::Object:
-				OcclusionEval->NormalSpace = FMeshOcclusionMapEvaluator::ESpace::Object;
-				break;
-			}
 		};
 
 		auto InitCurvatureEvaluator = [this] (FMeshCurvatureMapEvaluator* CurvatureEval)
@@ -342,9 +322,7 @@ void UBakeMeshAttributeVertexTool::Setup()
 	OcclusionSettings->WatchProperty(OcclusionSettings->OcclusionRays, [this](int32) { OpState = EBakeOpState::Evaluate; });
 	OcclusionSettings->WatchProperty(OcclusionSettings->MaxDistance, [this](float) { OpState = EBakeOpState::Evaluate; });
 	OcclusionSettings->WatchProperty(OcclusionSettings->SpreadAngle, [this](float) { OpState = EBakeOpState::Evaluate; });
-	OcclusionSettings->WatchProperty(OcclusionSettings->Distribution, [this](EOcclusionMapDistribution) { OpState = EBakeOpState::Evaluate; });
 	OcclusionSettings->WatchProperty(OcclusionSettings->BiasAngle, [this](float) { OpState = EBakeOpState::Evaluate; });
-	OcclusionSettings->WatchProperty(OcclusionSettings->NormalSpace, [this](ENormalMapSpace) { OpState = EBakeOpState::Evaluate; });
 
 	CurvatureSettings = NewObject<UBakedCurvatureMapToolProperties>(this);
 	CurvatureSettings->RestoreProperties(this);
@@ -779,9 +757,7 @@ EBakeOpState UBakeMeshAttributeVertexTool::UpdateResult_Occlusion()
 	OcclusionMapSettings.MaxDistance = (OcclusionSettings->MaxDistance == 0) ? TNumericLimits<float>::Max() : OcclusionSettings->MaxDistance;
 	OcclusionMapSettings.OcclusionRays = OcclusionSettings->OcclusionRays;
 	OcclusionMapSettings.SpreadAngle = OcclusionSettings->SpreadAngle;
-	OcclusionMapSettings.Distribution = OcclusionSettings->Distribution;
 	OcclusionMapSettings.BiasAngle = OcclusionSettings->BiasAngle;
-	OcclusionMapSettings.NormalSpace = OcclusionSettings->NormalSpace;
 
 	if (!(CachedOcclusionMapSettings == OcclusionMapSettings))
 	{
diff --git a/Engine/Plugins/Experimental/MeshModelingToolsetExp/Source/MeshModelingToolsExp/Public/BakeMeshAttributeToolCommon.h b/Engine/Plugins/Experimental/MeshModelingToolsetExp/Source/MeshModelingToolsExp/Public/BakeMeshAttributeToolCommon.h
--- a/Engine/Plugins/Experimental/MeshModelingToolsetExp/Source/MeshModelingToolsExp/Public/BakeMeshAttributeToolCommon.h
+++ b/Engine/Plugins/Experimental/MeshModelingToolsetExp/Source/MeshModelingToolsExp/Public/BakeMeshAttributeToolCommon.h
@@ -43,16 +43,6 @@ public:
 };
 
 
-UENUM()
-enum class ENormalMapSpace
-{
-	/** Tangent space */
-	Tangent UMETA(DisplayName = "Tangent space"),
-	/** Object space */
-	Object UMETA(DisplayName = "Object space")
-};
-
-
 UCLASS()
 class MESHMODELINGTOOLSEXP_API UBakedNormalMapToolProperties : public UInteractiveToolPropertySet
 {
@@ -62,16 +52,6 @@ public:
 };
 
 
-UENUM()
-enum class EOcclusionMapDistribution
-{
-	/** Uniform occlusion rays */
-	Uniform UMETA(DisplayName = "Uniform"),
-	/** Cosine weighted occlusion rays */
-	Cosine UMETA(DisplayName = "Cosine")
-};
-
-
 UCLASS()
 class MESHMODELINGTOOLSEXP_API UBakedOcclusionMapToolProperties : public UInteractiveToolPropertySet
 {
@@ -89,17 +69,9 @@ public:
 	UPROPERTY(EditAnywhere, Category = OcclusionMap, meta = (UIMin = "0", UIMax = "180.0", ClampMin = "0", ClampMax = "180.0"))
 	float SpreadAngle = 180.0;
 
-	/** Angular distribution of occlusion rays in the spread angle. */
-	UPROPERTY(EditAnywhere, Category = OcclusionMap)
-	EOcclusionMapDistribution Distribution = EOcclusionMapDistribution::Cosine;
-
 	/** Contribution of AO rays that are within this angle (degrees) from horizontal are attenuated. This reduces faceting artifacts. */
-	UPROPERTY(EditAnywhere, Category = "OcclusionMap|Ambient Occlusion", meta = (UIMin = "0", UIMax = "45.0", ClampMin = "0", ClampMax = "89.9"))
+	UPROPERTY(EditAnywhere, Category = OcclusionMap, meta = (UIMin = "0", UIMax = "45.0", ClampMin = "0", ClampMax = "89.9"))
 	float BiasAngle = 15.0;
-
-	/** Coordinate space of the bent normal map. */
-	UPROPERTY(EditAnywhere, Category = "OcclusionMap|Bent Normal")
-	ENormalMapSpace NormalSpace = ENormalMapSpace::Tangent;
 };
 
 
@@ -248,9 +220,7 @@ struct FOcclusionMapSettings
 	int32 OcclusionRays;
 	float MaxDistance;
 	float SpreadAngle;
-	EOcclusionMapDistribution Distribution;
 	float BiasAngle;
-	ENormalMapSpace NormalSpace;
 
 	bool operator==(const FOcclusionMapSettings& Other) const
 	{
@@ -258,9 +228,7 @@ struct FOcclusionMapSettings
 			OcclusionRays == Other.OcclusionRays &&
 			MaxDistance == Other.MaxDistance &&
 			SpreadAngle == Other.SpreadAngle &&
-			Distribution == Other.Distribution &&
-			BiasAngle == Other.BiasAngle &&
-			NormalSpace == Other.NormalSpace;
+			BiasAngle == Other.BiasAngle;
 	}
 };
 
