ModelingTools: Invalidate bake result when inputs are invalid.

#rnx

diff --git a/Engine/Plugins/Experimental/MeshModelingToolsetExp/Source/MeshModelingToolsExp/Private/BakeMeshAttributeMapsTool.cpp b/Engine/Plugins/Experimental/MeshModelingToolsetExp/Source/MeshModelingToolsExp/Private/BakeMeshAttributeMapsTool.cpp
--- a/Engine/Plugins/Experimental/MeshModelingToolsetExp/Source/MeshModelingToolsExp/Private/BakeMeshAttributeMapsTool.cpp
+++ b/Engine/Plugins/Experimental/MeshModelingToolsetExp/Source/MeshModelingToolsExp/Private/BakeMeshAttributeMapsTool.cpp
@@ -648,6 +648,7 @@ void UBakeMeshAttributeMapsTool::UpdateResult()
 	// Early exit if op input parameters are invalid.
 	if ((bool)(OpState & EBakeOpState::Invalid))
 	{
+		InvalidateResults();
 		return;
 	}
 
@@ -995,9 +996,16 @@ void UBakeMeshAttributeMapsTool::UpdateOnModeChange()
 		}
 	}
 
+	InvalidateResults();
+}
+
+
+void UBakeMeshAttributeMapsTool::InvalidateResults()
+{
 	Settings->Result.Empty();
 	Settings->Result.SetNum(ResultTypes.Num());
 }
 
 
+
 #undef LOCTEXT_NAMESPACE
diff --git a/Engine/Plugins/Experimental/MeshModelingToolsetExp/Source/MeshModelingToolsExp/Private/BakeMultiMeshAttributeMapsTool.cpp b/Engine/Plugins/Experimental/MeshModelingToolsetExp/Source/MeshModelingToolsExp/Private/BakeMultiMeshAttributeMapsTool.cpp
--- a/Engine/Plugins/Experimental/MeshModelingToolsetExp/Source/MeshModelingToolsExp/Private/BakeMultiMeshAttributeMapsTool.cpp
+++ b/Engine/Plugins/Experimental/MeshModelingToolsetExp/Source/MeshModelingToolsExp/Private/BakeMultiMeshAttributeMapsTool.cpp
@@ -611,6 +611,7 @@ void UBakeMultiMeshAttributeMapsTool::UpdateResult()
 	// Early exit if op input parameters are invalid.
 	if ((bool)(OpState & EBakeOpState::Invalid))
 	{
+		InvalidateResults();
 		return;
 	}
 
@@ -749,9 +750,15 @@ void UBakeMultiMeshAttributeMapsTool::UpdateOnModeChange()
 		}
 	}
 	
+	InvalidateResults();
+}
+
+void UBakeMultiMeshAttributeMapsTool::InvalidateResults()
+{
 	Settings->Result.Empty();
 	Settings->Result.SetNum(ResultTypes.Num());
 }
 
 
+
 #undef LOCTEXT_NAMESPACE
diff --git a/Engine/Plugins/Experimental/MeshModelingToolsetExp/Source/MeshModelingToolsExp/Public/BakeMeshAttributeMapsTool.h b/Engine/Plugins/Experimental/MeshModelingToolsetExp/Source/MeshModelingToolsExp/Public/BakeMeshAttributeMapsTool.h
--- a/Engine/Plugins/Experimental/MeshModelingToolsetExp/Source/MeshModelingToolsExp/Public/BakeMeshAttributeMapsTool.h
+++ b/Engine/Plugins/Experimental/MeshModelingToolsetExp/Source/MeshModelingToolsExp/Public/BakeMeshAttributeMapsTool.h
@@ -164,6 +164,8 @@ protected:
 	bool bInputsDirty = false;
 	void UpdateOnModeChange();
 
+	void InvalidateResults();
+
 	FDetailMeshSettings CachedDetailMeshSettings;
 	TSharedPtr<UE::Geometry::TImageBuilder<FVector4f>, ESPMode::ThreadSafe> CachedDetailNormalMap;
 	EBakeOpState UpdateResult_DetailNormalMap();
diff --git a/Engine/Plugins/Experimental/MeshModelingToolsetExp/Source/MeshModelingToolsExp/Public/BakeMultiMeshAttributeMapsTool.h b/Engine/Plugins/Experimental/MeshModelingToolsetExp/Source/MeshModelingToolsExp/Public/BakeMultiMeshAttributeMapsTool.h
--- a/Engine/Plugins/Experimental/MeshModelingToolsetExp/Source/MeshModelingToolsExp/Public/BakeMultiMeshAttributeMapsTool.h
+++ b/Engine/Plugins/Experimental/MeshModelingToolsetExp/Source/MeshModelingToolsExp/Public/BakeMultiMeshAttributeMapsTool.h
@@ -200,6 +200,8 @@ protected:
 	bool bInputsDirty = false;
 	void UpdateOnModeChange();
 
+	void InvalidateResults();
+
 	// Cached detail mesh data
 	FBakeMultiMeshDetailSettings CachedDetailSettings;
 	EBakeOpState UpdateResult_DetailMeshes();
