ToolsFramework: delete InteractiveToolStorableSelection/Subsystem
#rnx

diff --git a/Engine/Plugins/Experimental/MeshModelingToolsetExp/Source/ModelingComponents/Public/Selection/PersistentMeshSelection.h b/Engine/Plugins/Experimental/MeshModelingToolsetExp/Source/ModelingComponents/Public/Selection/PersistentMeshSelection.h
--- a/Engine/Plugins/Experimental/MeshModelingToolsetExp/Source/ModelingComponents/Public/Selection/PersistentMeshSelection.h
+++ b/Engine/Plugins/Experimental/MeshModelingToolsetExp/Source/ModelingComponents/Public/Selection/PersistentMeshSelection.h
@@ -5,7 +5,6 @@
 #include "CoreMinimal.h"
 #include "IndexTypes.h"
 #include "SegmentTypes.h"
-#include "InteractiveToolStorableSelection.h"
 
 #include "PersistentMeshSelection.generated.h"
 
diff --git a/Engine/Source/Runtime/InteractiveToolsFramework/Private/InteractiveToolManager.cpp b/Engine/Source/Runtime/InteractiveToolsFramework/Private/InteractiveToolManager.cpp
--- a/Engine/Source/Runtime/InteractiveToolsFramework/Private/InteractiveToolManager.cpp
+++ b/Engine/Source/Runtime/InteractiveToolsFramework/Private/InteractiveToolManager.cpp
@@ -4,7 +4,6 @@
 
 #include "Engine/Engine.h"
 #include "InteractiveToolsContext.h"
-#include "InteractiveToolsSelectionStoreSubsystem.h"
 #include "ContextObjectStore.h"
 
 #define LOCTEXT_NAMESPACE "UInteractiveToolManager"
diff --git a/Engine/Source/Runtime/InteractiveToolsFramework/Public/InteractiveToolStorableSelection.h b/Engine/Source/Runtime/InteractiveToolsFramework/Public/InteractiveToolStorableSelection.h
deleted file mode 100644
diff --git a/Engine/Source/Runtime/InteractiveToolsFramework/Public/InteractiveToolsSelectionStoreSubsystem.h b/Engine/Source/Runtime/InteractiveToolsFramework/Public/InteractiveToolsSelectionStoreSubsystem.h
deleted file mode 100644
