ModelingTools: Fixed the ability for DestructibleMeshes and DestructibleMeshComponents from being used as targets in Modeling Tools and the UVEditor. These classes do not provide the required functionality and, being deprecated, it was decided that removing them from consideration was more appropriate than implementing the missing behavior.

#rnx

diff --git a/Engine/Plugins/Experimental/MeshModelingToolsetExp/Source/ModelingComponentsEditorOnly/Private/ToolTargets/SkeletalMeshComponentToolTarget.cpp b/Engine/Plugins/Experimental/MeshModelingToolsetExp/Source/ModelingComponentsEditorOnly/Private/ToolTargets/SkeletalMeshComponentToolTarget.cpp
--- a/Engine/Plugins/Experimental/MeshModelingToolsetExp/Source/ModelingComponentsEditorOnly/Private/ToolTargets/SkeletalMeshComponentToolTarget.cpp
+++ b/Engine/Plugins/Experimental/MeshModelingToolsetExp/Source/ModelingComponentsEditorOnly/Private/ToolTargets/SkeletalMeshComponentToolTarget.cpp
@@ -142,7 +142,15 @@ USkeletalMesh* USkeletalMeshComponentToolTarget::GetSkeletalMesh() const
 
 bool USkeletalMeshComponentToolTargetFactory::CanBuildTarget(UObject* SourceObject, const FToolTargetTypeRequirements& Requirements) const
 {
+	// We are using an exact cast here to prevent subclasses, which might not meet all
+	// requirements for functionality such as the deprecated DestructibleMeshComponent, from 
+	// being caught up as valid targets.
+	// If you want to make the tool target work with some subclass of USkeletalMeshComponent,
+	// just add another factory that allows that class specifically(but make sure that
+	// GetMeshDescription and such work properly)
+
 	return Cast<USkinnedMeshComponent>(SourceObject) && Cast<USkinnedMeshComponent>(SourceObject)->SkeletalMesh &&
+		ExactCast<USkeletalMesh>(Cast<USkinnedMeshComponent>(SourceObject)->SkeletalMesh) &&
 		Requirements.AreSatisfiedBy(USkeletalMeshComponentToolTarget::StaticClass());
 }
 
diff --git a/Engine/Plugins/Experimental/MeshModelingToolsetExp/Source/ModelingComponentsEditorOnly/Private/ToolTargets/SkeletalMeshToolTarget.cpp b/Engine/Plugins/Experimental/MeshModelingToolsetExp/Source/ModelingComponentsEditorOnly/Private/ToolTargets/SkeletalMeshToolTarget.cpp
--- a/Engine/Plugins/Experimental/MeshModelingToolsetExp/Source/ModelingComponentsEditorOnly/Private/ToolTargets/SkeletalMeshToolTarget.cpp
+++ b/Engine/Plugins/Experimental/MeshModelingToolsetExp/Source/ModelingComponentsEditorOnly/Private/ToolTargets/SkeletalMeshToolTarget.cpp
@@ -211,7 +211,14 @@ USkeletalMesh* USkeletalMeshToolTarget::GetSkeletalMesh() const
 
 bool USkeletalMeshToolTargetFactory::CanBuildTarget(UObject* SourceObject, const FToolTargetTypeRequirements& Requirements) const
 {
-	return Cast<USkeletalMesh>(SourceObject) 
+	// We are using an exact cast here to prevent subclasses, which might not meet all
+	// requirements for functionality such as the deprecated DestructibleMesh, from 
+	// being caught up as valid targets.
+	// If you want to make the tool target work with some subclass of USkeletalMesh,
+	// just add another factory that allows that class specifically (but make sure that
+	// GetMeshDescription and such work properly)
+
+	return ExactCast<USkeletalMesh>(SourceObject) 
 		&& Requirements.AreSatisfiedBy(USkeletalMeshToolTarget::StaticClass());
 }
 
