Flip align/align-flipped meaning in pivot tool and clarify comment explaining what they do

diff --git a/Engine/Plugins/Experimental/MeshModelingToolsetExp/Source/MeshModelingToolsExp/Private/EditPivotTool.cpp b/Engine/Plugins/Experimental/MeshModelingToolsetExp/Source/MeshModelingToolsExp/Private/EditPivotTool.cpp
--- a/Engine/Plugins/Experimental/MeshModelingToolsetExp/Source/MeshModelingToolsExp/Private/EditPivotTool.cpp
+++ b/Engine/Plugins/Experimental/MeshModelingToolsetExp/Source/MeshModelingToolsExp/Private/EditPivotTool.cpp
@@ -390,7 +390,7 @@ void UEditPivotTool::OnClickDrag(const FInputDeviceRay& DragPos)
 
 
 	FVector HitPos = Result.ImpactPoint;
-	FVector TargetNormal = (-NormalSign) * Result.Normal;
+	FVector TargetNormal = (NormalSign) * Result.Normal;
 
 	FQuaterniond AlignRotation = (bRotate) ?
 		FQuaterniond(FVector3d::UnitZ(), (FVector3d)TargetNormal) : FQuaterniond::Identity();
diff --git a/Engine/Plugins/Experimental/MeshModelingToolsetExp/Source/MeshModelingToolsExp/Public/EditPivotTool.h b/Engine/Plugins/Experimental/MeshModelingToolsetExp/Source/MeshModelingToolsExp/Public/EditPivotTool.h
--- a/Engine/Plugins/Experimental/MeshModelingToolsetExp/Source/MeshModelingToolsExp/Public/EditPivotTool.h
+++ b/Engine/Plugins/Experimental/MeshModelingToolsetExp/Source/MeshModelingToolsExp/Public/EditPivotTool.h
@@ -43,10 +43,10 @@ enum class EEditPivotSnapDragRotationMode : uint8
 	/** Snap-Drag only translates, ignoring Normals */
 	Ignore = 0 UMETA(DisplayName = "Ignore"),
 
-	/** Snap-Drag aligns the Source and Target Normals to point in the same direction */
+	/** Snap-Drag aligns the pivot Z Axis and Target Normals to point in the same direction */
 	Align = 1 UMETA(DisplayName = "Align"),
 
-	/** Snap-Drag aligns the Source Normal to the opposite of the Target Normal direction */
+	/** Snap-Drag aligns the pivot Z Axis to the opposite of the Target Normal direction */
 	AlignFlipped = 2 UMETA(DisplayName = "Align Flipped"),
 
 	LastValue UMETA(Hidden)
