ModelingTools: Replace MapPreview with GetOptions combo box in BakeTx/BakeAll.
Replace Result array property with a map to add names to result textures in BakeTx/BakeAll.
Cleanup BakeTx/BakeAll enum names.
Refactor shared code into the BakeMeshAttributeMapsToolBase class.

#rnx

diff --git a/Engine/Plugins/Experimental/MeshModelingToolsetExp/Source/MeshModelingToolsExp/Private/BakeMeshAttributeMapsTool.cpp b/Engine/Plugins/Experimental/MeshModelingToolsetExp/Source/MeshModelingToolsExp/Private/BakeMeshAttributeMapsTool.cpp
--- a/Engine/Plugins/Experimental/MeshModelingToolsetExp/Source/MeshModelingToolsExp/Private/BakeMeshAttributeMapsTool.cpp
+++ b/Engine/Plugins/Experimental/MeshModelingToolsetExp/Source/MeshModelingToolsExp/Private/BakeMeshAttributeMapsTool.cpp
@@ -86,12 +86,19 @@ UInteractiveTool* UBakeMeshAttributeMapsToolBuilder::BuildTool(const FToolBuilde
 
 
 
-TArray<FString> UBakeMeshAttributeMapsToolProperties::GetUVLayerNamesFunc()
+const TArray<FString>& UBakeMeshAttributeMapsToolProperties::GetUVLayerNamesFunc()
 {
 	return UVLayerNamesList;
 }
 
 
+const TArray<FString>& UBakeMeshAttributeMapsToolProperties::GetMapPreviewNamesFunc()
+{
+	return MapPreviewNamesList;
+}
+
+
+
 /*
  * Operators
  */
@@ -147,7 +154,7 @@ public:
 		{
 			switch (BakeCacheSettings.BakeMapTypes & MapType)
 			{
-			case EBakeMapType::TangentSpaceNormalMap:
+			case EBakeMapType::TangentSpaceNormal:
 			{
 				TSharedPtr<FMeshNormalMapEvaluator, ESPMode::ThreadSafe> NormalEval = MakeShared<FMeshNormalMapEvaluator, ESPMode::ThreadSafe>();
 				DetailSampler.SetNormalMap(DetailMesh.Get(), IMeshBakerDetailSampler::FBakeDetailTexture(DetailMeshNormalMap.Get(), DetailMeshNormalUVLayer));
@@ -187,7 +194,7 @@ public:
 				Baker->AddEvaluator(CurvatureEval);
 				break;
 			}
-			case EBakeMapType::NormalImage:
+			case EBakeMapType::ObjectSpaceNormal:
 			{
 				TSharedPtr<FMeshPropertyMapEvaluator, ESPMode::ThreadSafe> PropertyEval = MakeShared<FMeshPropertyMapEvaluator, ESPMode::ThreadSafe>();
 				PropertyEval->Property = EMeshPropertyMapType::Normal;
@@ -195,14 +202,14 @@ public:
 				Baker->AddEvaluator(PropertyEval);
 				break;
 			}
-			case EBakeMapType::FaceNormalImage:
+			case EBakeMapType::FaceNormal:
 			{
 				TSharedPtr<FMeshPropertyMapEvaluator, ESPMode::ThreadSafe> PropertyEval = MakeShared<FMeshPropertyMapEvaluator, ESPMode::ThreadSafe>();
 				PropertyEval->Property = EMeshPropertyMapType::FacetNormal;
 				Baker->AddEvaluator(PropertyEval);
 				break;
 			}
-			case EBakeMapType::PositionImage:
+			case EBakeMapType::Position:
 			{
 				TSharedPtr<FMeshPropertyMapEvaluator, ESPMode::ThreadSafe> PropertyEval = MakeShared<FMeshPropertyMapEvaluator, ESPMode::ThreadSafe>();
 				PropertyEval->Property = EMeshPropertyMapType::Position;
@@ -216,14 +223,14 @@ public:
 				Baker->AddEvaluator(PropertyEval);
 				break;
 			}
-			case EBakeMapType::VertexColorImage:
+			case EBakeMapType::VertexColor:
 			{
 				TSharedPtr<FMeshPropertyMapEvaluator, ESPMode::ThreadSafe> PropertyEval = MakeShared<FMeshPropertyMapEvaluator, ESPMode::ThreadSafe>();
 				PropertyEval->Property = EMeshPropertyMapType::VertexColor;
 				Baker->AddEvaluator(PropertyEval);
 				break;
 			}
-			case EBakeMapType::Texture2DImage:
+			case EBakeMapType::Texture:
 			{
 				TSharedPtr<FMeshResampleImageEvaluator, ESPMode::ThreadSafe> TextureEval = MakeShared<FMeshResampleImageEvaluator, ESPMode::ThreadSafe>();
 				DetailSampler.SetColorMap(DetailMesh.Get(), IMeshBakerDetailSampler::FBakeDetailTexture(TextureImage.Get(), TextureSettings.UVLayer));
@@ -270,24 +277,11 @@ void UBakeMeshAttributeMapsTool::Setup()
 	// Setup tool property sets
 	Settings = NewObject<UBakeMeshAttributeMapsToolProperties>(this);
 	Settings->RestoreProperties(this);
-	Settings->UVLayerNamesList.Reset();
-	int32 FoundIndex = -1;
-	for (int32 k = 0; k < BaseMesh.Attributes()->NumUVLayers(); ++k)
-	{
-		Settings->UVLayerNamesList.Add(FString::FromInt(k));
-		if (Settings->UVLayer == Settings->UVLayerNamesList.Last())
-		{
-			FoundIndex = k;
-		}
-	}
-	if (FoundIndex == -1)
-	{
-		Settings->UVLayer = Settings->UVLayerNamesList[0];
-	}
+	UpdateUVLayerNames(Settings, BaseMesh);
 	AddToolPropertySource(Settings);
 
 	Settings->WatchProperty(Settings->MapTypes, [this](int32) { bInputsDirty = true; UpdateOnModeChange(); });
-	Settings->WatchProperty(Settings->MapPreview, [this](int32) { UpdateVisualization(); GetToolManager()->PostInvalidation(); });
+	Settings->WatchProperty(Settings->MapPreview, [this](FString) { UpdateVisualization(); GetToolManager()->PostInvalidation(); });
 	Settings->WatchProperty(Settings->Resolution, [this](EBakeTextureResolution) { bInputsDirty = true; });
 	Settings->WatchProperty(Settings->SourceFormat, [this](EBakeTextureFormat) { bInputsDirty = true; });
 	Settings->WatchProperty(Settings->UVLayer, [this](FString) { bInputsDirty = true; });
@@ -320,12 +314,6 @@ void UBakeMeshAttributeMapsTool::Setup()
 	});
 	
 
-	NormalMapProps = NewObject<UBakedNormalMapToolProperties>(this);
-	NormalMapProps->RestoreProperties(this);
-	AddToolPropertySource(NormalMapProps);
-	SetToolPropertySourceEnabled(NormalMapProps, false);
-
-
 	OcclusionMapProps = NewObject<UBakedOcclusionMapToolProperties>(this);
 	OcclusionMapProps->RestoreProperties(this);
 	AddToolPropertySource(OcclusionMapProps);
@@ -386,10 +374,9 @@ bool UBakeMeshAttributeMapsTool::CanAccept() const
 	if (bCanAccept)
 	{
 		// Allow Accept if all non-None types have valid results.
-		const int NumResults = Settings->Result.Num();
-		for (int ResultIdx = 0; ResultIdx < NumResults; ++ResultIdx)
+		for (const TTuple<EBakeMapType, TObjectPtr<UTexture2D>>& Result : Settings->Result)
 		{
-			bCanAccept = bCanAccept && Settings->Result[ResultIdx];
+			bCanAccept = bCanAccept && Result.Get<1>();
 		}
 	}
 	return bCanAccept;
@@ -404,7 +391,7 @@ TUniquePtr<UE::Geometry::TGenericDataOperator<FMeshMapBaker>> UBakeMeshAttribute
 	Op->BaseMesh = &BaseMesh;
 	Op->BakeCacheSettings = CachedBakeCacheSettings;
 
-	constexpr EBakeMapType RequiresTangents = EBakeMapType::TangentSpaceNormalMap | EBakeMapType::BentNormal;
+	constexpr EBakeMapType RequiresTangents = EBakeMapType::TangentSpaceNormal | EBakeMapType::BentNormal;
 	if ((bool)(CachedBakeCacheSettings.BakeMapTypes & RequiresTangents))
 	{
 		Op->BaseMeshTangents = BaseMeshTangents;
@@ -417,7 +404,7 @@ TUniquePtr<UE::Geometry::TGenericDataOperator<FMeshMapBaker>> UBakeMeshAttribute
 		Op->DetailMeshNormalUVLayer = CachedDetailMeshSettings.UVLayer;
 	}
 
-	if ((bool)(CachedBakeCacheSettings.BakeMapTypes & EBakeMapType::TangentSpaceNormalMap))
+	if ((bool)(CachedBakeCacheSettings.BakeMapTypes & EBakeMapType::TangentSpaceNormal))
 	{
 		Op->NormalSettings = CachedNormalMapSettings;
 	}
@@ -433,16 +420,16 @@ TUniquePtr<UE::Geometry::TGenericDataOperator<FMeshMapBaker>> UBakeMeshAttribute
 		Op->CurvatureSettings = CachedCurvatureMapSettings;
 	}
 
-	if ((bool)(CachedBakeCacheSettings.BakeMapTypes & EBakeMapType::NormalImage) ||
-		(bool)(CachedBakeCacheSettings.BakeMapTypes & EBakeMapType::FaceNormalImage) ||
-		(bool)(CachedBakeCacheSettings.BakeMapTypes & EBakeMapType::PositionImage) ||
+	if ((bool)(CachedBakeCacheSettings.BakeMapTypes & EBakeMapType::ObjectSpaceNormal) ||
+		(bool)(CachedBakeCacheSettings.BakeMapTypes & EBakeMapType::FaceNormal) ||
+		(bool)(CachedBakeCacheSettings.BakeMapTypes & EBakeMapType::Position) ||
 		(bool)(CachedBakeCacheSettings.BakeMapTypes & EBakeMapType::MaterialID) ||
-		(bool)(CachedBakeCacheSettings.BakeMapTypes & EBakeMapType::VertexColorImage))
+		(bool)(CachedBakeCacheSettings.BakeMapTypes & EBakeMapType::VertexColor))
 	{
 		Op->PropertySettings = CachedMeshPropertyMapSettings;
 	}
 
-	if ((bool)(CachedBakeCacheSettings.BakeMapTypes & EBakeMapType::Texture2DImage))
+	if ((bool)(CachedBakeCacheSettings.BakeMapTypes & EBakeMapType::Texture))
 	{
 		Op->TextureSettings = CachedTexture2DImageSettings;
 		Op->TextureImage = CachedTextureImage;
@@ -464,7 +451,6 @@ void UBakeMeshAttributeMapsTool::Shutdown(EToolShutdownType ShutdownType)
 	
 	Settings->SaveProperties(this);
 	OcclusionMapProps->SaveProperties(this);
-	NormalMapProps->SaveProperties(this);
 	CurvatureMapProps->SaveProperties(this);
 	Texture2DProps->SaveProperties(this);
 	MultiTextureProps->SaveProperties(this);
@@ -488,12 +474,11 @@ void UBakeMeshAttributeMapsTool::Shutdown(EToolShutdownType ShutdownType)
 		const FString BaseName = UE::ToolTarget::GetTargetActor(Targets[0])->GetName();
 
 		bool bCreatedAssetOK = true;
-		const int NumResults = Settings->Result.Num();
-		for (int ResultIdx = 0; ResultIdx < NumResults; ++ResultIdx)
+		for (const TTuple<EBakeMapType, TObjectPtr<UTexture2D>>& Result : Settings->Result)
 		{
 			FString TexName;
-			GetTextureName(ResultTypes[ResultIdx], BaseName, TexName);
-			bCreatedAssetOK = bCreatedAssetOK && UE::Modeling::CreateTextureObject(GetToolManager(), FCreateTextureObjectParams{ 0, SourceComponent->GetWorld(), SourceAsset, TexName, Settings->Result[ResultIdx] }).IsOK();
+			GetTextureName(Result.Get<0>(), BaseName, TexName);
+			bCreatedAssetOK = bCreatedAssetOK && UE::Modeling::CreateTextureObject(GetToolManager(), FCreateTextureObjectParams{ 0, SourceComponent->GetWorld(), SourceAsset, TexName, Result.Get<1>() }).IsOK();
 		}
 		ensure(bCreatedAssetOK);
 	}
@@ -615,7 +600,7 @@ void UBakeMeshAttributeMapsTool::UpdateResult()
 	OpState |= UpdateResult_DetailNormalMap();
 
 	// Update map type settings
-	if ((bool)(CachedBakeCacheSettings.BakeMapTypes & EBakeMapType::TangentSpaceNormalMap))
+	if ((bool)(CachedBakeCacheSettings.BakeMapTypes & EBakeMapType::TangentSpaceNormal))
 	{
 		OpState |= UpdateResult_Normal();
 	}
@@ -628,15 +613,15 @@ void UBakeMeshAttributeMapsTool::UpdateResult()
 	{
 		OpState |= UpdateResult_Curvature();
 	}
-	if ((bool)(CachedBakeCacheSettings.BakeMapTypes & EBakeMapType::NormalImage) ||
-		(bool)(CachedBakeCacheSettings.BakeMapTypes & EBakeMapType::FaceNormalImage) ||
-		(bool)(CachedBakeCacheSettings.BakeMapTypes & EBakeMapType::PositionImage) ||
+	if ((bool)(CachedBakeCacheSettings.BakeMapTypes & EBakeMapType::ObjectSpaceNormal) ||
+		(bool)(CachedBakeCacheSettings.BakeMapTypes & EBakeMapType::FaceNormal) ||
+		(bool)(CachedBakeCacheSettings.BakeMapTypes & EBakeMapType::Position) ||
 		(bool)(CachedBakeCacheSettings.BakeMapTypes & EBakeMapType::MaterialID) ||
-		(bool)(CachedBakeCacheSettings.BakeMapTypes & EBakeMapType::VertexColorImage))
+		(bool)(CachedBakeCacheSettings.BakeMapTypes & EBakeMapType::VertexColor))
 	{
 		OpState |= UpdateResult_MeshProperty();
 	}
-	if ((bool)(CachedBakeCacheSettings.BakeMapTypes & EBakeMapType::Texture2DImage))
+	if ((bool)(CachedBakeCacheSettings.BakeMapTypes & EBakeMapType::Texture))
 	{
 		OpState |= UpdateResult_Texture2DImage();
 	}
@@ -939,26 +924,28 @@ void UBakeMeshAttributeMapsTool::UpdateVisualization()
 {
 	PreviewMesh->SetOverrideRenderMaterial(PreviewMaterial);
 
-	// Map CachedMaps to Settings->Result
-	int NumResults = Settings->Result.Num();
-	for (int ResultIdx = 0; ResultIdx < NumResults; ResultIdx++)
+	// Populate Settings->Result from CachedMaps
+	for (const TTuple<EBakeMapType, TObjectPtr<UTexture2D>>& Map : CachedMaps)
 	{
-		Settings->Result[ResultIdx] = CachedMaps[CachedMapIndices[ResultTypes[ResultIdx]]];
+		if (Settings->Result.Contains(Map.Get<0>()))
+		{
+			Settings->Result[Map.Get<0>()] = Map.Get<1>();
+		}
 	}
 
-	// Set the preview material according to the preview index.
-	if (Settings->MapPreview >= 0 && Settings->MapPreview < Settings->Result.Num())
+	// Set the preview material according to the preview name.
+	if (const EBakeMapType* PreviewType = NameToMapTypeMap.Find(Settings->MapPreview))
 	{
-		UpdatePreview(Settings->MapPreview);
+		UpdatePreview(*PreviewType);
 	}
 }
 
 
 void UBakeMeshAttributeMapsTool::UpdateOnModeChange()
 {
-	OnMapTypesUpdated(Settings->MapTypes);
-	
-	SetToolPropertySourceEnabled(NormalMapProps, false);
+	OnMapTypesUpdated(Settings);
+
+	// Update tool property sets.
 	SetToolPropertySourceEnabled(OcclusionMapProps, false);
 	SetToolPropertySourceEnabled(CurvatureMapProps, false);
 	SetToolPropertySourceEnabled(Texture2DProps, false);
@@ -968,8 +955,7 @@ void UBakeMeshAttributeMapsTool::UpdateOnModeChange()
 	{
 		switch ((EBakeMapType)Settings->MapTypes & MapType)
 		{
-		case EBakeMapType::TangentSpaceNormalMap:
-			SetToolPropertySourceEnabled(NormalMapProps, true);
+		case EBakeMapType::TangentSpaceNormal:
 			break;
 		case EBakeMapType::AmbientOcclusion:
 		case EBakeMapType::BentNormal:
@@ -979,13 +965,13 @@ void UBakeMeshAttributeMapsTool::UpdateOnModeChange()
 		case EBakeMapType::Curvature:
 			SetToolPropertySourceEnabled(CurvatureMapProps, true);
 			break;
-		case EBakeMapType::NormalImage:
-		case EBakeMapType::FaceNormalImage:
-		case EBakeMapType::PositionImage:
+		case EBakeMapType::ObjectSpaceNormal:
+		case EBakeMapType::FaceNormal:
+		case EBakeMapType::Position:
 		case EBakeMapType::MaterialID:
-		case EBakeMapType::VertexColorImage:
+		case EBakeMapType::VertexColor:
 			break;
-		case EBakeMapType::Texture2DImage:
+		case EBakeMapType::Texture:
 			SetToolPropertySourceEnabled(Texture2DProps, true);
 			break;
 		case EBakeMapType::MultiTexture:
@@ -995,15 +981,15 @@ void UBakeMeshAttributeMapsTool::UpdateOnModeChange()
 			break;
 		}
 	}
-
-	InvalidateResults();
 }
 
 
 void UBakeMeshAttributeMapsTool::InvalidateResults()
 {
-	Settings->Result.Empty();
-	Settings->Result.SetNum(ResultTypes.Num());
+	for (TTuple<EBakeMapType, TObjectPtr<UTexture2D>>& Result : Settings->Result)
+	{
+		Result.Get<1>() = nullptr;
+	}
 }
 
 
diff --git a/Engine/Plugins/Experimental/MeshModelingToolsetExp/Source/MeshModelingToolsExp/Private/BakeMeshAttributeMapsToolBase.cpp b/Engine/Plugins/Experimental/MeshModelingToolsetExp/Source/MeshModelingToolsExp/Private/BakeMeshAttributeMapsToolBase.cpp
--- a/Engine/Plugins/Experimental/MeshModelingToolsetExp/Source/MeshModelingToolsExp/Private/BakeMeshAttributeMapsToolBase.cpp
+++ b/Engine/Plugins/Experimental/MeshModelingToolsetExp/Source/MeshModelingToolsExp/Private/BakeMeshAttributeMapsToolBase.cpp
@@ -25,6 +25,48 @@ using namespace UE::Geometry;
 
 #define LOCTEXT_NAMESPACE "UBakeMeshAttributeMapsToolBase"
 
+// Initialize map type name map. These names should match the
+// auto-generated enum names. 
+static const char* BakeType_TangentSpaceNormals = "Tangent Space Normal";
+static const char* BakeType_AmbientOcclusion = "Ambient Occlusion";
+static const char* BakeType_BentNormals = "Bent Normal";
+static const char* BakeType_Curvature = "Curvature";
+static const char* BakeType_Texture = "Texture";
+static const char* BakeType_ObjectSpaceNormals = "Object Space Normal";
+static const char* BakeType_FaceNormals = "Face Normal";
+static const char* BakeType_Position = "Position";
+static const char* BakeType_MaterialID = "Material ID";
+static const char* BakeType_MultiTexture = "Multi Texture";
+static const char* BakeType_VertexColors = "Vertex Color";
+
+TMap<FString, EBakeMapType> UBakeMeshAttributeMapsToolBase::NameToMapTypeMap = {
+	TTuple<FString, EBakeMapType>(BakeType_TangentSpaceNormals, EBakeMapType::TangentSpaceNormal),
+	TTuple<FString, EBakeMapType>(BakeType_AmbientOcclusion, EBakeMapType::AmbientOcclusion),
+	TTuple<FString, EBakeMapType>(BakeType_BentNormals, EBakeMapType::BentNormal),
+	TTuple<FString, EBakeMapType>(BakeType_Curvature, EBakeMapType::Curvature),
+	TTuple<FString, EBakeMapType>(BakeType_Texture, EBakeMapType::Texture),
+	TTuple<FString, EBakeMapType>(BakeType_ObjectSpaceNormals, EBakeMapType::ObjectSpaceNormal),
+	TTuple<FString, EBakeMapType>(BakeType_FaceNormals, EBakeMapType::FaceNormal),
+	TTuple<FString, EBakeMapType>(BakeType_Position, EBakeMapType::Position),
+	TTuple<FString, EBakeMapType>(BakeType_MaterialID, EBakeMapType::MaterialID),
+	TTuple<FString, EBakeMapType>(BakeType_MultiTexture, EBakeMapType::MultiTexture),
+	TTuple<FString, EBakeMapType>(BakeType_VertexColors, EBakeMapType::VertexColor)
+};
+
+TMap<EBakeMapType, FString> UBakeMeshAttributeMapsToolBase::MapTypeToNameMap = {
+	TTuple<EBakeMapType, FString>(EBakeMapType::TangentSpaceNormal, BakeType_TangentSpaceNormals),
+	TTuple<EBakeMapType, FString>(EBakeMapType::AmbientOcclusion, BakeType_AmbientOcclusion),
+	TTuple<EBakeMapType, FString>(EBakeMapType::BentNormal, BakeType_BentNormals),
+	TTuple<EBakeMapType, FString>(EBakeMapType::Curvature, BakeType_Curvature),
+	TTuple<EBakeMapType, FString>(EBakeMapType::Texture, BakeType_Texture),
+	TTuple<EBakeMapType, FString>(EBakeMapType::ObjectSpaceNormal, BakeType_ObjectSpaceNormals),
+	TTuple<EBakeMapType, FString>(EBakeMapType::FaceNormal, BakeType_FaceNormals),
+	TTuple<EBakeMapType, FString>(EBakeMapType::Position, BakeType_Position),
+	TTuple<EBakeMapType, FString>(EBakeMapType::MaterialID, BakeType_MaterialID),
+	TTuple<EBakeMapType, FString>(EBakeMapType::MultiTexture, BakeType_MultiTexture),
+	TTuple<EBakeMapType, FString>(EBakeMapType::VertexColor, BakeType_VertexColors)
+};
+
 
 void UBakeMeshAttributeMapsToolBase::Setup()
 {
@@ -138,48 +180,20 @@ void UBakeMeshAttributeMapsToolBase::UpdateVisualization()
 {
 }
 
-void UBakeMeshAttributeMapsToolBase::OnMapTypesUpdated(const int32 MapTypes)
-{
-	const EBakeMapType BakeMapTypes = GetMapTypes(MapTypes);
-	ResultTypes = GetMapTypesArray(MapTypes);
-	
-	// Generate a map between EBakeMapType and CachedMaps
-	CachedMapIndices.Empty();
-	int32 CachedMapIdx = 0;
 
-	// Use the processed bitfield which may contain additional targets
-	// (ex. AO if BentNormal was requested).
-	for (EBakeMapType MapType : ALL_BAKE_MAP_TYPES)
+void UBakeMeshAttributeMapsToolBase::UpdatePreview(const EBakeMapType PreviewMapType)
+{
+	if (PreviewMapType == EBakeMapType::None)
 	{
-		if (MapType == EBakeMapType::Occlusion)
-		{
-			if ((bool)(BakeMapTypes & EBakeMapType::AmbientOcclusion))
-			{
-				CachedMapIndices.Add(EBakeMapType::AmbientOcclusion, CachedMapIdx++);
-			}
-			if ((bool)(BakeMapTypes & EBakeMapType::BentNormal))
-			{
-				CachedMapIndices.Add(EBakeMapType::BentNormal, CachedMapIdx++);
-			}
-		}
-		else if( (bool)(BakeMapTypes & MapType) )
-		{
-			CachedMapIndices.Add(MapType, CachedMapIdx++);
-		}
+		return;
 	}
-	CachedMaps.Empty();
-	CachedMaps.SetNum(CachedMapIndices.Num());
-}
 
-void UBakeMeshAttributeMapsToolBase::UpdatePreview(const int PreviewIdx)
-{
-	const EBakeMapType& PreviewMapType = ResultTypes[PreviewIdx];
-	if (PreviewMapType == EBakeMapType::None)
+	if (!CachedMaps.Contains(PreviewMapType))
 	{
 		return;
 	}
 	
-	UTexture2D* PreviewMap = CachedMaps[CachedMapIndices[PreviewMapType]];
+	UTexture2D* PreviewMap = CachedMaps[PreviewMapType];
 	switch (PreviewMapType)
 	{
 	default:
@@ -187,7 +201,7 @@ void UBakeMeshAttributeMapsToolBase::UpdatePreview(const int PreviewIdx)
 		PreviewMaterial->SetTextureParameterValue(TEXT("OcclusionMap"), EmptyColorMapWhite);
 		PreviewMaterial->SetTextureParameterValue(TEXT("ColorMap"), EmptyColorMapWhite);
 		break;
-	case EBakeMapType::TangentSpaceNormalMap:
+	case EBakeMapType::TangentSpaceNormal:
 		PreviewMaterial->SetTextureParameterValue(TEXT("NormalMap"), PreviewMap);
 		PreviewMaterial->SetTextureParameterValue(TEXT("OcclusionMap"), EmptyColorMapWhite);
 		PreviewMaterial->SetTextureParameterValue(TEXT("ColorMap"), EmptyColorMapWhite);
@@ -199,7 +213,7 @@ void UBakeMeshAttributeMapsToolBase::UpdatePreview(const int PreviewIdx)
 		break;
 	case EBakeMapType::BentNormal:
 	{
-		UTexture2D* AOMap = CachedMapIndices.Contains(EBakeMapType::AmbientOcclusion) ? CachedMaps[CachedMapIndices[EBakeMapType::AmbientOcclusion]] : EmptyColorMapWhite;
+		UTexture2D* AOMap = CachedMaps.Contains(EBakeMapType::AmbientOcclusion) ? CachedMaps[EBakeMapType::AmbientOcclusion] : EmptyColorMapWhite;
 		BentNormalPreviewMaterial->SetTextureParameterValue(TEXT("NormalMap"), EmptyNormalMap);
 		BentNormalPreviewMaterial->SetTextureParameterValue(TEXT("OcclusionMap"), AOMap);
 		BentNormalPreviewMaterial->SetTextureParameterValue(TEXT("ColorMap"), EmptyColorMapWhite);
@@ -208,13 +222,13 @@ void UBakeMeshAttributeMapsToolBase::UpdatePreview(const int PreviewIdx)
 		break;
 	}	
 	case EBakeMapType::Curvature:
-	case EBakeMapType::NormalImage:
-	case EBakeMapType::FaceNormalImage:
-	case EBakeMapType::PositionImage:
+	case EBakeMapType::ObjectSpaceNormal:
+	case EBakeMapType::FaceNormal:
+	case EBakeMapType::Position:
 	case EBakeMapType::MaterialID:
-	case EBakeMapType::Texture2DImage:
+	case EBakeMapType::Texture:
 	case EBakeMapType::MultiTexture:
-	case EBakeMapType::VertexColorImage:
+	case EBakeMapType::VertexColor:
 		PreviewMaterial->SetTextureParameterValue(TEXT("NormalMap"), EmptyNormalMap);
 		PreviewMaterial->SetTextureParameterValue(TEXT("OcclusionMap"), EmptyColorMapWhite);
 		PreviewMaterial->SetTextureParameterValue(TEXT("ColorMap"), PreviewMap);
@@ -251,11 +265,11 @@ void UBakeMeshAttributeMapsToolBase::OnMapsUpdated(const TUniquePtr<UE::Geometry
 			const bool bConvertSourceToSRGB = bConvertToSRGB && SourceDataFormat == TSF_BGRA8;
 			TextureBuilder.CopyImageToSourceData(*NewResult->GetBakeResults(EvalIdx)[ResultIdx], SourceDataFormat, bConvertSourceToSRGB);
 
-			// The CachedMap & CachedMapIndices can be thrown out of sync if updated during
-			// a background compute. Validate the computed type against our cached maps.
-			if (CachedMapIndices.Contains(MapType))
+			// The CachedMap can be thrown out of sync if updated during a background
+			// compute. Validate the computed type against our cached maps.
+			if (CachedMaps.Contains(MapType))
 			{
-				CachedMaps[CachedMapIndices[MapType]] = TextureBuilder.GetTexture2D();
+				CachedMaps[MapType] = TextureBuilder.GetTexture2D();
 			}
 		};
 
@@ -263,7 +277,7 @@ void UBakeMeshAttributeMapsToolBase::OnMapsUpdated(const TUniquePtr<UE::Geometry
 		{
 		case EMeshMapEvaluatorType::Normal:
 		{
-			constexpr EBakeMapType MapType = EBakeMapType::TangentSpaceNormalMap;
+			constexpr EBakeMapType MapType = EBakeMapType::TangentSpaceNormal;
 			UpdateCachedMap(MapType, Format, 0);
 			break;
 		}
@@ -297,19 +311,19 @@ void UBakeMeshAttributeMapsToolBase::OnMapsUpdated(const TUniquePtr<UE::Geometry
 			switch (PropertyEval->Property)
 			{
 			case EMeshPropertyMapType::Normal:
-				MapType = EBakeMapType::NormalImage;
+				MapType = EBakeMapType::ObjectSpaceNormal;
 				break;
 			case EMeshPropertyMapType::FacetNormal:
-				MapType = EBakeMapType::FaceNormalImage;
+				MapType = EBakeMapType::FaceNormal;
 				break;
 			case EMeshPropertyMapType::Position:
-				MapType = EBakeMapType::PositionImage;
+				MapType = EBakeMapType::Position;
 				break;
 			case EMeshPropertyMapType::MaterialID:
 				MapType = EBakeMapType::MaterialID;
 				break;
 			case EMeshPropertyMapType::VertexColor:
-				MapType = EBakeMapType::VertexColorImage;
+				MapType = EBakeMapType::VertexColor;
 				break;
 			case EMeshPropertyMapType::UVPosition:
 			default:
@@ -321,8 +335,8 @@ void UBakeMeshAttributeMapsToolBase::OnMapsUpdated(const TUniquePtr<UE::Geometry
 		}
 		case EMeshMapEvaluatorType::ResampleImage:
 		{
-			constexpr EBakeMapType MapType = EBakeMapType::Texture2DImage;
-			UpdateCachedMap(EBakeMapType::Texture2DImage, Format, 0);
+			constexpr EBakeMapType MapType = EBakeMapType::Texture;
+			UpdateCachedMap(MapType, Format, 0);
 			break;
 		}
 		case EMeshMapEvaluatorType::MultiResampleImage:
@@ -343,7 +357,7 @@ void UBakeMeshAttributeMapsToolBase::OnMapsUpdated(const TUniquePtr<UE::Geometry
 
 EBakeMapType UBakeMeshAttributeMapsToolBase::GetMapTypes(const int32& MapTypes)
 {
-	EBakeMapType OutMapTypes = (EBakeMapType)MapTypes & EBakeMapType::All;
+	EBakeMapType OutMapTypes = static_cast<EBakeMapType>(MapTypes);
 	// Force AO bake for BentNormal preview
 	if ((bool)(OutMapTypes & EBakeMapType::BentNormal))
 	{
@@ -353,21 +367,6 @@ EBakeMapType UBakeMeshAttributeMapsToolBase::GetMapTypes(const int32& MapTypes)
 }
 
 
-TArray<EBakeMapType> UBakeMeshAttributeMapsToolBase::GetMapTypesArray(const int32& MapTypes)
-{
-	TArray<EBakeMapType> OutMapTypes;
-	int32 Bitfield = MapTypes & (int32)EBakeMapType::All;
-	for (int32 BitIdx = 0; Bitfield; Bitfield >>= 1, ++BitIdx)
-	{
-		if (Bitfield & 1)
-		{
-			OutMapTypes.Add((EBakeMapType)(1 << BitIdx));
-		}
-	}
-	return OutMapTypes;
-}
-
-
 FTexture2DBuilder::ETextureType UBakeMeshAttributeMapsToolBase::GetTextureType(const EBakeMapType MapType, const EBakeTextureFormat MapFormat)
 {
 	FTexture2DBuilder::ETextureType TexType = FTexture2DBuilder::ETextureType::Color;
@@ -376,7 +375,7 @@ FTexture2DBuilder::ETextureType UBakeMeshAttributeMapsToolBase::GetTextureType(c
 	default:
 		checkNoEntry();
 		break;
-	case EBakeMapType::TangentSpaceNormalMap:
+	case EBakeMapType::TangentSpaceNormal:
 		TexType = FTexture2DBuilder::ETextureType::NormalMap;
 		break;
 	case EBakeMapType::AmbientOcclusion:
@@ -386,15 +385,15 @@ FTexture2DBuilder::ETextureType UBakeMeshAttributeMapsToolBase::GetTextureType(c
 		TexType = FTexture2DBuilder::ETextureType::NormalMap;
 		break;
 	case EBakeMapType::Curvature:
-	case EBakeMapType::NormalImage:
-	case EBakeMapType::FaceNormalImage:
-	case EBakeMapType::PositionImage:
+	case EBakeMapType::ObjectSpaceNormal:
+	case EBakeMapType::FaceNormal:
+	case EBakeMapType::Position:
 		TexType = FTexture2DBuilder::ETextureType::ColorLinear;
 		break;
 	case EBakeMapType::MaterialID:
-	case EBakeMapType::VertexColorImage:
+	case EBakeMapType::VertexColor:
 		break;
-	case EBakeMapType::Texture2DImage:
+	case EBakeMapType::Texture:
 	case EBakeMapType::MultiTexture:
 		// For texture output with 16-bit source data, output HDR texture
 		if (MapFormat == EBakeTextureFormat::ChannelBits16)
@@ -414,7 +413,7 @@ void UBakeMeshAttributeMapsToolBase::GetTextureName(const EBakeMapType MapType,
 	default:
 		checkNoEntry();
 		break;
-	case EBakeMapType::TangentSpaceNormalMap:
+	case EBakeMapType::TangentSpaceNormal:
 		TexName = FString::Printf(TEXT("%s_Normals"), *BaseName);
 		break;
 	case EBakeMapType::AmbientOcclusion:
@@ -426,26 +425,26 @@ void UBakeMeshAttributeMapsToolBase::GetTextureName(const EBakeMapType MapType,
 	case EBakeMapType::Curvature:
 		TexName = FString::Printf(TEXT("%s_Curvature"), *BaseName);
 		break;
-	case EBakeMapType::NormalImage:
-		TexName = FString::Printf(TEXT("%s_NormalImg"), *BaseName);
+	case EBakeMapType::ObjectSpaceNormal:
+		TexName = FString::Printf(TEXT("%s_NormalsObject"), *BaseName);
 		break;
-	case EBakeMapType::FaceNormalImage:
-		TexName = FString::Printf(TEXT("%s_FaceNormalImg"), *BaseName);
+	case EBakeMapType::FaceNormal:
+		TexName = FString::Printf(TEXT("%s_FaceNormal"), *BaseName);
 		break;
 	case EBakeMapType::MaterialID:
-		TexName = FString::Printf(TEXT("%s_MaterialIDImg"), *BaseName);
+		TexName = FString::Printf(TEXT("%s_MaterialID"), *BaseName);
 		break;
-	case EBakeMapType::VertexColorImage:
-		TexName = FString::Printf(TEXT("%s_VertexColorIDImg"), *BaseName);
+	case EBakeMapType::VertexColor:
+		TexName = FString::Printf(TEXT("%s_VertexColorID"), *BaseName);
 		break;
-	case EBakeMapType::PositionImage:
-		TexName = FString::Printf(TEXT("%s_PositionImg"), *BaseName);
+	case EBakeMapType::Position:
+		TexName = FString::Printf(TEXT("%s_Position"), *BaseName);
 		break;
-	case EBakeMapType::Texture2DImage:
-		TexName = FString::Printf(TEXT("%s_TextureImg"), *BaseName);
+	case EBakeMapType::Texture:
+		TexName = FString::Printf(TEXT("%s_Texture"), *BaseName);
 		break;
 	case EBakeMapType::MultiTexture:
-		TexName = FString::Printf(TEXT("%s_MultiTextureImg"), *BaseName);
+		TexName = FString::Printf(TEXT("%s_MultiTexture"), *BaseName);
 		break;
 	}
 }
diff --git a/Engine/Plugins/Experimental/MeshModelingToolsetExp/Source/MeshModelingToolsExp/Private/BakeMultiMeshAttributeMapsTool.cpp b/Engine/Plugins/Experimental/MeshModelingToolsetExp/Source/MeshModelingToolsExp/Private/BakeMultiMeshAttributeMapsTool.cpp
--- a/Engine/Plugins/Experimental/MeshModelingToolsetExp/Source/MeshModelingToolsExp/Private/BakeMultiMeshAttributeMapsTool.cpp
+++ b/Engine/Plugins/Experimental/MeshModelingToolsetExp/Source/MeshModelingToolsExp/Private/BakeMultiMeshAttributeMapsTool.cpp
@@ -73,6 +73,12 @@ TArray<FString> UBakeMultiMeshAttributeMapsToolProperties::GetUVLayerNamesFunc()
 }
 
 
+const TArray<FString>& UBakeMultiMeshAttributeMapsToolProperties::GetMapPreviewNamesFunc()
+{
+	return MapPreviewNamesList;
+}
+
+
 /**
  * MeshSceneAdapter bake detail sampler for baking N detail meshes to 1 target mesh.
  */
@@ -351,13 +357,13 @@ public:
 		{
 			switch (BakeCacheSettings.BakeMapTypes & MapType)
 			{
-			case EBakeMapType::TangentSpaceNormalMap:
+			case EBakeMapType::TangentSpaceNormal:
 			{
 				TSharedPtr<FMeshNormalMapEvaluator, ESPMode::ThreadSafe> NormalEval = MakeShared<FMeshNormalMapEvaluator, ESPMode::ThreadSafe>();
 				Baker->AddEvaluator(NormalEval);
 				break;	
 			}
-			case EBakeMapType::Texture2DImage:
+			case EBakeMapType::Texture:
 			{
 				TSharedPtr<FMeshResampleImageEvaluator, ESPMode::ThreadSafe> TextureEval = MakeShared<FMeshResampleImageEvaluator, ESPMode::ThreadSafe>();
 				DetailSampler.SetColorMaps(CachedMeshToColorImageMap);
@@ -414,24 +420,11 @@ void UBakeMultiMeshAttributeMapsTool::Setup()
 	// Setup tool property sets
 	Settings = NewObject<UBakeMultiMeshAttributeMapsToolProperties>(this);
 	Settings->RestoreProperties(this);
-	Settings->UVLayerNamesList.Reset();
-	int32 FoundIndex = -1;
-	for (int32 k = 0; k < BaseMesh.Attributes()->NumUVLayers(); ++k)
-	{
-		Settings->UVLayerNamesList.Add(FString::FromInt(k));
-		if (Settings->UVLayer == Settings->UVLayerNamesList.Last())
-		{
-			FoundIndex = k;
-		}
-	}
-	if (FoundIndex == -1)
-	{
-		Settings->UVLayer = Settings->UVLayerNamesList[0];
-	}
+	UpdateUVLayerNames(Settings, BaseMesh);
 	AddToolPropertySource(Settings);
 
 	Settings->WatchProperty(Settings->MapTypes, [this](int32) { bInputsDirty = true; UpdateOnModeChange(); });
-	Settings->WatchProperty(Settings->MapPreview, [this](int32) { UpdateVisualization(); GetToolManager()->PostInvalidation(); });
+	Settings->WatchProperty(Settings->MapPreview, [this](FString) { UpdateVisualization(); GetToolManager()->PostInvalidation(); });
 	Settings->WatchProperty(Settings->Resolution, [this](EBakeTextureResolution) { bInputsDirty = true; });
 	Settings->WatchProperty(Settings->SourceFormat, [this](EBakeTextureFormat) { bInputsDirty = true; });
 	Settings->WatchProperty(Settings->UVLayer, [this](FString) { bInputsDirty = true; });
@@ -499,10 +492,9 @@ bool UBakeMultiMeshAttributeMapsTool::CanAccept() const
 	if (bCanAccept)
 	{
 		// Allow Accept if all non-None types have valid results.
-		const int NumResults = Settings->Result.Num();
-		for (int ResultIdx = 0; ResultIdx < NumResults; ++ResultIdx)
+		for (const TTuple<EBakeMapType, TObjectPtr<UTexture2D>>& Result : Settings->Result)
 		{
-			bCanAccept = bCanAccept && Settings->Result[ResultIdx];
+			bCanAccept = bCanAccept && Result.Get<1>();
 		}
 	}
 	return bCanAccept;
@@ -516,13 +508,13 @@ TUniquePtr<UE::Geometry::TGenericDataOperator<FMeshMapBaker>> UBakeMultiMeshAttr
 	Op->BaseMesh = &BaseMesh;
 	Op->BakeCacheSettings = CachedBakeCacheSettings;
 
-	constexpr EBakeMapType RequiresTangents = EBakeMapType::TangentSpaceNormalMap | EBakeMapType::BentNormal;
+	constexpr EBakeMapType RequiresTangents = EBakeMapType::TangentSpaceNormal | EBakeMapType::BentNormal;
 	if ((bool)(CachedBakeCacheSettings.BakeMapTypes & RequiresTangents))
 	{
 		Op->BaseMeshTangents = BaseMeshTangents;
 	}
 
-	if ((bool)(CachedBakeCacheSettings.BakeMapTypes & EBakeMapType::Texture2DImage))
+	if ((bool)(CachedBakeCacheSettings.BakeMapTypes & EBakeMapType::Texture))
 	{
 		Op->CachedColorImages = CachedColorImages;
 		Op->CachedMeshToColorImageMap = CachedMeshToColorImagesMap;
@@ -558,12 +550,11 @@ void UBakeMultiMeshAttributeMapsTool::Shutdown(EToolShutdownType ShutdownType)
 		FString BaseName = UE::ToolTarget::GetTargetActor(Targets[0])->GetName();
 
 		bool bCreatedAssetOK = true;
-		const int NumResults = Settings->Result.Num();
-		for (int ResultIdx = 0; ResultIdx < NumResults; ++ResultIdx)
+		for (const TTuple<EBakeMapType, TObjectPtr<UTexture2D>>& Result : Settings->Result)
 		{
 			FString TexName;
-			GetTextureName(ResultTypes[ResultIdx], BaseName, TexName);
-			bCreatedAssetOK = bCreatedAssetOK && UE::Modeling::CreateTextureObject(GetToolManager(), FCreateTextureObjectParams{ 0, StaticMeshAsset->GetWorld(), StaticMeshAsset, TexName, Settings->Result[ResultIdx] }).IsOK();
+			GetTextureName(Result.Get<0>(), BaseName, TexName);
+			bCreatedAssetOK = bCreatedAssetOK && UE::Modeling::CreateTextureObject(GetToolManager(), FCreateTextureObjectParams{ 0, StaticMeshAsset->GetWorld(), StaticMeshAsset, TexName, Result.Get<1>() }).IsOK();
 		}
 		ensure(bCreatedAssetOK);
 	}
@@ -651,7 +642,7 @@ EBakeOpState UBakeMultiMeshAttributeMapsTool::UpdateResult_DetailMeshes()
 		ActorToDataMap.Emplace(ActorComponent, Idx);
 		
 		// Color map data
-		if ((bool)(CachedBakeCacheSettings.BakeMapTypes & EBakeMapType::Texture2DImage))
+		if ((bool)(CachedBakeCacheSettings.BakeMapTypes & EBakeMapType::Texture))
 		{
 			UTexture2D* ColorMapSourceTexture = DetailProps->DetailProperties[Idx].DetailColorMap;
 			const int ColorMapUVLayer = DetailProps->DetailProperties[Idx].DetailColorMapUVLayer;
@@ -680,7 +671,7 @@ EBakeOpState UBakeMultiMeshAttributeMapsTool::UpdateResult_DetailMeshes()
 		{
 			if (const int* DataIndex = ActorToDataMap.Find(ChildMesh->SourceComponent))
 			{
-				if ((bool)(CachedBakeCacheSettings.BakeMapTypes & EBakeMapType::Texture2DImage))
+				if ((bool)(CachedBakeCacheSettings.BakeMapTypes & EBakeMapType::Texture))
 				{
 					CachedMeshToColorImagesMap.Emplace(
 					ChildMesh->MeshSpatial,
@@ -699,64 +690,34 @@ void UBakeMultiMeshAttributeMapsTool::UpdateVisualization()
 {
 	PreviewMesh->SetOverrideRenderMaterial(PreviewMaterial);
 
-	// Map CachedMaps to Settings->Result
-	const int NumResults = Settings->Result.Num();
-	for (int ResultIdx = 0; ResultIdx < NumResults; ResultIdx++)
+	// Populate Settings->Result from CachedMaps
+	for (const TTuple<EBakeMapType, TObjectPtr<UTexture2D>>& Map : CachedMaps)
 	{
-		Settings->Result[ResultIdx] = CachedMaps[CachedMapIndices[ResultTypes[ResultIdx]]];
+		if (Settings->Result.Contains(Map.Get<0>()))
+		{
+			Settings->Result[Map.Get<0>()] = Map.Get<1>();
+		}
 	}
 
-	// Set the preview material according to the preview index.
-	if (Settings->MapPreview >= 0 && Settings->MapPreview < Settings->Result.Num())
+	// Set the preview material according to the preview name.
+	if (const EBakeMapType* PreviewType = NameToMapTypeMap.Find(Settings->MapPreview))
 	{
-		UpdatePreview(Settings->MapPreview);
+		UpdatePreview(*PreviewType);
 	}
 }
 
 
 void UBakeMultiMeshAttributeMapsTool::UpdateOnModeChange()
 {
-	OnMapTypesUpdated(Settings->MapTypes);
-	
-	for (const EBakeMapType MapType : ALL_BAKE_MAP_TYPES)
-	{
-		switch ((EBakeMapType)Settings->MapTypes & MapType)
-		{
-		case EBakeMapType::TangentSpaceNormalMap:
-			//SetToolPropertySourceEnabled(NormalMapProps, true);
-			break;
-		case EBakeMapType::AmbientOcclusion:
-		case EBakeMapType::BentNormal:
-		case EBakeMapType::Occlusion:
-			//SetToolPropertySourceEnabled(OcclusionMapProps, true);
-			break;
-		case EBakeMapType::Curvature:
-			//SetToolPropertySourceEnabled(CurvatureMapProps, true);
-			break;
-		case EBakeMapType::NormalImage:
-		case EBakeMapType::FaceNormalImage:
-		case EBakeMapType::PositionImage:
-		case EBakeMapType::MaterialID:
-		case EBakeMapType::VertexColorImage:
-			break;
-		case EBakeMapType::Texture2DImage:
-			//SetToolPropertySourceEnabled(Texture2DProps, true);
-			break;
-		case EBakeMapType::MultiTexture:
-			//SetToolPropertySourceEnabled(MultiTextureProps, true);
-			break;
-		default:
-			break;
-		}
-	}
-	
-	InvalidateResults();
+	OnMapTypesUpdated(Settings);
 }
 
 void UBakeMultiMeshAttributeMapsTool::InvalidateResults()
 {
-	Settings->Result.Empty();
-	Settings->Result.SetNum(ResultTypes.Num());
+	for (TTuple<EBakeMapType, TObjectPtr<UTexture2D>>& Result : Settings->Result)
+	{
+		Result.Get<1>() = nullptr;
+	}
 }
 
 
diff --git a/Engine/Plugins/Experimental/MeshModelingToolsetExp/Source/MeshModelingToolsExp/Public/BakeMeshAttributeMapsTool.h b/Engine/Plugins/Experimental/MeshModelingToolsetExp/Source/MeshModelingToolsExp/Public/BakeMeshAttributeMapsTool.h
--- a/Engine/Plugins/Experimental/MeshModelingToolsetExp/Source/MeshModelingToolsExp/Public/BakeMeshAttributeMapsTool.h
+++ b/Engine/Plugins/Experimental/MeshModelingToolsetExp/Source/MeshModelingToolsExp/Public/BakeMeshAttributeMapsTool.h
@@ -56,9 +56,9 @@ public:
 	UPROPERTY(EditAnywhere, Category = MapSettings, meta=(Bitmask, BitmaskEnum=EBakeMapType))
 	int32 MapTypes = (int32) EBakeMapType::None;
 
-	/** The map type index to preview */
-	UPROPERTY(EditAnywhere, Category = MapSettings, meta=(ArrayClamp="Result", TransientToolProperty))
-	int MapPreview = 0;
+	/** The map type to preview */
+	UPROPERTY(EditAnywhere, Category = MapSettings, meta=(TransientToolProperty, GetOptions = GetMapPreviewNamesFunc))
+	FString MapPreview;
 
 	/** The pixel resolution of the generated map */
 	UPROPERTY(EditAnywhere, Category = MapSettings, meta = (TransientToolProperty))
@@ -84,13 +84,17 @@ public:
 	FString UVLayer;
 
 	UFUNCTION()
-	TArray<FString> GetUVLayerNamesFunc();
+	const TArray<FString>& GetUVLayerNamesFunc();
 	UPROPERTY(meta = (TransientToolProperty))
 	TArray<FString> UVLayerNamesList;
 
-	UPROPERTY(VisibleAnywhere, Category = MapSettings, meta = (TransientToolProperty))
-	TArray<TObjectPtr<UTexture2D>> Result;
+	UFUNCTION()
+	const TArray<FString>& GetMapPreviewNamesFunc();
+	UPROPERTY(meta = (TransientToolProperty))
+	TArray<FString> MapPreviewNamesList;
 
+	UPROPERTY(VisibleAnywhere, Category = MapSettings, meta = (TransientToolProperty))
+	TMap<EBakeMapType, TObjectPtr<UTexture2D>> Result;
 };
 
 
@@ -128,9 +132,6 @@ protected:
 	UPROPERTY()
 	TObjectPtr<UDetailMeshToolProperties> DetailMeshProps;
 
-	UPROPERTY()
-	TObjectPtr<UBakedNormalMapToolProperties> NormalMapProps;
-
 	UPROPERTY()
 	TObjectPtr<UBakedOcclusionMapToolProperties> OcclusionMapProps;
 
diff --git a/Engine/Plugins/Experimental/MeshModelingToolsetExp/Source/MeshModelingToolsExp/Public/BakeMeshAttributeMapsToolBase.h b/Engine/Plugins/Experimental/MeshModelingToolsetExp/Source/MeshModelingToolsExp/Public/BakeMeshAttributeMapsToolBase.h
--- a/Engine/Plugins/Experimental/MeshModelingToolsetExp/Source/MeshModelingToolsExp/Public/BakeMeshAttributeMapsToolBase.h
+++ b/Engine/Plugins/Experimental/MeshModelingToolsetExp/Source/MeshModelingToolsExp/Public/BakeMeshAttributeMapsToolBase.h
@@ -28,27 +28,27 @@ enum class EBakeMapType
 {
 	None                   = 0 UMETA(Hidden),
 	/* Sample normals from the detail mesh in tangent space */
-	TangentSpaceNormalMap  = 1 << 0 UMETA(DisplayName="Tangent Space Normals"),
+	TangentSpaceNormal     = 1 << 0,
 	/* Sample ambient occlusion from the detail mesh */
-	AmbientOcclusion       = 1 << 1 UMETA(DisplayName="Ambient Occlusion"),
+	AmbientOcclusion       = 1 << 1,
 	/* Sample normals skewed towards the least occluded direction from the detail mesh */
-	BentNormal             = 1 << 2 UMETA(DisplayName="Bent Normals"),
+	BentNormal             = 1 << 2,
 	/* Sample mesh curvatures from the detail mesh */
-	Curvature              = 1 << 3 UMETA(DisplayName="Curvature"),
+	Curvature              = 1 << 3,
 	/* Sample a source texture from the detail mesh UVs */
-	Texture2DImage         = 1 << 4 UMETA(DisplayName="Texture"),
+	Texture                = 1 << 4,
 	/* Sample object space normals from the detail mesh */
-	NormalImage            = 1 << 5 UMETA(DisplayName="Object Space Normals"),
+	ObjectSpaceNormal      = 1 << 5,
 	/* Sample object space face normals from the detail mesh */
-	FaceNormalImage        = 1 << 6 UMETA(DisplayName="Face Normals"),
+	FaceNormal             = 1 << 6,
 	/* Sample bounding box relative positions from the detail mesh */
-	PositionImage          = 1 << 7 UMETA(DisplayName="Position"),
+	Position               = 1 << 7,
 	/* Sample material IDs as unique colors from the detail mesh */
 	MaterialID             = 1 << 8 UMETA(DisplayName="Material ID"),
 	/* Sample a source texture per material ID on the detail mesh */
-	MultiTexture           = 1 << 9 UMETA(DisplayName="MultiTexture"),
+	MultiTexture           = 1 << 9,
 	/* Sample the interpolated vertex colors from the detail mesh */
-	VertexColorImage       = 1 << 10 UMETA(DisplayName="Vertex Colors"),
+	VertexColor            = 1 << 10,
 	Occlusion              = (AmbientOcclusion | BentNormal) UMETA(Hidden),
 	All                    = 0x7FF UMETA(Hidden)
 };
@@ -63,16 +63,16 @@ ENUM_CLASS_FLAGS(EBakeMapType);
 // iteration loop.
 static constexpr EBakeMapType ALL_BAKE_MAP_TYPES[] =
 {
-	EBakeMapType::TangentSpaceNormalMap,
+	EBakeMapType::TangentSpaceNormal,
 	EBakeMapType::Occlusion, // (AmbientOcclusion | BentNormal)
 	EBakeMapType::Curvature,
-	EBakeMapType::Texture2DImage,
-	EBakeMapType::NormalImage,
-	EBakeMapType::FaceNormalImage,
-	EBakeMapType::PositionImage,
+	EBakeMapType::Texture,
+	EBakeMapType::ObjectSpaceNormal,
+	EBakeMapType::FaceNormal,
+	EBakeMapType::Position,
 	EBakeMapType::MaterialID,
 	EBakeMapType::MultiTexture,
-	EBakeMapType::VertexColorImage
+	EBakeMapType::VertexColor
 };
 
 
@@ -141,6 +141,14 @@ public:
 
 	void SetWorld(UWorld* World);
 
+public:
+	/** @return An enumerated map type from a bitfield */
+	static EBakeMapType GetMapTypes(const int32& MapTypes);
+
+	/** Internal maps of name to BakeMapType */
+	static TMap<FString, EBakeMapType> NameToMapTypeMap;
+	static TMap<EBakeMapType, FString> MapTypeToNameMap;
+
 protected:
 	//
 	// Tool property sets
@@ -217,13 +225,18 @@ protected:
 	/**
 	 * To be invoked by client when bake map types change.
 	 */
-	void OnMapTypesUpdated(int32 MapTypes);
+	template <typename PropertySet>
+	void OnMapTypesUpdated(PropertySet& Properties);
 
 	//
 	// Background compute
 	//
 	TUniquePtr<TGenericDataBackgroundCompute<UE::Geometry::FMeshMapBaker>> Compute = nullptr;
 
+	/** Internal cache of bake texture results. */
+	UPROPERTY()
+	TMap<EBakeMapType, TObjectPtr<UTexture2D>> CachedMaps;
+
 	/**
 	 * Retrieves the result of the FMeshMapBaker and generates UTexture2D into the CachedMaps.
 	 * It is the responsibility of the client to ensure that CachedMaps is appropriately sized for
@@ -233,33 +246,42 @@ protected:
 	 */
 	void OnMapsUpdated(const TUniquePtr<UE::Geometry::FMeshMapBaker>& NewResult, EBakeTextureFormat Format);
 
+
 	/**
 	 * Update the preview material parameters for a given a result index.
-	 * @param PreviewIdx index into the ResultTypes array to preview
+	 * @param PreviewMapType EBakeMapType to preview
 	 */
-	void UpdatePreview(int PreviewIdx);
+	void UpdatePreview(EBakeMapType PreviewMapType);
+	
+
+	/**
+	 * Updates a tool property set's MapPreviewNamesList from the list of
+	 * active map types. Also updates the MapPreview property if the current
+	 * preview option is no longer available.
+	 *
+	 * @param Properties the UInteractiveToolPropertySet to update.
+	 */
+	template <typename PropertySet>
+	void UpdatePreviewNames(PropertySet& Properties);
 
-	/** Internal cache of bake texture results. */
-	UPROPERTY()
-	TArray<TObjectPtr<UTexture2D>> CachedMaps;
 
-	/** Internal map of map type to index into CachedMaps array */
-	TMap<EBakeMapType, int32> CachedMapIndices;
+	/**
+	 * Updates a tool property set's UVLayerNamesList from the list of UV layers
+	 * on a given mesh. Also updates the UVLayer property if the current UV layer
+	 * is no longer available.
+	 *
+	 * @param Properties the UInteractiveToolPropertySet to update.
+	 * @param Mesh the mesh to query
+	 */
+	template <typename PropertySet>
+	static void UpdateUVLayerNames(PropertySet& Properties, const FDynamicMesh3& Mesh);
 
-	/** Internal array of CachedMaps index to map type */
-	TArray<EBakeMapType> ResultTypes;
 	
 	//
 	// Utilities
 	//
 	const bool bPreferPlatformData = false;
 	
-	/** @return An enumerated map type from a bitfield */
-	static EBakeMapType GetMapTypes(const int32& MapTypes);
-
-	/** @return An array of map types from a bitfield */
-	static TArray<EBakeMapType> GetMapTypesArray(const int32& MapTypes);
-
 	/** @return the Texture2D type for a given map type */
 	static UE::Geometry::FTexture2DBuilder::ETextureType GetTextureType(EBakeMapType MapType, EBakeTextureFormat MapFormat);
 
@@ -282,6 +304,15 @@ protected:
 	template <typename ProcessFn>
 	static void ProcessComponentTextures(const UPrimitiveComponent* Component, ProcessFn&& ProcessFunc);
 
+	/**
+	 * @param Properties the tool property set to validate the map type against.
+	 * @param MapType the map type to validate.
+	 * @return true if MapType was requested from the given PropertySet
+	 */
+	template <typename PropertySet>
+	static bool IsRequestedMapType(PropertySet& Properties, EBakeMapType MapType);
+	
+
 	// empty maps are shown when nothing is computed
 	UPROPERTY()
 	TObjectPtr<UTexture2D> EmptyNormalMap;
@@ -296,6 +327,97 @@ protected:
 };
 
 
+template <typename PropertySet>
+void UBakeMeshAttributeMapsToolBase::OnMapTypesUpdated(PropertySet& Properties)
+{
+	const EBakeMapType BakeMapTypes = GetMapTypes(Properties->MapTypes);
+	
+	// Use the processed bitfield which may contain additional targets
+	// (ex. AO if BentNormal was requested).
+	CachedMaps.Empty();
+	for (EBakeMapType MapType : ALL_BAKE_MAP_TYPES)
+	{
+		if (MapType == EBakeMapType::Occlusion)
+		{
+			if ((bool)(BakeMapTypes & EBakeMapType::AmbientOcclusion))
+			{
+				CachedMaps.Add(EBakeMapType::AmbientOcclusion, nullptr);
+			}
+			if ((bool)(BakeMapTypes & EBakeMapType::BentNormal))
+			{
+				CachedMaps.Add(EBakeMapType::BentNormal, nullptr);
+			}
+		}
+		else if( (bool)(BakeMapTypes & MapType) )
+		{
+			CachedMaps.Add(MapType, nullptr);
+		}
+	}
+
+	// Initialize Properties->Result with requested MapTypes
+	for (const TTuple<EBakeMapType, TObjectPtr<UTexture2D>>& Map : CachedMaps)
+	{
+		// Only populate map types that were requested. Some map types like
+		// AO may have only been added for preview of other types (ex. BentNormal)
+		if (IsRequestedMapType(Properties, Map.Get<0>()))
+		{
+			Properties->Result.Add(Map.Get<0>(), nullptr);
+		}
+	}
+
+	UpdatePreviewNames(Properties);
+}
+
+
+template <typename PropertySet>
+void UBakeMeshAttributeMapsToolBase::UpdatePreviewNames(PropertySet& Properties)
+{
+	// Update our preview names list.
+	Properties->MapPreviewNamesList.Reset();
+	bool bFoundMapType = false;
+	for (const TTuple<EBakeMapType, TObjectPtr<UTexture2D>>& Map : CachedMaps)
+	{
+		// Only populate map types that were requested. Some map types like
+		// AO may have only been added for preview of other types (ex. BentNormal)
+		if (IsRequestedMapType(Properties, Map.Get<0>()))
+		{
+			if (const FString* MapTypeName = MapTypeToNameMap.Find(Map.Get<0>()))
+			{
+				Properties->MapPreviewNamesList.Add(*MapTypeName);
+				if (Properties->MapPreview == Properties->MapPreviewNamesList.Last())
+				{
+					bFoundMapType = true;
+				}
+			}
+		}
+	}
+	if (!bFoundMapType)
+	{
+		Properties->MapPreview = Properties->MapPreviewNamesList.Num() > 0 ? Properties->MapPreviewNamesList[0] : TEXT("");
+	}
+}
+
+
+template <typename PropertySet>
+void UBakeMeshAttributeMapsToolBase::UpdateUVLayerNames(PropertySet& Properties, const FDynamicMesh3& Mesh)
+{
+	Properties->UVLayerNamesList.Reset();
+	int32 FoundIndex = -1;
+	for (int32 k = 0; k < Mesh.Attributes()->NumUVLayers(); ++k)
+	{
+		Properties->UVLayerNamesList.Add(FString::FromInt(k));
+		if (Properties->UVLayer == Properties->UVLayerNamesList.Last())
+		{
+			FoundIndex = k;
+		}
+	}
+	if (FoundIndex == -1)
+	{
+		Properties->UVLayer = Properties->UVLayerNamesList[0];
+	}
+}
+
+
 template <typename ProcessFn>
 void UBakeMeshAttributeMapsToolBase::ProcessComponentTextures(const UPrimitiveComponent* Component, ProcessFn&& ProcessFunc)
 {
@@ -321,5 +443,15 @@ void UBakeMeshAttributeMapsToolBase::ProcessComponentTextures(const UPrimitiveCo
 	}
 }
 
+template <typename PropertySet>
+bool UBakeMeshAttributeMapsToolBase::IsRequestedMapType(PropertySet& Properties, EBakeMapType MapType)
+{
+	EBakeMapType SourceMapTypes = static_cast<EBakeMapType>(Properties->MapTypes);
+	return static_cast<bool>(SourceMapTypes & MapType);
+}
+
+
+
+
 
 
diff --git a/Engine/Plugins/Experimental/MeshModelingToolsetExp/Source/MeshModelingToolsExp/Public/BakeMultiMeshAttributeMapsTool.h b/Engine/Plugins/Experimental/MeshModelingToolsetExp/Source/MeshModelingToolsExp/Public/BakeMultiMeshAttributeMapsTool.h
--- a/Engine/Plugins/Experimental/MeshModelingToolsetExp/Source/MeshModelingToolsExp/Public/BakeMultiMeshAttributeMapsTool.h
+++ b/Engine/Plugins/Experimental/MeshModelingToolsetExp/Source/MeshModelingToolsExp/Public/BakeMultiMeshAttributeMapsTool.h
@@ -34,23 +34,35 @@ class MESHMODELINGTOOLSEXP_API UBakeMultiMeshAttributeMapsToolBuilder : public U
 
 
 // This enumeration must match EBakeMapType. This duplicate enum is
-// intended to only be used for UI and internally converted to EBakeMapType.
+// intended to only be used for UI purposes to hide unsupported map types
+// but is internally converted to EBakeMapType.
 UENUM(meta=(Bitflags, UseEnumValuesAsMaskValuesInEditor="true"))
 enum class EBakeMultiMapType
 {
 	None                   = 0 UMETA(Hidden),
-	TangentSpaceNormalMap  = 1 << 0,
+	/* Sample normals from the detail mesh in tangent space */
+	TangentSpaceNormals    = 1 << 0,
+	/* Sample ambient occlusion from the detail mesh */
 	AmbientOcclusion       = 1 << 1 UMETA(Hidden),
-	BentNormal             = 1 << 2 UMETA(Hidden),
+	/* Sample normals skewed towards the least occluded direction from the detail mesh */
+	BentNormals            = 1 << 2 UMETA(Hidden),
+	/* Sample mesh curvatures from the detail mesh */
 	Curvature              = 1 << 3 UMETA(Hidden),
-	Texture2DImage         = 1 << 4,
-	NormalImage            = 1 << 5 UMETA(Hidden),
-	FaceNormalImage        = 1 << 6 UMETA(Hidden),
-	PositionImage          = 1 << 7 UMETA(Hidden),
-	MaterialID             = 1 << 8 UMETA(Hidden),
+	/* Sample a source texture from the detail mesh UVs */
+	Texture                = 1 << 4,
+	/* Sample object space normals from the detail mesh */
+	ObjectSpaceNormals     = 1 << 5 UMETA(Hidden),
+	/* Sample object space face normals from the detail mesh */
+	FaceNormals            = 1 << 6 UMETA(Hidden),
+	/* Sample bounding box relative positions from the detail mesh */
+	Position               = 1 << 7 UMETA(Hidden),
+	/* Sample material IDs as unique colors from the detail mesh */
+	MaterialID             = 1 << 8 UMETA(DisplayName="Material ID", Hidden),
+	/* Sample a source texture per material ID on the detail mesh */
 	MultiTexture           = 1 << 9 UMETA(Hidden),
-	VertexColorImage       = 1 << 10 UMETA(Hidden),
-	Occlusion              = (AmbientOcclusion | BentNormal) UMETA(Hidden),
+	/* Sample the interpolated vertex colors from the detail mesh */
+	VertexColors           = 1 << 10 UMETA(Hidden),
+	Occlusion              = (AmbientOcclusion | BentNormals) UMETA(Hidden),
 	All                    = 0x7FF UMETA(Hidden)
 };
 ENUM_CLASS_FLAGS(EBakeMultiMapType);
@@ -67,8 +79,8 @@ public:
 	int32 MapTypes = (int32) EBakeMapType::None;
 
 	/** The map type index to preview */
-	UPROPERTY(EditAnywhere, Category = MapSettings, meta=(ArrayClamp="Result"))
-	int MapPreview = 0;
+	UPROPERTY(EditAnywhere, Category = MapSettings, meta=(TransientToolProperty, GetOptions = GetMapPreviewNamesFunc))
+	FString MapPreview;
 
 	/** The pixel resolution of the generated map */
 	UPROPERTY(EditAnywhere, Category = MapSettings, meta = (TransientToolProperty))
@@ -95,8 +107,13 @@ public:
 	UPROPERTY(meta = (TransientToolProperty))
 	TArray<FString> UVLayerNamesList;
 
+	UFUNCTION()
+	const TArray<FString>& GetMapPreviewNamesFunc();
+	UPROPERTY(meta = (TransientToolProperty))
+	TArray<FString> MapPreviewNamesList;
+
 	UPROPERTY(VisibleAnywhere, Category = MapSettings, meta = (TransientToolProperty))
-	TArray<TObjectPtr<UTexture2D>> Result;
+	TMap<EBakeMapType, TObjectPtr<UTexture2D>> Result;
 
 };
 
