ToolsFramework: move support for scene location/angle snapping from ToolContextInterfaces/IContextQueriesAPI to USceneSnappingManager, a context store object that different modes may wish to implement.
USceneSnappingManager currently only provides the current snapping interface and a no-op implementation.
Move existing EdModeInteractiveToolsContext snapping implementation to ModelingSceneSnappingManager in ModelingComponents module.
Add IContextQueriesAPI::GetCurrentSnappingSettings to allow ITF Host to specify snapping settings, implement in EdModeInteractiveToolsContext.
Update affected Tools/Gizmos, register ModelingSceneSnappingManager instance in Modeling EdMode.

#rnx

diff --git a/Engine/Plugins/Editor/ModelingToolsEditorMode/Source/ModelingToolsEditorMode/Private/ModelingToolsEditorMode.cpp b/Engine/Plugins/Editor/ModelingToolsEditorMode/Source/ModelingToolsEditorMode/Private/ModelingToolsEditorMode.cpp
--- a/Engine/Plugins/Editor/ModelingToolsEditorMode/Source/ModelingToolsEditorMode/Private/ModelingToolsEditorMode.cpp
+++ b/Engine/Plugins/Editor/ModelingToolsEditorMode/Source/ModelingToolsEditorMode/Private/ModelingToolsEditorMode.cpp
@@ -20,6 +20,7 @@
 #include "BaseGizmos/TransformGizmoUtil.h"
 #include "Selection/PersistentMeshSelectionManager.h"
 #include "Selection/StoredMeshSelectionUtil.h"
+#include "Snapping/ModelingSceneSnappingManager.h"
 
 #include "Features/IModularFeatures.h"
 #include "ModelingModeToolExtensions.h"
@@ -326,6 +327,9 @@ void UModelingToolsEditorMode::Enter()
 	// register gizmo helper
 	UE::TransformGizmoUtil::RegisterTransformGizmoContextObject(GetInteractiveToolsContext());
 
+	// register snapping manager
+	UE::Geometry::RegisterSceneSnappingManager(GetInteractiveToolsContext());
+
 	// register selection manager, if this feature is enabled in the mode settings
 	const UModelingToolsEditorModeSettings* ModelingModeSettings = GetDefault<UModelingToolsEditorModeSettings>();
 	if (ModelingModeSettings && ModelingModeSettings->bEnablePersistentSelections)
@@ -762,6 +766,9 @@ void UModelingToolsEditorMode::Exit()
 	// will be called before our Exit()
 	//UE::TransformGizmoUtil::DeregisterTransformGizmoContextObject(ToolsContext.Get());
 	
+	// deregister snapping manager
+	UE::Geometry::DeregisterSceneSnappingManager(GetInteractiveToolsContext());
+
 	// TODO: cannot deregister currently because if another mode is also registering, its Enter()
 	// will be called before our Exit()
 	UEditorModelingObjectsCreationAPI* ObjectCreationAPI = UEditorModelingObjectsCreationAPI::Find(GetInteractiveToolsContext());
diff --git a/Engine/Plugins/Experimental/MeshModelingToolsetExp/Source/ModelingComponents/Private/Snapping/ModelingSceneSnappingManager.cpp b/Engine/Plugins/Experimental/MeshModelingToolsetExp/Source/ModelingComponents/Private/Snapping/ModelingSceneSnappingManager.cpp
new file mode 100644
--- /dev/null
+++ b/Engine/Plugins/Experimental/MeshModelingToolsetExp/Source/ModelingComponents/Private/Snapping/ModelingSceneSnappingManager.cpp
@@ -0,0 +1,402 @@
+// Copyright Epic Games, Inc. All Rights Reserved.
+
+#include "Snapping/ModelingSceneSnappingManager.h"
+
+#include "InteractiveToolsContext.h"
+#include "InteractiveToolManager.h"
+#include "ContextObjectStore.h"
+
+#include "Engine/World.h"
+#include "Engine/StaticMesh.h"
+#include "Components/PrimitiveComponent.h"
+#include "Components/StaticMeshComponent.h"
+
+
+#define LOCTEXT_NAMESPACE "USceneSnappingManager"
+
+
+static float SnapToIncrement(float fValue, float fIncrement, float offset = 0)
+{
+	if (!FMath::IsFinite(fValue))
+	{
+		return 0;
+	}
+	fValue -= offset;
+	float sign = FMath::Sign(fValue);
+	fValue = FMath::Abs(fValue);
+	int nInc = (int)(fValue / fIncrement);
+	float fRem = (float)fmod(fValue, fIncrement);
+	if (fRem > fIncrement / 2)
+	{
+		++nInc;
+	}
+	return sign * (float)nInc * fIncrement + offset;
+}
+
+
+//@ todo this are mirrored from GeometryProcessing, which is still experimental...replace w/ direct calls once GP component is standardized
+static float OpeningAngleDeg(FVector A, FVector B, const FVector& P)
+{
+	A -= P;
+	A.Normalize();
+	B -= P;
+	B.Normalize();
+	float Dot = FMath::Clamp(FVector::DotProduct(A,B), -1.0f, 1.0f);
+	return acos(Dot) * (180.0f / 3.141592653589f);
+}
+
+static FVector NearestSegmentPt(FVector A, FVector B, const FVector& P)
+{
+	FVector Direction = (B - A);
+	float Length = Direction.Size();
+	Direction /= Length;
+	float t = FVector::DotProduct( (P - A), Direction);
+	if (t >= Length)
+	{
+		return B;
+	}
+	if (t <= 0)
+	{
+		return A;
+	}
+	return A + t * Direction;
+}
+
+
+
+
+void UModelingSceneSnappingManager::Initialize(TObjectPtr<UInteractiveToolsContext> ToolsContext)
+{
+	ParentContext = ToolsContext;
+
+	QueriesAPI = (ParentContext && ParentContext->ToolManager) ?
+		ParentContext->ToolManager->GetContextQueriesAPI() : nullptr;
+}
+
+void UModelingSceneSnappingManager::Shutdown()
+{
+
+}
+
+
+bool UModelingSceneSnappingManager::ExecuteSceneSnapQuery(const FSceneSnapQueryRequest& Request, TArray<FSceneSnapQueryResult>& Results) const
+{
+	switch (Request.RequestType)
+	{
+	case ESceneSnapQueryType::Position:
+		return ExecuteSceneSnapQueryPosition(Request, Results);
+		break;
+	case ESceneSnapQueryType::Rotation:
+		return ExecuteSceneSnapQueryRotation(Request, Results);
+		break;
+	default:
+		check(!"Only Position and Rotation Snap Queries are supported");
+	}
+	return false;
+}
+
+
+bool UModelingSceneSnappingManager::ExecuteSceneSnapQueryRotation(const FSceneSnapQueryRequest& Request, TArray<FSceneSnapQueryResult>& Results) const
+{
+	if (!QueriesAPI)
+	{
+		return false;
+	}
+
+	if ((Request.TargetTypes & ESceneSnapQueryTargetType::Grid) != ESceneSnapQueryTargetType::None)
+	{
+		FToolContextSnappingConfiguration SnappingConfig = QueriesAPI->GetCurrentSnappingSettings();
+		if (SnappingConfig.bEnableRotationGridSnapping)
+		{
+			FRotator Rotator(Request.DeltaRotation);
+			FRotator RotGrid = Request.RotGridSize.Get(SnappingConfig.RotationGridAngles);
+			Rotator = Rotator.GridSnap(RotGrid);
+
+			FSceneSnapQueryResult SnapResult;
+			SnapResult.TargetType = ESceneSnapQueryTargetType::Grid;
+			SnapResult.DeltaRotation = Rotator.Quaternion();
+			Results.Add(SnapResult);
+			return true;
+		}
+	}
+	return false;
+}
+
+
+
+static bool IsHiddenActor(AActor* Actor)
+{
+#if WITH_EDITOR
+	return (Actor != nullptr) && ( Actor->IsHidden() || Actor->IsHiddenEd() );
+#else
+	return (Actor != nullptr) && ( Actor->IsHidden() );
+#endif
+}
+
+static bool IsHiddenComponent(UPrimitiveComponent* Component)
+{
+#if WITH_EDITOR
+	return (Component != nullptr) && ( (Component->IsVisible() == false) || (Component->IsVisibleInEditor() == false) );
+#else
+	return (Component != nullptr) && (Component->IsVisible() == false);
+#endif
+}
+
+
+static bool IsVisibleObjectHit_Internal(const FHitResult& HitResult)
+{
+	AActor* Actor = HitResult.GetActor();
+	if (IsHiddenActor(Actor))
+	{
+		return false;
+	}
+	UPrimitiveComponent* Component = HitResult.GetComponent();
+	if ( IsHiddenComponent(Component) )
+	{
+		return false;
+	}
+	return true;
+}
+
+static bool FindNearestVisibleObjectHit_Internal(UWorld* World, FHitResult& HitResultOut, const FVector& Start, const FVector& End,
+	bool bIsSceneGeometrySnapQuery, const TArray<const UPrimitiveComponent*>* ComponentsToIgnore, 
+	const TArray<const UPrimitiveComponent*>* InvisibleComponentsToInclude)
+{
+	FCollisionObjectQueryParams ObjectQueryParams(FCollisionObjectQueryParams::AllObjects);
+	FCollisionQueryParams QueryParams = FCollisionQueryParams::DefaultQueryParam;
+	QueryParams.bTraceComplex = true;
+	QueryParams.bReturnFaceIndex = bIsSceneGeometrySnapQuery;
+
+	TArray<FHitResult> OutHits;
+	if (World->LineTraceMultiByObjectType(OutHits, Start, End, ObjectQueryParams, QueryParams) == false)
+	{
+		return false;
+	}
+
+	float NearestVisible = TNumericLimits<float>::Max();
+	for (const FHitResult& CurResult : OutHits)
+	{
+		if (CurResult.Distance < NearestVisible)
+		{
+			if (IsVisibleObjectHit_Internal(CurResult) 
+				|| (InvisibleComponentsToInclude && InvisibleComponentsToInclude->Contains(CurResult.GetComponent())))
+			{
+				if (!ComponentsToIgnore || !ComponentsToIgnore->Contains(CurResult.GetComponent()))
+				{
+					HitResultOut = CurResult;
+					NearestVisible = CurResult.Distance;
+				}
+			}
+		}
+	}
+
+	return NearestVisible < TNumericLimits<float>::Max();
+}
+
+
+
+
+
+bool UModelingSceneSnappingManager::ExecuteSceneSnapQueryPosition(const FSceneSnapQueryRequest& Request, TArray<FSceneSnapQueryResult>& Results) const
+{
+	if (!QueriesAPI)
+	{
+		return false;
+	}
+
+	int FoundResultCount = 0;
+
+	if ((Request.TargetTypes & ESceneSnapQueryTargetType::Grid) != ESceneSnapQueryTargetType::None)
+	{
+		FToolContextSnappingConfiguration SnappingConfig = QueriesAPI->GetCurrentSnappingSettings();
+
+		if (SnappingConfig.bEnablePositionGridSnapping)
+		{
+			FSceneSnapQueryResult SnapResult;
+			SnapResult.TargetType = ESceneSnapQueryTargetType::Grid;
+
+			FVector GridSize = Request.GridSize.Get(SnappingConfig.PositionGridDimensions);
+
+			SnapResult.Position.X = SnapToIncrement(Request.Position.X, GridSize.X);
+			SnapResult.Position.Y = SnapToIncrement(Request.Position.Y, GridSize.Y);
+			SnapResult.Position.Z = SnapToIncrement(Request.Position.Z, GridSize.Z);
+
+			Results.Add(SnapResult);
+			FoundResultCount++;
+		}
+	}
+
+	FViewCameraState ViewState;
+	QueriesAPI->GetCurrentViewState(ViewState);
+
+	//
+	// Run a snap query by casting ray into the world.
+	// If a hit is found, we look up what triangle was hit, and then test its vertices and edges
+	//
+
+	// cast ray into world
+	FVector RayStart = ViewState.Position;
+	FVector RayDirection = Request.Position - RayStart; RayDirection.Normalize();
+	FVector RayEnd = RayStart + HALF_WORLD_MAX * RayDirection;
+	FHitResult HitResult;
+	bool bHitWorld = FindNearestVisibleObjectHit_Internal(
+		QueriesAPI->GetCurrentEditingWorld(), 
+		HitResult, 
+		RayStart, RayEnd, 
+		true, 
+		Request.ComponentsToIgnore, 
+		Request.InvisibleComponentsToInclude);
+
+	if (bHitWorld && HitResult.FaceIndex >= 0)
+	{
+		float VisualAngle = OpeningAngleDeg(Request.Position, HitResult.ImpactPoint, RayStart);
+		//UE_LOG(LogTemp, Warning, TEXT("[HIT] visualangle %f faceindex %d"), VisualAngle, HitResult.FaceIndex);
+		if (VisualAngle < Request.VisualAngleThresholdDegrees)
+		{
+			UPrimitiveComponent* Component = HitResult.Component.Get();
+			if (Cast<UStaticMeshComponent>(Component) != nullptr)
+			{
+				// HitResult.FaceIndex is apparently an index into the TriMeshCollisionData, not sure how
+				// to directly access it. Calling GetPhysicsTriMeshData is expensive!
+				//UBodySetup* BodySetup = Cast<UStaticMeshComponent>(Component)->GetBodySetup();
+				//UObject* CDPObj = BodySetup->GetOuter();
+				//IInterface_CollisionDataProvider* CDP = Cast<IInterface_CollisionDataProvider>(CDPObj);
+				//FTriMeshCollisionData TriMesh;
+				//CDP->GetPhysicsTriMeshData(&TriMesh, true);
+				//FTriIndices Triangle = TriMesh.Indices[HitResult.FaceIndex];
+				//FVector Positions[3] = { TriMesh.Vertices[Triangle.v0], TriMesh.Vertices[Triangle.v1], TriMesh.Vertices[Triangle.v2] };
+
+				// physics collision data is created from StaticMesh RenderData
+				// so use HitResult.FaceIndex to extract triangle from the LOD0 mesh
+				// (note: this may be incorrect if there are multiple sections...in that case I think we have to
+				//  first find section whose accumulated index range would contain .FaceIndexX)
+				UStaticMesh* StaticMesh = Cast<UStaticMeshComponent>(Component)->GetStaticMesh();
+				FStaticMeshLODResources& LOD = StaticMesh->GetRenderData()->LODResources[0];
+				FIndexArrayView Indices = LOD.IndexBuffer.GetArrayView();
+				int32 TriIdx = 3 * HitResult.FaceIndex;
+				FVector Positions[3];
+				Positions[0] = LOD.VertexBuffers.PositionVertexBuffer.VertexPosition(Indices[TriIdx]);
+				Positions[1] = LOD.VertexBuffers.PositionVertexBuffer.VertexPosition(Indices[TriIdx+1]);
+				Positions[2] = LOD.VertexBuffers.PositionVertexBuffer.VertexPosition(Indices[TriIdx+2]);
+
+				// transform to world space
+				FTransform ComponentTransform = Component->GetComponentTransform();
+				Positions[0] = ComponentTransform.TransformPosition(Positions[0]);
+				Positions[1] = ComponentTransform.TransformPosition(Positions[1]);
+				Positions[2] = ComponentTransform.TransformPosition(Positions[2]);
+
+				FSceneSnapQueryResult SnapResult;
+				SnapResult.TriVertices[0] = Positions[0];
+				SnapResult.TriVertices[1] = Positions[1];
+				SnapResult.TriVertices[2] = Positions[2];
+
+				// try snapping to vertices
+				float SmallestAngle = Request.VisualAngleThresholdDegrees;
+				if ( (Request.TargetTypes & ESceneSnapQueryTargetType::MeshVertex) != ESceneSnapQueryTargetType::None)
+				{
+					for (int j = 0; j < 3; ++j)
+					{
+						VisualAngle = OpeningAngleDeg(Request.Position, Positions[j], RayStart);
+						if (VisualAngle < SmallestAngle)
+						{
+							SmallestAngle = VisualAngle;
+							SnapResult.Position = Positions[j];
+							SnapResult.TargetType = ESceneSnapQueryTargetType::MeshVertex;
+							SnapResult.TriSnapIndex = j;
+						}
+					}
+				}
+
+				// try snapping to nearest points on edges
+				if ( ((Request.TargetTypes & ESceneSnapQueryTargetType::MeshEdge) != ESceneSnapQueryTargetType::None) &&
+					(SnapResult.TargetType != ESceneSnapQueryTargetType::MeshVertex) )
+				{
+					for (int j = 0; j < 3; ++j)
+					{
+						FVector EdgeNearestPt = NearestSegmentPt(Positions[j], Positions[(j+1)%3], Request.Position);
+						VisualAngle = OpeningAngleDeg(Request.Position, EdgeNearestPt, RayStart);
+						if (VisualAngle < SmallestAngle )
+						{
+							SmallestAngle = VisualAngle;
+							SnapResult.Position = EdgeNearestPt;
+							SnapResult.TargetType = ESceneSnapQueryTargetType::MeshEdge;
+							SnapResult.TriSnapIndex = j;
+						}
+					}
+				}
+
+				// if we found a valid snap, return it
+				if (SmallestAngle < Request.VisualAngleThresholdDegrees)
+				{
+					SnapResult.TargetActor = HitResult.HitObjectHandle.FetchActor();
+					SnapResult.TargetComponent = HitResult.Component.Get();
+					Results.Add(SnapResult);
+					FoundResultCount++;
+				}
+			}
+		}
+
+	}
+
+	return (FoundResultCount > 0);
+}
+
+
+
+bool UE::Geometry::RegisterSceneSnappingManager(UInteractiveToolsContext* ToolsContext)
+{
+	if (ensure(ToolsContext))
+	{
+		UModelingSceneSnappingManager* Found = ToolsContext->ContextObjectStore->FindContext<UModelingSceneSnappingManager>();
+		if (Found == nullptr)
+		{
+			UModelingSceneSnappingManager* SelectionManager = NewObject<UModelingSceneSnappingManager>(ToolsContext->ToolManager);
+			if (ensure(SelectionManager))
+			{
+				SelectionManager->Initialize(ToolsContext);
+				ToolsContext->ContextObjectStore->AddContextObject(SelectionManager);
+				return true;
+			}
+			else
+			{
+				return false;
+			}
+		}
+		return true;
+	}
+	return false;
+}
+
+
+
+bool UE::Geometry::DeregisterSceneSnappingManager(UInteractiveToolsContext* ToolsContext)
+{
+	if (ensure(ToolsContext))
+	{
+		UModelingSceneSnappingManager* Found = ToolsContext->ContextObjectStore->FindContext<UModelingSceneSnappingManager>();
+		if (Found != nullptr)
+		{
+			Found->Shutdown();
+			ToolsContext->ContextObjectStore->RemoveContextObject(Found);
+		}
+		return true;
+	}
+	return false;
+}
+
+
+UModelingSceneSnappingManager* UE::Geometry::FindModelingSceneSnappingManager(UInteractiveToolManager* ToolManager)
+{
+	if (ensure(ToolManager))
+	{
+		UModelingSceneSnappingManager* Found = ToolManager->GetContextObjectStore()->FindContext<UModelingSceneSnappingManager>();
+		if (Found != nullptr)
+		{
+			return Found;
+		}
+	}
+	return nullptr;
+}
+
+
+#undef LOCTEXT_NAMESPACE
\ No newline at end of file
diff --git a/Engine/Plugins/Experimental/MeshModelingToolsetExp/Source/ModelingComponents/Private/ToolSceneQueriesUtil.cpp b/Engine/Plugins/Experimental/MeshModelingToolsetExp/Source/ModelingComponents/Private/ToolSceneQueriesUtil.cpp
--- a/Engine/Plugins/Experimental/MeshModelingToolsetExp/Source/ModelingComponents/Private/ToolSceneQueriesUtil.cpp
+++ b/Engine/Plugins/Experimental/MeshModelingToolsetExp/Source/ModelingComponents/Private/ToolSceneQueriesUtil.cpp
@@ -6,6 +6,7 @@
 #include "Quaternion.h"
 #include "GameFramework/Actor.h"
 #include "Components/PrimitiveComponent.h"
+#include "SceneQueries/SceneSnappingManager.h"
 
 
 static double VISUAL_ANGLE_SNAP_THRESHOLD_DEG = 1.0;
@@ -157,11 +158,16 @@ bool ToolSceneQueriesUtil::FindSceneSnapPoint(FFindSceneSnapPointParams& Params)
 {
 	double UseThreshold = (Params.VisualAngleThreshold <= 0) ? GetDefaultVisualAngleSnapThreshD() : Params.VisualAngleThreshold;
 
+	USceneSnappingManager* SnapManager = USceneSnappingManager::Find(Params.Tool->GetToolManager());
+	if (!SnapManager)
+	{
+		return false;
+	}
+
 	FViewCameraState CameraState;
 	Params.Tool->GetToolManager()->GetContextQueriesAPI()->GetCurrentViewState(CameraState);
 	UseThreshold *= CameraState.GetFOVAngleNormalizationFactor();
 
-	IToolsContextQueriesAPI* QueryAPI = Params.Tool->GetToolManager()->GetContextQueriesAPI();
 	FSceneSnapQueryRequest Request;
 	Request.RequestType = ESceneSnapQueryType::Position;
 	Request.TargetTypes = ESceneSnapQueryTargetType::None;
@@ -179,7 +185,7 @@ bool ToolSceneQueriesUtil::FindSceneSnapPoint(FFindSceneSnapPointParams& Params)
 	Request.InvisibleComponentsToInclude = Params.InvisibleComponentsToInclude;
 
 	TArray<FSceneSnapQueryResult> Results;
-	if (QueryAPI->ExecuteSceneSnapQuery(Request, Results))
+	if (SnapManager->ExecuteSceneSnapQuery(Request, Results))
 	{
 		*Params.SnapPointOut = (FVector3d)Results[0].Position;
 
@@ -210,13 +216,18 @@ bool ToolSceneQueriesUtil::FindSceneSnapPoint(FFindSceneSnapPointParams& Params)
 
 bool ToolSceneQueriesUtil::FindWorldGridSnapPoint(const UInteractiveTool* Tool, const FVector3d& Point, FVector3d& GridSnapPointOut)
 {
-	IToolsContextQueriesAPI* QueryAPI = Tool->GetToolManager()->GetContextQueriesAPI();
+	USceneSnappingManager* SnapManager = USceneSnappingManager::Find(Tool->GetToolManager());
+	if (!SnapManager)
+	{
+		return false;
+	}
+
 	FSceneSnapQueryRequest Request;
 	Request.RequestType = ESceneSnapQueryType::Position;
 	Request.TargetTypes = ESceneSnapQueryTargetType::Grid;
 	Request.Position = (FVector)Point;
 	TArray<FSceneSnapQueryResult> Results;
-	if ( QueryAPI->ExecuteSceneSnapQuery(Request, Results) )
+	if ( SnapManager->ExecuteSceneSnapQuery(Request, Results) )
 	{
 		GridSnapPointOut = (FVector3d)Results[0].Position;
 		return true;
diff --git a/Engine/Plugins/Experimental/MeshModelingToolsetExp/Source/ModelingComponents/Public/Snapping/ModelingSceneSnappingManager.h b/Engine/Plugins/Experimental/MeshModelingToolsetExp/Source/ModelingComponents/Public/Snapping/ModelingSceneSnappingManager.h
new file mode 100644
--- /dev/null
+++ b/Engine/Plugins/Experimental/MeshModelingToolsetExp/Source/ModelingComponents/Public/Snapping/ModelingSceneSnappingManager.h
@@ -0,0 +1,95 @@
+// Copyright Epic Games, Inc. All Rights Reserved.
+
+#pragma once
+
+#include "UObject/Object.h"
+#include "SceneQueries/SceneSnappingManager.h"
+#include "ModelingSceneSnappingManager.generated.h"
+
+
+class IToolsContextQueriesAPI;
+class UInteractiveToolsContext;
+class UInteractiveToolManager;
+
+
+/**
+ * UModelingSceneSnappingManager is an implementation of snapping suitable for use in
+ * Modeling Tools/Gizmos (and potentially other places). 
+ * 
+ * Currently Supports:
+ *    - snap to position/rotation grid
+ *    - snap to vertex position for Static Mesh Assets
+ *    - snap to edge position for Static Mesh Assets
+ * 
+ * Note that currently the StaticMesh vertex/edge snapping is dependent on the Physics
+ * system, and may fail or return nonsense results in some cases, due to the physics
+ * complex-collision mesh deviating from the source-model mesh
+ */
+UCLASS()
+class MODELINGCOMPONENTS_API UModelingSceneSnappingManager : public USceneSnappingManager
+{
+	GENERATED_BODY()
+
+public:
+
+	virtual void Initialize(TObjectPtr<UInteractiveToolsContext> ToolsContext);
+	virtual void Shutdown();
+
+
+	/**
+	* Try to find Snap Targets in the scene that satisfy the Snap Query.
+	* @param Request snap query configuration
+	* @param Results list of potential snap results
+	* @return true if any valid snap target was found
+	* @warning implementations are not required (and may not be able) to support snapping
+	*/
+	virtual bool ExecuteSceneSnapQuery(const FSceneSnapQueryRequest& Request, TArray<FSceneSnapQueryResult>& Results) const override;
+
+
+protected:
+
+	UPROPERTY()
+	TObjectPtr<UInteractiveToolsContext> ParentContext;
+
+	const IToolsContextQueriesAPI* QueriesAPI = nullptr;
+
+	virtual bool ExecuteSceneSnapQueryRotation(const FSceneSnapQueryRequest& Request, TArray<FSceneSnapQueryResult>& Results) const;
+	virtual bool ExecuteSceneSnapQueryPosition(const FSceneSnapQueryRequest& Request, TArray<FSceneSnapQueryResult>& Results) const;
+};
+
+
+
+
+
+namespace UE
+{
+namespace Geometry
+{
+	//
+	// The functions below are helper functions that simplify usage of a UModelingSceneSnappingManager 
+	// that is registered as a ContextStoreObject in an InteractiveToolsContext
+	//
+
+	/**
+	 * If one does not already exist, create a new instance of UModelingSceneSnappingManager and add it to the
+	 * ToolsContext's ContextObjectStore
+	 * @return true if the ContextObjectStore now has a UModelingSceneSnappingManager (whether it already existed, or was created)
+	 */
+	MODELINGCOMPONENTS_API bool RegisterSceneSnappingManager(UInteractiveToolsContext* ToolsContext);
+
+	/**
+	 * Remove any existing UModelingSceneSnappingManager from the ToolsContext's ContextObjectStore
+	 * @return true if the ContextObjectStore no longer has a UModelingSceneSnappingManager (whether it was removed, or did not exist)
+	 */
+	MODELINGCOMPONENTS_API bool DeregisterSceneSnappingManager(UInteractiveToolsContext* ToolsContext);
+
+
+	/**
+	 * Find an existing UModelingSceneSnappingManager in the ToolsContext's ContextObjectStore
+	 * @return SelectionManager pointer or nullptr if not found
+	 */
+	MODELINGCOMPONENTS_API UModelingSceneSnappingManager* FindModelingSceneSnappingManager(UInteractiveToolManager* ToolManager);
+
+
+}
+}
\ No newline at end of file
diff --git a/Engine/Plugins/Runtime/MeshModelingToolset/Source/MeshModelingTools/Private/AddPrimitiveTool.cpp b/Engine/Plugins/Runtime/MeshModelingToolset/Source/MeshModelingTools/Private/AddPrimitiveTool.cpp
--- a/Engine/Plugins/Runtime/MeshModelingToolset/Source/MeshModelingTools/Private/AddPrimitiveTool.cpp
+++ b/Engine/Plugins/Runtime/MeshModelingToolset/Source/MeshModelingTools/Private/AddPrimitiveTool.cpp
@@ -3,6 +3,7 @@
 #include "AddPrimitiveTool.h"
 #include "ToolBuilderUtil.h"
 #include "InteractiveToolManager.h"
+#include "SceneQueries/SceneSnappingManager.h"
 #include "BaseBehaviors/MouseHoverBehavior.h"
 #include "Selection/ToolSelectionUtil.h"
 #include "ModelingObjectsCreationAPI.h"
@@ -248,14 +249,18 @@ void UAddPrimitiveTool::UpdatePreviewPosition(const FInputDeviceRay& DeviceClick
 	// Snap to grid if applicable
 	if (ShapeSettings->bSnapToGrid)
 	{
-		FSceneSnapQueryRequest Request;
-		Request.RequestType = ESceneSnapQueryType::Position;
-		Request.TargetTypes = ESceneSnapQueryTargetType::Grid;
-		Request.Position = ShapeFrame.Origin;
-		TArray<FSceneSnapQueryResult> Results;
-		if (GetToolManager()->GetContextQueriesAPI()->ExecuteSceneSnapQuery(Request, Results))
+		USceneSnappingManager* SnapManager = USceneSnappingManager::Find(GetToolManager());
+		if (SnapManager)
 		{
-			ShapeFrame.Origin = Results[0].Position;
+			FSceneSnapQueryRequest Request;
+			Request.RequestType = ESceneSnapQueryType::Position;
+			Request.TargetTypes = ESceneSnapQueryTargetType::Grid;
+			Request.Position = ShapeFrame.Origin;
+			TArray<FSceneSnapQueryResult> Results;
+			if (SnapManager->ExecuteSceneSnapQuery(Request, Results))
+			{
+				ShapeFrame.Origin = Results[0].Position;
+			}
 		}
 	}
 
diff --git a/Engine/Source/Editor/Experimental/EditorInteractiveToolsFramework/Private/EdModeInteractiveToolsContext.cpp b/Engine/Source/Editor/Experimental/EditorInteractiveToolsFramework/Private/EdModeInteractiveToolsContext.cpp
--- a/Engine/Source/Editor/Experimental/EditorInteractiveToolsFramework/Private/EdModeInteractiveToolsContext.cpp
+++ b/Engine/Source/Editor/Experimental/EditorInteractiveToolsFramework/Private/EdModeInteractiveToolsContext.cpp
@@ -38,79 +38,6 @@
 
 //#define ENABLE_DEBUG_PRINTING
 
-static float SnapToIncrement(float fValue, float fIncrement, float offset = 0)
-{
-	if (!FMath::IsFinite(fValue))
-	{
-		return 0;
-	}
-	fValue -= offset;
-	float sign = FMath::Sign(fValue);
-	fValue = FMath::Abs(fValue);
-	int nInc = (int)(fValue / fIncrement);
-	float fRem = (float)fmod(fValue, fIncrement);
-	if (fRem > fIncrement / 2)
-	{
-		++nInc;
-	}
-	return sign * (float)nInc * fIncrement + offset;
-}
-
-
-
-
-static bool IsVisibleObjectHit_Internal(const FHitResult& HitResult)
-{
-	AActor* Actor = HitResult.GetActor();
-	if (Actor != nullptr && (Actor->IsHidden() || Actor->IsHiddenEd()) )
-	{
-		return false;
-	}
-	UPrimitiveComponent* Component = HitResult.GetComponent();
-	if (Component != nullptr && (Component->IsVisible() == false && Component->IsVisibleInEditor() == false))
-	{
-		return false;
-	}
-	return true;
-}
-
-static bool FindNearestVisibleObjectHit_Internal(UWorld* World, FHitResult& HitResultOut, const FVector& Start, const FVector& End,
-	bool bIsSceneGeometrySnapQuery, const TArray<const UPrimitiveComponent*>* ComponentsToIgnore, 
-	const TArray<const UPrimitiveComponent*>* InvisibleComponentsToInclude)
-{
-	FCollisionObjectQueryParams ObjectQueryParams(FCollisionObjectQueryParams::AllObjects);
-	FCollisionQueryParams QueryParams = FCollisionQueryParams::DefaultQueryParam;
-	QueryParams.bTraceComplex = true;
-	QueryParams.bReturnFaceIndex = bIsSceneGeometrySnapQuery;
-
-	TArray<FHitResult> OutHits;
-	if (World->LineTraceMultiByObjectType(OutHits, Start, End, ObjectQueryParams, QueryParams) == false)
-	{
-		return false;
-	}
-
-	float NearestVisible = TNumericLimits<float>::Max();
-	for (const FHitResult& CurResult : OutHits)
-	{
-		if (CurResult.Distance < NearestVisible)
-		{
-			if (IsVisibleObjectHit_Internal(CurResult) 
-				|| (InvisibleComponentsToInclude && InvisibleComponentsToInclude->Contains(CurResult.GetComponent())))
-			{
-				if (!ComponentsToIgnore || !ComponentsToIgnore->Contains(CurResult.GetComponent()))
-				{
-					HitResultOut = CurResult;
-					NearestVisible = CurResult.Distance;
-				}
-			}
-		}
-	}
-
-	return NearestVisible < TNumericLimits<float>::Max();
-}
-
-
-
 
 
 
@@ -181,6 +108,11 @@ public:
 		CachedViewState.bIsVR = false;
 	}
 
+	virtual UWorld* GetCurrentEditingWorld() const override
+	{
+		return EditorModeManager->GetWorld();
+	}
+
 	virtual void GetCurrentSelectionState(FToolBuilderState& StateOut) const override
 	{
 		StateOut.ToolManager = ToolsContext->ToolManager;
@@ -203,197 +135,15 @@ public:
 		return (CoordSys == COORD_World) ? EToolContextCoordinateSystem::World : EToolContextCoordinateSystem::Local;
 	}
 
-	bool ExecuteSceneSnapQueryRotation(const FSceneSnapQueryRequest& Request, TArray<FSceneSnapQueryResult>& Results) const
-	{
-		if ((Request.TargetTypes & ESceneSnapQueryTargetType::Grid) != ESceneSnapQueryTargetType::None)
-		{
-			if (GetDefault<ULevelEditorViewportSettings>()->RotGridEnabled != 0)
-			{
-				FRotator Rotator(Request.DeltaRotation);
-				FRotator RotGrid = Request.RotGridSize.Get(GEditor->GetRotGridSize());
-				Rotator = Rotator.GridSnap(RotGrid);
-
-				FSceneSnapQueryResult SnapResult;
-				SnapResult.TargetType = ESceneSnapQueryTargetType::Grid;
-				SnapResult.DeltaRotation = Rotator.Quaternion();
-				Results.Add(SnapResult);
-				return true;
-			}
-		}
-		return false;
-	}
-
-	bool ExecuteSceneSnapQueryPosition(const FSceneSnapQueryRequest& Request, TArray<FSceneSnapQueryResult>& Results) const
-	{
-		int FoundResultCount = 0;
-
-		if ((Request.TargetTypes & ESceneSnapQueryTargetType::Grid) != ESceneSnapQueryTargetType::None)
-		{
-			if (GetDefault<ULevelEditorViewportSettings>()->GridEnabled != 0)
-			{
-				FSceneSnapQueryResult SnapResult;
-				SnapResult.TargetType = ESceneSnapQueryTargetType::Grid;
-
-				float SnapSize = GEditor->GetGridSize();
-				FVector GridSize = Request.GridSize.Get(FVector(SnapSize, SnapSize, SnapSize));
-
-				SnapResult.Position.X = SnapToIncrement(Request.Position.X, GridSize.X);
-				SnapResult.Position.Y = SnapToIncrement(Request.Position.Y, GridSize.Y);
-				SnapResult.Position.Z = SnapToIncrement(Request.Position.Z, GridSize.Z);
-
-				Results.Add(SnapResult);
-				FoundResultCount++;
-			}
-		}
-
-		//
-		// Run a snap query by casting ray into the world.
-		// If a hit is found, we look up what triangle was hit, and then test its vertices and edges
-		//
-
-		// cast ray into world
-		FVector RayStart = CachedViewState.Position;
-		FVector RayDirection = Request.Position - RayStart; RayDirection.Normalize();
-		FVector RayEnd = RayStart + HALF_WORLD_MAX * RayDirection;
-		FHitResult HitResult;
-		bool bHitWorld = FindNearestVisibleObjectHit_Internal(EditorModeManager->GetWorld(), HitResult, 
-			RayStart, RayEnd, true, Request.ComponentsToIgnore, Request.InvisibleComponentsToInclude);
-		if (bHitWorld && HitResult.FaceIndex >= 0)
-		{
-			float VisualAngle = OpeningAngleDeg(Request.Position, HitResult.ImpactPoint, RayStart);
-			//UE_LOG(LogTemp, Warning, TEXT("[HIT] visualangle %f faceindex %d"), VisualAngle, HitResult.FaceIndex);
-			if (VisualAngle < Request.VisualAngleThresholdDegrees)
-			{
-				UPrimitiveComponent* Component = HitResult.Component.Get();
-				if (Cast<UStaticMeshComponent>(Component) != nullptr)
-				{
-					// HitResult.FaceIndex is apparently an index into the TriMeshCollisionData, not sure how
-					// to directly access it. Calling GetPhysicsTriMeshData is expensive!
-					//UBodySetup* BodySetup = Cast<UStaticMeshComponent>(Component)->GetBodySetup();
-					//UObject* CDPObj = BodySetup->GetOuter();
-					//IInterface_CollisionDataProvider* CDP = Cast<IInterface_CollisionDataProvider>(CDPObj);
-					//FTriMeshCollisionData TriMesh;
-					//CDP->GetPhysicsTriMeshData(&TriMesh, true);
-					//FTriIndices Triangle = TriMesh.Indices[HitResult.FaceIndex];
-					//FVector Positions[3] = { TriMesh.Vertices[Triangle.v0], TriMesh.Vertices[Triangle.v1], TriMesh.Vertices[Triangle.v2] };
-
-					// physics collision data is created from StaticMesh RenderData
-					// so use HitResult.FaceIndex to extract triangle from the LOD0 mesh
-					// (note: this may be incorrect if there are multiple sections...in that case I think we have to
-					//  first find section whose accumulated index range would contain .FaceIndexX)
-					UStaticMesh* StaticMesh = Cast<UStaticMeshComponent>(Component)->GetStaticMesh();
-					FStaticMeshLODResources& LOD = StaticMesh->GetRenderData()->LODResources[0];
-					FIndexArrayView Indices = LOD.IndexBuffer.GetArrayView();
-					int32 TriIdx = 3 * HitResult.FaceIndex;
-					FVector Positions[3];
-					Positions[0] = LOD.VertexBuffers.PositionVertexBuffer.VertexPosition(Indices[TriIdx]);
-					Positions[1] = LOD.VertexBuffers.PositionVertexBuffer.VertexPosition(Indices[TriIdx+1]);
-					Positions[2] = LOD.VertexBuffers.PositionVertexBuffer.VertexPosition(Indices[TriIdx+2]);
-
-					// transform to world space
-					FTransform ComponentTransform = Component->GetComponentTransform();
-					Positions[0] = ComponentTransform.TransformPosition(Positions[0]);
-					Positions[1] = ComponentTransform.TransformPosition(Positions[1]);
-					Positions[2] = ComponentTransform.TransformPosition(Positions[2]);
-
-					FSceneSnapQueryResult SnapResult;
-					SnapResult.TriVertices[0] = Positions[0];
-					SnapResult.TriVertices[1] = Positions[1];
-					SnapResult.TriVertices[2] = Positions[2];
-
-					// try snapping to vertices
-					float SmallestAngle = Request.VisualAngleThresholdDegrees;
-					if ( (Request.TargetTypes & ESceneSnapQueryTargetType::MeshVertex) != ESceneSnapQueryTargetType::None)
-					{
-						for (int j = 0; j < 3; ++j)
-						{
-							VisualAngle = OpeningAngleDeg(Request.Position, Positions[j], RayStart);
-							if (VisualAngle < SmallestAngle)
-							{
-								SmallestAngle = VisualAngle;
-								SnapResult.Position = Positions[j];
-								SnapResult.TargetType = ESceneSnapQueryTargetType::MeshVertex;
-								SnapResult.TriSnapIndex = j;
-							}
-						}
-					}
-
-					// try snapping to nearest points on edges
-					if ( ((Request.TargetTypes & ESceneSnapQueryTargetType::MeshEdge) != ESceneSnapQueryTargetType::None) &&
-						 (SnapResult.TargetType != ESceneSnapQueryTargetType::MeshVertex) )
-					{
-						for (int j = 0; j < 3; ++j)
-						{
-							FVector EdgeNearestPt = NearestSegmentPt(Positions[j], Positions[(j+1)%3], Request.Position);
-							VisualAngle = OpeningAngleDeg(Request.Position, EdgeNearestPt, RayStart);
-							if (VisualAngle < SmallestAngle )
-							{
-								SmallestAngle = VisualAngle;
-								SnapResult.Position = EdgeNearestPt;
-								SnapResult.TargetType = ESceneSnapQueryTargetType::MeshEdge;
-								SnapResult.TriSnapIndex = j;
-							}
-						}
-					}
-
-					// if we found a valid snap, return it
-					if (SmallestAngle < Request.VisualAngleThresholdDegrees)
-					{
-						SnapResult.TargetActor = HitResult.HitObjectHandle.FetchActor();
-						SnapResult.TargetComponent = HitResult.Component.Get();
-						Results.Add(SnapResult);
-						FoundResultCount++;
-					}
-				}
-			}
-
-		}
-
-		return (FoundResultCount > 0);
-	}
-
-	virtual bool ExecuteSceneSnapQuery(const FSceneSnapQueryRequest& Request, TArray<FSceneSnapQueryResult>& Results) const override
+	virtual FToolContextSnappingConfiguration GetCurrentSnappingSettings() const override
 	{
-		switch (Request.RequestType)
-		{
-		case ESceneSnapQueryType::Position:
-			return ExecuteSceneSnapQueryPosition(Request, Results);
-			break;
-		case ESceneSnapQueryType::Rotation:
-			return ExecuteSceneSnapQueryRotation(Request, Results);
-			break;
-		default:
-			check(!"Only Position and Rotation Snap Queries are supported");
-		}
-		return false;
-	}
-
-	//@ todo this are mirrored from GeometryProcessing, which is still experimental...replace w/ direct calls once GP component is standardized
-	static float OpeningAngleDeg(FVector A, FVector B, const FVector& P)
-	{
-		A -= P;
-		A.Normalize();
-		B -= P;
-		B.Normalize();
-		float Dot = FMath::Clamp(FVector::DotProduct(A,B), -1.0f, 1.0f);
-		return acos(Dot) * (180.0f / 3.141592653589f);
-	}
-
-	static FVector NearestSegmentPt(FVector A, FVector B, const FVector& P)
-	{
-		FVector Direction = (B - A);
-		float Length = Direction.Size();
-		Direction /= Length;
-		float t = FVector::DotProduct( (P - A), Direction);
-		if (t >= Length)
-		{
-			return B;
-		}
-		if (t <= 0)
-		{
-			return A;
-		}
-		return A + t * Direction;
+		FToolContextSnappingConfiguration Config;
+		Config.bEnablePositionGridSnapping = (GetDefault<ULevelEditorViewportSettings>()->GridEnabled != 0);
+		float EditorGridSize = GEditor->GetGridSize();
+		Config.PositionGridDimensions = FVector(EditorGridSize, EditorGridSize, EditorGridSize);
+		Config.bEnableRotationGridSnapping = (GetDefault<ULevelEditorViewportSettings>()->RotGridEnabled != 0);
+		Config.RotationGridAngles = GEditor->GetRotGridSize();
+		return Config;
 	}
 
 	virtual UMaterialInterface* GetStandardMaterial(EStandardToolContextMaterials MaterialType) const
@@ -474,9 +224,11 @@ public:
 	virtual void AppendChange(UObject* TargetObject, TUniquePtr<FToolCommandChange> Change, const FText& Description) override
 	{
 		FScopedTransaction Transaction(Description);
-		check(GUndo != nullptr);
-		GUndo->StoreUndo(TargetObject, MoveTemp(Change));
-		// end transaction
+		//if (ensure(GUndo != nullptr))		// ideally we would ensure here, but currently this can be hit on world teardown, resolution TBD
+		if (GUndo != nullptr)
+		{
+			GUndo->StoreUndo(TargetObject, MoveTemp(Change));
+		}
 	}
 
 
diff --git a/Engine/Source/Editor/Experimental/EditorInteractiveToolsFramework/Private/EditorGizmos/TransformGizmo.cpp b/Engine/Source/Editor/Experimental/EditorInteractiveToolsFramework/Private/EditorGizmos/TransformGizmo.cpp
--- a/Engine/Source/Editor/Experimental/EditorInteractiveToolsFramework/Private/EditorGizmos/TransformGizmo.cpp
+++ b/Engine/Source/Editor/Experimental/EditorInteractiveToolsFramework/Private/EditorGizmos/TransformGizmo.cpp
@@ -770,7 +770,7 @@ void UTransformGizmo::ClearActiveTarget()
 bool UTransformGizmo::PositionSnapFunction(const FVector& WorldPosition, FVector& SnappedPositionOut) const
 {
 	SnappedPositionOut = WorldPosition;
-
+#if 0
 	// only snap if we want snapping obvs
 	if (bSnapToWorldGrid == false)
 	{
@@ -797,7 +797,7 @@ bool UTransformGizmo::PositionSnapFunction(const FVector& WorldPosition, FVector
 		SnappedPositionOut = Results[0].Position;
 		return true;
 	};
-
+#endif
 	return false;
 }
 
@@ -805,7 +805,7 @@ bool UTransformGizmo::PositionSnapFunction(const FVector& WorldPosition, FVector
 FQuat UTransformGizmo::RotationSnapFunction(const FQuat& DeltaRotation) const
 {
 	FQuat SnappedDeltaRotation = DeltaRotation;
-
+#if 0
 	// only snap if we want snapping 
 	if (bSnapToWorldRotGrid)
 	{
@@ -823,6 +823,7 @@ FQuat UTransformGizmo::RotationSnapFunction(const FQuat& DeltaRotation) const
 			SnappedDeltaRotation = Results[0].DeltaRotation;
 		};
 	}
+#endif	
 	return SnappedDeltaRotation;
 }
 
diff --git a/Engine/Source/Runtime/InteractiveToolsFramework/Private/BaseGizmos/CombinedTransformGizmo.cpp b/Engine/Source/Runtime/InteractiveToolsFramework/Private/BaseGizmos/CombinedTransformGizmo.cpp
--- a/Engine/Source/Runtime/InteractiveToolsFramework/Private/BaseGizmos/CombinedTransformGizmo.cpp
+++ b/Engine/Source/Runtime/InteractiveToolsFramework/Private/BaseGizmos/CombinedTransformGizmo.cpp
@@ -2,6 +2,7 @@
 
 #include "BaseGizmos/CombinedTransformGizmo.h"
 #include "InteractiveGizmoManager.h"
+#include "SceneQueries/SceneSnappingManager.h"
 #include "BaseGizmos/AxisPositionGizmo.h"
 #include "BaseGizmos/PlanePositionGizmo.h"
 #include "BaseGizmos/AxisAngleGizmo.h"
@@ -933,6 +934,12 @@ bool UCombinedTransformGizmo::PositionSnapFunction(const FVector& WorldPosition,
 		return false;
 	}
 
+	USceneSnappingManager* SnapManager = USceneSnappingManager::Find(GetGizmoManager());
+	if (!SnapManager)
+	{
+		return false;
+	}
+
 	FSceneSnapQueryRequest Request;
 	Request.RequestType = ESceneSnapQueryType::Position;
 	Request.TargetTypes = ESceneSnapQueryTargetType::Grid;
@@ -942,7 +949,7 @@ bool UCombinedTransformGizmo::PositionSnapFunction(const FVector& WorldPosition,
 		Request.GridSize = ExplicitGridSize;
 	}
 	TArray<FSceneSnapQueryResult> Results;
-	if (GetGizmoManager()->GetContextQueriesAPI()->ExecuteSceneSnapQuery(Request, Results))
+	if (SnapManager->ExecuteSceneSnapQuery(Request, Results))
 	{
 		SnappedPositionOut = Results[0].Position;
 		return true;
@@ -958,19 +965,23 @@ FQuat UCombinedTransformGizmo::RotationSnapFunction(const FQuat& DeltaRotation)
 	// only snap if we want snapping obvs
 	if (bSnapToWorldRotGrid)
 	{
-		FSceneSnapQueryRequest Request;
-		Request.RequestType   = ESceneSnapQueryType::Rotation;
-		Request.TargetTypes   = ESceneSnapQueryTargetType::Grid;
-		Request.DeltaRotation = DeltaRotation;
-		if ( bRotationGridSizeIsExplicit )
+		USceneSnappingManager* SnapManager = USceneSnappingManager::Find(GetGizmoManager());
+		if ( SnapManager )
 		{
-			Request.RotGridSize = ExplicitRotationGridSize;
+			FSceneSnapQueryRequest Request;
+			Request.RequestType   = ESceneSnapQueryType::Rotation;
+			Request.TargetTypes   = ESceneSnapQueryTargetType::Grid;
+			Request.DeltaRotation = DeltaRotation;
+			if ( bRotationGridSizeIsExplicit )
+			{
+				Request.RotGridSize = ExplicitRotationGridSize;
+			}
+			TArray<FSceneSnapQueryResult> Results;
+			if (SnapManager->ExecuteSceneSnapQuery(Request, Results))
+			{
+				SnappedDeltaRotation = Results[0].DeltaRotation;
+			};
 		}
-		TArray<FSceneSnapQueryResult> Results;
-		if (GetGizmoManager()->GetContextQueriesAPI()->ExecuteSceneSnapQuery(Request, Results))
-		{
-			SnappedDeltaRotation = Results[0].DeltaRotation;
-		};
 	}
 	return SnappedDeltaRotation;
 }
diff --git a/Engine/Source/Runtime/InteractiveToolsFramework/Private/SceneQueries/SceneSnappingManager.cpp b/Engine/Source/Runtime/InteractiveToolsFramework/Private/SceneQueries/SceneSnappingManager.cpp
new file mode 100644
--- /dev/null
+++ b/Engine/Source/Runtime/InteractiveToolsFramework/Private/SceneQueries/SceneSnappingManager.cpp
@@ -0,0 +1,33 @@
+// Copyright Epic Games, Inc. All Rights Reserved.
+
+#include "SceneQueries/SceneSnappingManager.h"
+#include "InteractiveToolManager.h"
+#include "InteractiveGizmoManager.h"
+#include "ContextObjectStore.h"
+
+USceneSnappingManager* USceneSnappingManager::Find(UInteractiveToolManager* ToolManager)
+{
+	if (ensure(ToolManager))
+	{
+		USceneSnappingManager* Found = ToolManager->GetContextObjectStore()->FindContext<USceneSnappingManager>();
+		if (Found != nullptr)
+		{
+			return Found;
+		}
+	}
+	return nullptr;
+}
+
+
+USceneSnappingManager* USceneSnappingManager::Find(UInteractiveGizmoManager* GizmoManager)
+{
+	if (ensure(GizmoManager))
+	{
+		USceneSnappingManager* Found = GizmoManager->GetContextObjectStore()->FindContext<USceneSnappingManager>();
+		if (Found != nullptr)
+		{
+			return Found;
+		}
+	}
+	return nullptr;
+}
diff --git a/Engine/Source/Runtime/InteractiveToolsFramework/Public/SceneQueries/SceneSnappingManager.h b/Engine/Source/Runtime/InteractiveToolsFramework/Public/SceneQueries/SceneSnappingManager.h
new file mode 100644
--- /dev/null
+++ b/Engine/Source/Runtime/InteractiveToolsFramework/Public/SceneQueries/SceneSnappingManager.h
@@ -0,0 +1,142 @@
+// Copyright Epic Games, Inc. All Rights Reserved.
+
+#pragma once
+
+#include "UObject/Object.h"
+#include "Misc/Optional.h"
+#include "Math/Vector.h"
+#include "Math/Rotator.h"
+#include "Math/Quat.h"
+#include "SceneSnappingManager.generated.h"
+
+class AActor;
+class UActorComponent;
+class UPrimitiveComponent;
+class UInteractiveToolManager;
+class UInteractiveGizmoManager;
+
+
+/** Types of Snap Queries that a USceneSnappingManager may support */
+UENUM()
+enum class ESceneSnapQueryType : uint8
+{
+	/** snapping a position */
+	Position = 1,
+	Rotation = 2
+};
+
+
+/** Types of Snap Targets that a caller may want to run snap queries against. */
+UENUM()
+enum class ESceneSnapQueryTargetType : uint8
+{
+	None = 0,
+	/** Consider any mesh vertex */
+	MeshVertex = 1,
+	/** Consider any mesh edge */
+	MeshEdge = 2,
+	/** Grid Snapping */
+	Grid = 4,
+
+	All = MeshVertex | MeshEdge | Grid
+};
+ENUM_CLASS_FLAGS(ESceneSnapQueryTargetType);
+
+
+/**
+ * Configuration variables for a USceneSnappingManager snap query request.
+ */
+struct INTERACTIVETOOLSFRAMEWORK_API FSceneSnapQueryRequest
+{
+	/** What type of snap query geometry is this */
+	ESceneSnapQueryType RequestType = ESceneSnapQueryType::Position;
+	/** What does caller want to try to snap to */
+	ESceneSnapQueryTargetType TargetTypes = ESceneSnapQueryTargetType::Grid;
+
+	/** Optional explicitly specified position grid */
+	TOptional<FVector> GridSize{};
+
+	/** Optional explicitly specified rotation grid */
+	TOptional<FRotator> RotGridSize{};
+
+	/** Snap input position */
+	FVector Position;
+	/** Another position must deviate less than this number of degrees (in visual angle) to be considered an acceptable snap position */
+	float VisualAngleThresholdDegrees;
+
+	/** Snap input rotation delta */
+	FQuat DeltaRotation;
+
+	/** Optional: components to consider invisible even if they aren't. */
+	const TArray<const UPrimitiveComponent*>* ComponentsToIgnore = nullptr;
+	/** Optional: components to consider visible even if they aren't. */
+	const TArray<const UPrimitiveComponent*>* InvisibleComponentsToInclude = nullptr;
+};
+
+
+/**
+ * Computed result of a USceneSnappingManager snap query request
+ */
+struct INTERACTIVETOOLSFRAMEWORK_API FSceneSnapQueryResult
+{
+	/** Actor that owns snap target */
+	AActor* TargetActor = nullptr;
+	/** Component that owns snap target */
+	UActorComponent* TargetComponent = nullptr;
+	/** What kind of geometric element was snapped to */
+	ESceneSnapQueryTargetType TargetType = ESceneSnapQueryTargetType::None;
+
+	/** Snap position (may not be set depending on query types) */
+	FVector Position;
+	/** Snap normal (may not be set depending on query types) */
+	FVector Normal;
+	/** Snap rotation delta (may not be set depending on query types) */
+	FQuat   DeltaRotation;
+
+	/** Vertices of triangle that contains result (for debugging, may not be set) */
+	FVector TriVertices[3];
+	/** Vertex/Edge index we snapped to in triangle */
+	int TriSnapIndex;
+
+};
+
+
+/**
+ * USceneSnappingManager is intended to be used as a base class for a Snapping implementation
+ * stored in the ContextObjectStore of an InteractiveToolsContext. ITF classes like Tools and Gizmos
+ * can then access this object and run snap queries via the various API functions.
+ * 
+ * USceneSnappingManager::Find() can be used to look up a registered USceneSnappingManager, if one is available
+ * 
+ * See UModelingSceneSnappingManager for a sample implementation.
+ */
+UCLASS()
+class INTERACTIVETOOLSFRAMEWORK_API USceneSnappingManager : public UObject
+{
+	GENERATED_BODY()
+public:
+
+	/**
+	* Try to find Snap Targets/Results in the scene that satisfy the Snap Query.
+	* @param Request snap query configuration
+	* @param Results list of potential snap results
+	* @return true if any valid snap target/result was found
+	* @warning implementations are not required (and may not be able) to support snapping
+	*/
+	virtual bool ExecuteSceneSnapQuery(const FSceneSnapQueryRequest& Request, TArray<FSceneSnapQueryResult>& ResultsOut) const
+	{
+		return false;
+	}
+
+
+public:
+	/**
+	 * @return existing USceneSnappingManager registered in context store via the ToolManager, or nullptr if not found
+	 */
+	static USceneSnappingManager* Find(UInteractiveToolManager* ToolManager);
+
+	/**
+	 * @return existing USceneSnappingManager registered in context store via the ToolManager, or nullptr if not found
+	 */
+	static USceneSnappingManager* Find(UInteractiveGizmoManager* GizmoManager);
+};
\ No newline at end of file
diff --git a/Engine/Source/Runtime/InteractiveToolsFramework/Public/ToolContextInterfaces.h b/Engine/Source/Runtime/InteractiveToolsFramework/Public/ToolContextInterfaces.h
--- a/Engine/Source/Runtime/InteractiveToolsFramework/Public/ToolContextInterfaces.h
+++ b/Engine/Source/Runtime/InteractiveToolsFramework/Public/ToolContextInterfaces.h
@@ -86,97 +86,6 @@ struct INTERACTIVETOOLSFRAMEWORK_API FViewCameraState
 
 
 
-/** Types of Snap Queries that a ToolsContext parent may support, that Tools may request */
-UENUM()
-enum class ESceneSnapQueryType
-{
-	/** snapping a position */
-	Position = 1,
-	Rotation = 2
-};
-
-/** Types of snap targets that a Tool may want to run snap queries against. */
-UENUM()
-enum class ESceneSnapQueryTargetType
-{
-	None = 0,
-	/** Consider any mesh vertex */
-	MeshVertex = 1,
-	/** Consider any mesh edge */
-	MeshEdge = 2,
-	/** Grid Snapping */
-	Grid = 4,
-
-	All = MeshVertex | MeshEdge | Grid
-};
-ENUM_CLASS_FLAGS(ESceneSnapQueryTargetType);
-
-/**
- * Configuration variables for a IToolsContextQueriesAPI snap query request.
- */
-struct INTERACTIVETOOLSFRAMEWORK_API FSceneSnapQueryRequest
-{
-	/** What type of snap query geometry is this */
-	ESceneSnapQueryType RequestType = ESceneSnapQueryType::Position;
-	/** What does caller want to try to snap to */
-	ESceneSnapQueryTargetType TargetTypes = ESceneSnapQueryTargetType::Grid;
-
-	/** Optional explicitly specified position grid */
-	TOptional<FVector> GridSize{};
-
-	/** Optional explicitly specified rotation grid */
-	TOptional<FRotator> RotGridSize{};
-
-	/** Snap input position */
-	FVector Position;
-	/** Another position must deviate less than this number of degrees (in visual angle) to be considered an acceptable snap position */
-	float VisualAngleThresholdDegrees;
-
-	/** Snap input direction */
-	FVector Direction;
-	/** Another direction must deviate less than this number of degrees from Direction to be considered an acceptable snap direction */
-	float DirectionAngleThresholdDegrees;
-
-	/** Snap input rotation delta */
-	FQuat DeltaRotation;
-
-	/** Optional: components to consider invisible even if they aren't. */
-	const TArray<const UPrimitiveComponent*>* ComponentsToIgnore = nullptr;
-	/** Optional: components to consider visible even if they aren't. */
-	const TArray<const UPrimitiveComponent*>* InvisibleComponentsToInclude = nullptr;
-};
-
-
-/**
- * Computed result of a IToolsContextQueriesAPI snap query request
- */
-struct INTERACTIVETOOLSFRAMEWORK_API FSceneSnapQueryResult
-{
-	/** Actor that owns snap target */
-	AActor* TargetActor = nullptr;
-	/** Component that owns snap target */
-	UActorComponent* TargetComponent = nullptr;
-	/** What kind of geometric element was snapped to */
-	ESceneSnapQueryTargetType TargetType = ESceneSnapQueryTargetType::None;
-
-	/** Snap position (may not be set depending on query types) */
-	FVector Position;
-	/** Snap normal (may not be set depending on query types) */
-	FVector Normal;
-	/** Snap direction (may not be set depending on query types) */
-	FVector Direction;
-	/** Snap rotation delta (may not be set depending on query types) */
-	FQuat   DeltaRotation;
-
-	/** Vertices of triangle that contains result (for debugging, may not be set) */
-	FVector TriVertices[3];
-	/** Vertex/Edge index we snapped to in triangle */
-	int TriSnapIndex;
-
-};
-
-
-
 /** Types of standard materials that Tools may request from Context */
 UENUM()
 enum class EStandardToolContextMaterials
@@ -194,6 +103,22 @@ enum class EToolContextCoordinateSystem
 	Local = 1
 };
 
+/**
+ * Snapping configuration settings/state
+ */
+struct INTERACTIVETOOLSFRAMEWORK_API FToolContextSnappingConfiguration
+{
+	/** Specify whether position snapping should be applied */
+	bool bEnablePositionGridSnapping = false;
+	/** Specify position grid spacing or cell dimensions */
+	FVector PositionGridDimensions = FVector::Zero();
+
+	/** Specify whether rotation snapping should be applied */
+	bool bEnableRotationGridSnapping = false;
+	/** Specify rotation snapping step size */
+	FRotator RotationGridAngles = FRotator::ZeroRotator;
+};
+
 
 /**
  * Users of the Tools Framework need to implement IToolsContextQueriesAPI to provide
@@ -204,35 +129,32 @@ class IToolsContextQueriesAPI
 public:
 	virtual ~IToolsContextQueriesAPI() {}
 
+	/**
+	 * @return the UWorld currently being targetted by the ToolsContext, default location for new Actors/etc
+	 */
+	virtual UWorld* GetCurrentEditingWorld() const = 0;
+
 	/**
 	 * Collect up current-selection information for the current scene state (ie what is selected in Editor, etc)
 	 * @param StateOut this structure is populated with available state information
 	 */
 	virtual void GetCurrentSelectionState(FToolBuilderState& StateOut) const = 0;
 
-
 	/**
 	 * Request information about current view state
 	 * @param StateOut this structure is populated with available state information
 	 */
 	virtual void GetCurrentViewState(FViewCameraState& StateOut) const = 0;
 
-
 	/**
 	 * Request current external coordinate-system setting
 	 */
 	virtual EToolContextCoordinateSystem GetCurrentCoordinateSystem() const = 0;	
 
-
 	/**
-	 * Try to find Snap Targets in the scene that satisfy the Snap Query.
-	 * @param Request snap query configuration
-	 * @param Results list of potential snap results
-	 * @return true if any valid snap target was found
-	 * @warning implementations are not required (and may not be able) to support snapping
+	 * Request current external snapping settings
 	 */
-	virtual bool ExecuteSceneSnapQuery(const FSceneSnapQueryRequest& Request, TArray<FSceneSnapQueryResult>& Results ) const = 0;
-
+	virtual FToolContextSnappingConfiguration GetCurrentSnappingSettings() const = 0;
 
 	/**
 	 * Many tools need standard types of materials that the user should provide (eg a vertex-color material, etc)
