ModelingTools: Cleanup BakeVertex type enum names.

#rnx

diff --git a/Engine/Plugins/Experimental/MeshModelingToolsetExp/Source/MeshModelingToolsExp/Private/BakeMeshAttributeVertexTool.cpp b/Engine/Plugins/Experimental/MeshModelingToolsetExp/Source/MeshModelingToolsExp/Private/BakeMeshAttributeVertexTool.cpp
--- a/Engine/Plugins/Experimental/MeshModelingToolsetExp/Source/MeshModelingToolsExp/Private/BakeMeshAttributeVertexTool.cpp
+++ b/Engine/Plugins/Experimental/MeshModelingToolsetExp/Source/MeshModelingToolsExp/Private/BakeMeshAttributeVertexTool.cpp
@@ -182,21 +182,21 @@ public:
 				Baker->ColorEvaluator = CurvatureEval;
 				break;
 			}
-			case EBakeVertexTypeColor::PositionImage:
+			case EBakeVertexTypeColor::Position:
 			{
 				TSharedPtr<FMeshPropertyMapEvaluator, ESPMode::ThreadSafe> PropertyEval = MakeShared<FMeshPropertyMapEvaluator, ESPMode::ThreadSafe>();
 				PropertyEval->Property = EMeshPropertyMapType::Position;
 				Baker->ColorEvaluator = PropertyEval;
 				break;
 			}
-			case EBakeVertexTypeColor::NormalImage:
+			case EBakeVertexTypeColor::ObjectSpaceNormal:
 			{
 				TSharedPtr<FMeshPropertyMapEvaluator, ESPMode::ThreadSafe> PropertyEval = MakeShared<FMeshPropertyMapEvaluator, ESPMode::ThreadSafe>();
 				PropertyEval->Property = EMeshPropertyMapType::Normal;
 				Baker->ColorEvaluator = PropertyEval;
 				break;
 			}
-			case EBakeVertexTypeColor::FaceNormalImage:
+			case EBakeVertexTypeColor::FaceNormal:
 			{
 				TSharedPtr<FMeshPropertyMapEvaluator, ESPMode::ThreadSafe> PropertyEval = MakeShared<FMeshPropertyMapEvaluator, ESPMode::ThreadSafe>();
 				PropertyEval->Property = EMeshPropertyMapType::FacetNormal;
@@ -210,7 +210,7 @@ public:
 				Baker->ColorEvaluator = PropertyEval;
 				break;
 			}
-			case EBakeVertexTypeColor::Texture2DImage:
+			case EBakeVertexTypeColor::Texture:
 			{
 				TSharedPtr<FMeshResampleImageEvaluator, ESPMode::ThreadSafe> TextureEval = MakeShared<FMeshResampleImageEvaluator, ESPMode::ThreadSafe>();
 				DetailSampler.SetColorMap(DetailMesh.Get(), IMeshBakerDetailSampler::FBakeDetailTexture(TextureImage.Get(), TextureSettings.UVLayer));
@@ -493,7 +493,7 @@ void UBakeMeshAttributeVertexTool::UpdateOnModeChange()
 			case EBakeVertexTypeColor::Curvature:
 				SetToolPropertySourceEnabled(CurvatureSettings, true);
 				break;
-			case EBakeVertexTypeColor::Texture2DImage:
+			case EBakeVertexTypeColor::Texture:
 				SetToolPropertySourceEnabled(TextureSettings, true);
 				break;
 			case EBakeVertexTypeColor::MultiTexture:
@@ -671,13 +671,13 @@ void UBakeMeshAttributeVertexTool::UpdateResult()
 		case EBakeVertexTypeColor::Curvature:
 			OpState = UpdateResult_Curvature();
 			break;
-		case EBakeVertexTypeColor::NormalImage:
-		case EBakeVertexTypeColor::FaceNormalImage:
-		case EBakeVertexTypeColor::PositionImage:
+		case EBakeVertexTypeColor::ObjectSpaceNormal:
+		case EBakeVertexTypeColor::FaceNormal:
+		case EBakeVertexTypeColor::Position:
 		case EBakeVertexTypeColor::MaterialID:
 			OpState = UpdateResult_MeshProperty();
 			break;
-		case EBakeVertexTypeColor::Texture2DImage:
+		case EBakeVertexTypeColor::Texture:
 			OpState = UpdateResult_Texture2DImage();
 			break;
 		case EBakeVertexTypeColor::MultiTexture:
diff --git a/Engine/Plugins/Experimental/MeshModelingToolsetExp/Source/MeshModelingToolsExp/Public/BakeMeshAttributeVertexTool.h b/Engine/Plugins/Experimental/MeshModelingToolsetExp/Source/MeshModelingToolsExp/Public/BakeMeshAttributeVertexTool.h
--- a/Engine/Plugins/Experimental/MeshModelingToolsetExp/Source/MeshModelingToolsExp/Public/BakeMeshAttributeVertexTool.h
+++ b/Engine/Plugins/Experimental/MeshModelingToolsetExp/Source/MeshModelingToolsExp/Public/BakeMeshAttributeVertexTool.h
@@ -49,25 +49,25 @@ UENUM()
 enum class EBakeVertexTypeColor
 {
 	/* Sample tangent space normals from the detail mesh */
-	TangentSpaceNormal UMETA(DisplayName="Tangent Space Normals"),
+	TangentSpaceNormal,
 	/* Sample ambient occlusion from the detail mesh */
-	AmbientOcclusion UMETA(DisplayName="Ambient Occlusion"),
+	AmbientOcclusion,
 	/* Sample normals skewed towards the least occluded direction from the detail mesh */
-	BentNormal UMETA(DisplayName="Bent Normals"),
+	BentNormal,
 	/* Sample mesh curvatures from the detail mesh */
-	Curvature UMETA(DisplayName="Curvature"),
+	Curvature,
 	/* Sample a source texture from the detail mesh UVs */
-	Texture2DImage UMETA(DisplayName="Texture"),
+	Texture,
 	/* Sample object space normals from the detail mesh */
-	NormalImage UMETA(DisplayName="Object Space Normals"),
+	ObjectSpaceNormal,
 	/* Sample object space face normals from the detail mesh */
-	FaceNormalImage UMETA(DisplayName="Face Normals"),
+	FaceNormal,
 	/* Sample bounding box relative positions from the detail mesh */
-	PositionImage UMETA(DisplayName="Position"),
+	Position,
 	/* Sample material IDs as unique colors from the detail mesh */
 	MaterialID UMETA(DisplayName="Material ID"),
 	/* Sample a source texture per material ID on the detail mesh */
-	MultiTexture UMETA(DisplayName="MultiTexture")
+	MultiTexture
 };
 
 
@@ -76,9 +76,9 @@ enum class EBakeVertexTypeChannel
 {
 	None,
 	/* Sample ambient occlusion from the detail mesh */
-	AmbientOcclusion UMETA(DisplayName="Ambient Occlusion"),
+	AmbientOcclusion,
 	/* Sample mesh curvatures from the detail mesh */
-	Curvature UMETA(DisplayName="Curvature"),
+	Curvature
 };
 
 
