UVEditor: Use AABB spatial index to implement marquee selection in the UV editor

#rnx

diff --git a/Engine/Plugins/Experimental/MeshModelingToolsetExp/Source/ModelingComponents/Private/Selection/MeshSelectionMechanic.cpp b/Engine/Plugins/Experimental/MeshModelingToolsetExp/Source/ModelingComponents/Private/Selection/MeshSelectionMechanic.cpp
--- a/Engine/Plugins/Experimental/MeshModelingToolsetExp/Source/ModelingComponents/Private/Selection/MeshSelectionMechanic.cpp
+++ b/Engine/Plugins/Experimental/MeshModelingToolsetExp/Source/ModelingComponents/Private/Selection/MeshSelectionMechanic.cpp
@@ -76,106 +76,199 @@ namespace MeshSelectionMechanicLocals
 		}
 	}
 	
-	UE::Geometry::FDynamicMeshSelection::EType ToCompatibleDynamicMeshSelectionType(const EMeshSelectionMechanicMode& Mode)
+	FDynamicMeshSelection::EType ToCompatibleDynamicMeshSelectionType(const EMeshSelectionMechanicMode& Mode)
 	{
 		switch (Mode)
 		{
 			case EMeshSelectionMechanicMode::Mesh:
 			case EMeshSelectionMechanicMode::Component:
 			case EMeshSelectionMechanicMode::Triangle:
-				return UE::Geometry::FDynamicMeshSelection::EType::Triangle;
+				return FDynamicMeshSelection::EType::Triangle;
 			case EMeshSelectionMechanicMode::Edge:
-				return UE::Geometry::FDynamicMeshSelection::EType::Edge;
+				return FDynamicMeshSelection::EType::Edge;
 			case EMeshSelectionMechanicMode::Vertex:
-				return UE::Geometry::FDynamicMeshSelection::EType::Vertex;
+				return FDynamicMeshSelection::EType::Vertex;
 		}
 		checkNoEntry();
-		return UE::Geometry::FDynamicMeshSelection::EType::Vertex;
-	}
-
-	template <typename Real>
-	TVector2<FVector::FReal> XY(const TVector<Real>& Point)
-	{
-		return {static_cast<FVector::FReal>(Point.X), static_cast<FVector::FReal>(Point.Y)};
+		return FDynamicMeshSelection::EType::Vertex;
 	}
 
 	// Returns the marquee selection rectangle, obtained from the given CameraRectangle, projected to the XY plane
-	FCameraRectangle::FAxisAlignedBox2 GetRectangleXY(const FCameraRectangle& CameraRectangle)
+	FAxisAlignedBox2d GetRectangleXY(const FCameraRectangle& CameraRectangle)
 	{
 		ensure(CameraRectangle.bIsInitialized);
+		FAxisAlignedBox2d Result;
 		
 		double Offset = CameraRectangle.SelectionDomain.Plane.DistanceTo(FVector::ZeroVector);
 		FCameraRectangle::FRectangleInPlane Domain = CameraRectangle.ProjectSelectionDomain(Offset);
-		FCameraRectangle::FAxisAlignedBox2 RectangleXY;
 		
 		// This works because we know the UV axes are aligned with the XY axes, see the comment in UUVEditorMode::InitializeTargets
-		RectangleXY.Contain(XY(CameraRectangle.PointUVToPoint3D(Domain.Plane, Domain.Rectangle.Min)));
-		RectangleXY.Contain(XY(CameraRectangle.PointUVToPoint3D(Domain.Plane, Domain.Rectangle.Max)));
+		const FVector MinPoint3D = CameraRectangle.PointUVToPoint3D(Domain.Plane, Domain.Rectangle.Min);
+		const FVector MaxPoint3D = CameraRectangle.PointUVToPoint3D(Domain.Plane, Domain.Rectangle.Max);
+		Result.Contain(FVector2d{MinPoint3D.X, MinPoint3D.Y}); // Convert to 2D and convert to double
+		Result.Contain(FVector2d{MaxPoint3D.X, MaxPoint3D.Y});
+	
+		return Result;
+	}
 	
-		return RectangleXY;
+	FVector2d XY(const FVector3d& Point)
+	{
+		return {Point.X, Point.Y};
+	}
+
+	void AppendVertexIDs(const FDynamicMesh3& MeshXY0, int TriangleID, TArray<int>& VertexIDs)
+	{
+		const FIndex3i& Triangle = MeshXY0.GetTriangleRef(TriangleID);
+		VertexIDs.Add(Triangle.A);
+		VertexIDs.Add(Triangle.B);
+		VertexIDs.Add(Triangle.C);
 	}
 
-	// Returns the points of the given mesh that lie inside the given rectangle. Assumptions:
-	// - Mesh vertices lie in the XY plane (have zero Z coordinate)
-	// - The camera is above the mesh and looking in the -Z direction (so no clip testing is needed)
-	TArray<int32> FindProjectedVerticesInRectangleXY0(const FCameraRectangle& CameraRectangle, const FDynamicMesh3& MeshXY0)
+	void AppendVertexIDsIfIntersected(const FDynamicMesh3& MeshXY0, const FAxisAlignedBox2d& RectangleXY, int TriangleID, TArray<int>& VertexIDs)
 	{
-		FCameraRectangle::FAxisAlignedBox2 RectangleXY = GetRectangleXY(CameraRectangle);
+		const FIndex3i& Triangle = MeshXY0.GetTriangleRef(TriangleID);
+		if (RectangleXY.Contains(XY(MeshXY0.GetVertex(Triangle.A))))
+		{
+			VertexIDs.Add(Triangle.A);
+		}
 		
-		TArray<int32> Result;
-		for (int32 VertexIndex : MeshXY0.VertexIndicesItr())
+		if (RectangleXY.Contains(XY(MeshXY0.GetVertex(Triangle.B))))
 		{
-			if (RectangleXY.Contains(XY(MeshXY0.GetVertex(VertexIndex))))
-			{
-				Result.Add(VertexIndex);
-			}
+			VertexIDs.Add(Triangle.B);
+		}
+		
+		if (RectangleXY.Contains(XY(MeshXY0.GetVertex(Triangle.C))))
+		{
+			VertexIDs.Add(Triangle.C);
 		}
-
-		return Result;
 	}
 	
-	// Returns the edges of the given mesh that lie inside the given rectangle. Assumptions:
-	// - Mesh vertices lie in the XY plane (have zero Z coordinate)
-	// - The camera is above the mesh and looking in the -Z direction (so no clip testing is needed)
-	TArray<int32> FindProjectedEdgesInRectangleXY0(const FCameraRectangle& CameraRectangle, const FDynamicMesh3& MeshXY0)
+	void AppendEdgeIDs(const FDynamicMesh3& MeshXY0, int TriangleID, TArray<int>& EdgeIDs)
 	{
-		FCameraRectangle::FAxisAlignedBox2 RectangleXY = GetRectangleXY(CameraRectangle);
+		const FIndex3i& Edges = MeshXY0.GetTriEdgesRef(TriangleID);
+		EdgeIDs.Add(Edges.A);
+		EdgeIDs.Add(Edges.B);
+		EdgeIDs.Add(Edges.C);
+	}
+	
+	void AppendEdgeIDsIfIntersected(const FDynamicMesh3& MeshXY0, const FAxisAlignedBox2d& RectangleXY, int TriangleID, TArray<int>& EdgeIDs)
+	{
+		const FIndex3i& Edges = MeshXY0.GetTriEdgesRef(TriangleID);
+
+		const FIndex2i& EdgeA = MeshXY0.GetEdgeRef(Edges.A).Vert;
+		const FSegment2d SegmentA(XY(MeshXY0.GetVertex(EdgeA.A)), XY(MeshXY0.GetVertex(EdgeA.B)));
+		if (TestIntersection(SegmentA, RectangleXY))
+		{
+			EdgeIDs.Add(Edges.A);
+		}
 		
-		TArray<int32> Result;
-		for (int32 EdgeIndex : MeshXY0.EdgeIndicesItr())
+		const FIndex2i& EdgeB = MeshXY0.GetEdgeRef(Edges.B).Vert;
+		const FSegment2d SegmentB(XY(MeshXY0.GetVertex(EdgeB.A)), XY(MeshXY0.GetVertex(EdgeB.B)));
+		if (TestIntersection(SegmentB, RectangleXY))
 		{
-			const FDynamicMesh3::FEdge& Edge = MeshXY0.GetEdgeRef(EdgeIndex);
-			FCameraRectangle::FSegment2 SegmentXY(XY(MeshXY0.GetVertex(Edge.Vert.A)),
-												  XY(MeshXY0.GetVertex(Edge.Vert.B)));
-			if (TestIntersection(SegmentXY, RectangleXY))
-			{
-				Result.Add(EdgeIndex);
-			}
+			EdgeIDs.Add(Edges.B);
+		}
+		
+		const FIndex2i& EdgeC = MeshXY0.GetEdgeRef(Edges.C).Vert;
+		const FSegment2d SegmentC(XY(MeshXY0.GetVertex(EdgeC.A)), XY(MeshXY0.GetVertex(EdgeC.B)));
+		if (TestIntersection(SegmentC, RectangleXY))
+		{
+			EdgeIDs.Add(Edges.C);
 		}
-		return Result;
 	}
 	
-	// Returns the triangles of the given mesh that lie inside the given rectangle. Assumptions:
-	// - Mesh vertices lie in the XY plane (have zero Z coordinate)
-	// - The camera is above the mesh and looking in the -Z direction (so no clip testing is needed)
-	TArray<int32> FindProjectedTrianglesInRectangleXY0(const FCameraRectangle& CameraRectangle, const FDynamicMesh3& MeshXY0)
+	void AppendTriangleID(const FDynamicMesh3&, int TriangleID, TArray<int>& TriangleIDs)
 	{
-		FCameraRectangle::FAxisAlignedBox2 RectangleXY = GetRectangleXY(CameraRectangle);
+		TriangleIDs.Add(TriangleID);
+	}
+
+	void AppendTriangleIDIfIntersected(const FDynamicMesh3& MeshXY0, const FAxisAlignedBox2d& RectangleXY, int TriangleID, TArray<int>& TriangleIDs)
+	{
+		const FIndex3i& Triangle = MeshXY0.GetTriangleRef(TriangleID);
+		const FTriangle2d TriangleXY(XY(MeshXY0.GetVertex(Triangle.A)),
+									 XY(MeshXY0.GetVertex(Triangle.B)),
+									 XY(MeshXY0.GetVertex(Triangle.C)));
+		
+		// Check with bTriangleIsOriented = false since some triangles maybe oriented away from the camera
+		if (FIntrTriangle2AxisAlignedBox2d Intersects(TriangleXY, RectangleXY, false); Intersects.Test())
+		{
+			TriangleIDs.Add(TriangleID);
+		}
+	}
+	
+	// Returns indices, collected by the given functions, from triangles which are intersected by the given rectangle.
+	// TreeXY0 must contain a mesh with vertices in the XY plane (have zero Z coordinate)
+	template<typename IDsFromTriangleF, typename IDsFromTriangleIfIntersectedF>
+	TArray<int32> FindAllIntersectionsAxisAlignedBox2(const FDynamicMeshAABBTree3& TreeXY0,
+													  const FAxisAlignedBox2d& RectangleXY,
+													  IDsFromTriangleF AppendIDs,
+													  IDsFromTriangleIfIntersectedF AppendIDsIfIntersected)
+	{
+		TRACE_CPUPROFILER_EVENT_SCOPE(FindAllIntersectionsAxisAlignedBox2);
+		
+		check(TreeXY0.GetMesh());
 		
 		TArray<int32> Result;
-		for (int32 TriangleIndex : MeshXY0.TriangleIndicesItr())
+		FAxisAlignedBox2d TreeRectangleXY;
+		TreeRectangleXY.Contain(XY(TreeXY0.GetBoundingBox().Min));
+		TreeRectangleXY.Contain(XY(TreeXY0.GetBoundingBox().Max));
+		if (RectangleXY.Contains(TreeRectangleXY))
 		{
-			const FIndex3i& Triangle = MeshXY0.GetTriangleRef(TriangleIndex);
-			const TTriangle2<FVector::FReal> TriangleXY(XY(MeshXY0.GetVertex(Triangle.A)),
-														XY(MeshXY0.GetVertex(Triangle.B)),
-														XY(MeshXY0.GetVertex(Triangle.C)));
-
-			// Check with bTriangleIsOriented = false since some triangles maybe oriented away from the camera
-			if (TIntrTriangle2AxisAlignedBox2<FVector::FReal> Intersects(TriangleXY, RectangleXY, false); Intersects.Test())
+			// Early out selecting everything
+			Result.Reserve(TreeXY0.GetMesh()->TriangleCount());
+			for (int TriangleID : TreeXY0.GetMesh()->TriangleIndicesItr())
 			{
-				Result.Add(TriangleIndex);
+				AppendIDs(*TreeXY0.GetMesh(), TriangleID, Result);
 			}
+			return Result;
 		}
+		
+		int SelectAllDepth = TNumericLimits<int>::Max();
+		int CurrentDepth = -1;
+		
+		// Traversal is depth first
+		FDynamicMeshAABBTree3::FTreeTraversal Traversal;
+		
+		Traversal.NextBoxF =
+			[&RectangleXY, &SelectAllDepth, &CurrentDepth](const FAxisAlignedBox3d& Box, int Depth)
+		{
+			CurrentDepth = Depth;
+			if (Depth > SelectAllDepth)
+			{
+				// We are deeper than the depth whose AABB was first detected to be contained in the RectangleXY,
+				// descend and collect all leaf triangles
+				return true;
+			}
+			
+			SelectAllDepth = TNumericLimits<int>::Max();
+			
+			const FAxisAlignedBox2d BoxXY(XY(Box.Min), XY(Box.Max));
+			if (RectangleXY.Intersects(BoxXY))
+			{
+				if (RectangleXY.Contains(BoxXY))
+				{
+					SelectAllDepth = Depth;
+				}
+				
+				return true;		
+			}
+			return false;
+		};
+		
+		Traversal.NextTriangleF =
+			[&RectangleXY, &SelectAllDepth, &CurrentDepth, &TreeXY0, &Result, &AppendIDs, &AppendIDsIfIntersected]
+			(int TriangleID)
+		{
+			if (CurrentDepth >= SelectAllDepth)
+			{
+				// This TriangleID is entirely contained in the selection rectangle so we can skip intersection testing
+				return AppendIDs(*TreeXY0.GetMesh(), TriangleID, Result);
+			}
+			return AppendIDsIfIntersected(*TreeXY0.GetMesh(), RectangleXY, TriangleID, Result);
+		};
+		
+		TreeXY0.DoTraversal(Traversal);
+
 		return Result;
 	}
 } // namespace MeshSelectionMechanicLocals
@@ -341,6 +434,9 @@ void UMeshSelectionMechanic::RebuildDrawnElements(const FTransform& StartTransfo
 
 	if (CurrentSelection.Type == FDynamicMeshSelection::EType::Triangle)
 	{
+		TRACE_CPUPROFILER_EVENT_SCOPE(RebuildDrawnElements_Triangle);
+		
+		LineSet->ReserveLines(CurrentSelection.SelectedIDs.Num() * 3);
 		for (int32 Tid : CurrentSelection.SelectedIDs)
 		{
 			FIndex3i Vids = CurrentSelection.Mesh->GetTriangle(Tid);
@@ -359,6 +455,9 @@ void UMeshSelectionMechanic::RebuildDrawnElements(const FTransform& StartTransfo
 	}
 	else if (CurrentSelection.Type == FDynamicMeshSelection::EType::Edge)
 	{
+		TRACE_CPUPROFILER_EVENT_SCOPE(RebuildDrawnElements_Edge);
+
+		LineSet->ReserveLines(CurrentSelection.SelectedIDs.Num());
 		for (int32 Eid : CurrentSelection.SelectedIDs)
 		{
 			FIndex2i EdgeVids = CurrentSelection.Mesh->GetEdgeV(Eid);
@@ -370,6 +469,8 @@ void UMeshSelectionMechanic::RebuildDrawnElements(const FTransform& StartTransfo
 	}
 	else if (CurrentSelection.Type == FDynamicMeshSelection::EType::Vertex)
 	{
+		TRACE_CPUPROFILER_EVENT_SCOPE(RebuildDrawnElements_Vertex);
+
 		PointSet->ReservePoints(CurrentSelection.SelectedIDs.Num());
 		for (int32 Vid : CurrentSelection.SelectedIDs)
 		{
@@ -383,6 +484,8 @@ void UMeshSelectionMechanic::RebuildDrawnElements(const FTransform& StartTransfo
 
 void UMeshSelectionMechanic::UpdateCentroid()
 {
+	TRACE_CPUPROFILER_EVENT_SCOPE(MeshSelectionMechanic_UpdateCentroid);
+
 	CurrentSelectionCentroid = FVector3d(0);
 	if (CurrentSelection.IsEmpty())
 	{
@@ -471,6 +574,8 @@ void UMeshSelectionMechanic::ClearCurrentSelection()
 
 void UMeshSelectionMechanic::UpdateCurrentSelection(const TSet<int32>& NewSelection, bool CalledFromOnDragRectangleChanged)
 {
+	TRACE_CPUPROFILER_EVENT_SCOPE(MeshSelectionMechanic_UpdateCurrentSelection);
+	
 	if (!ShouldRestartSelection() && CalledFromOnDragRectangleChanged)
 	{
 		// If we're modifying (adding/removing/toggling, but not restarting) the selection, we should start with the
@@ -638,6 +743,8 @@ void UMeshSelectionMechanic::OnDragRectangleStarted()
 
 void UMeshSelectionMechanic::OnDragRectangleChanged(const FCameraRectangle& CurrentRectangle)
 {
+	TRACE_CPUPROFILER_EVENT_SCOPE(MeshSelectionMechanic_OnDragRectangleChanged);
+	
 	auto FindIDsInsideCurrentRectangle = [this](auto&& AddIDsToSelection)
 	{
 		TSet<int32> RectangleSelectedIDs;
@@ -652,7 +759,7 @@ void UMeshSelectionMechanic::OnDragRectangleChanged(const FCameraRectangle& Curr
 				//  the case in the UV editor, this restriction should be lifted. When we do this we should also apply the
 				//  (inverse) tranform to the CurrentRectangle query rays. Search :ApplyTransformToQuery
 				ensure(MeshTransforms[MeshIndex].Identical(&FTransform::Identity, 0));
-				AddIDsToSelection(Mesh, MeshTransforms[MeshIndex], RectangleSelectedIDs);
+				AddIDsToSelection(*MeshSpatials[MeshIndex], RectangleSelectedIDs);
 	
 				if (!RectangleSelectedIDs.IsEmpty())
 				{
@@ -668,23 +775,25 @@ void UMeshSelectionMechanic::OnDragRectangleChanged(const FCameraRectangle& Curr
 			ensure(CurrentSelectionIndex != IndexConstants::InvalidID);
 			ensure(CurrentSelection.Mesh == MeshSpatials[CurrentSelectionIndex]->GetMesh());
 
-			const FDynamicMesh3& Mesh = *(MeshSpatials[CurrentSelectionIndex]->GetMesh());
-
 			// TODO See :ApplyTransformToQuery
 			ensure(MeshTransforms[CurrentSelectionIndex].Identical(&FTransform::Identity, 0));
-			AddIDsToSelection(Mesh, MeshTransforms[CurrentSelectionIndex], RectangleSelectedIDs);
+			AddIDsToSelection(*MeshSpatials[CurrentSelectionIndex], RectangleSelectedIDs);
 		}
 
 		return RectangleSelectedIDs;
 	};
+	using namespace MeshSelectionMechanicLocals;
+	FAxisAlignedBox2d RectangleXY = GetRectangleXY(CurrentRectangle);
 	
 	TSet<int32> RectangleSelectedIDs;
 	if (SelectionMode == EMeshSelectionMechanicMode::Vertex)
 	{
-		auto AddVertexIDsToSelection = [this, &CurrentRectangle]
-			(const FDynamicMesh3& Mesh, const FTransform& Transform, TSet<int32>& OutSelectedIDs)
+		TRACE_CPUPROFILER_EVENT_SCOPE(Vertex);
+	
+		auto AddVertexIDsToSelection = [this, &RectangleXY](const FDynamicMeshAABBTree3& Tree, TSet<int32>& OutSelectedIDs)
 		{
-			TArray<int32> SelectedIDs = MeshSelectionMechanicLocals::FindProjectedVerticesInRectangleXY0(CurrentRectangle, Mesh);
+			TArray<int32> SelectedIDs =
+				FindAllIntersectionsAxisAlignedBox2(Tree, RectangleXY, AppendVertexIDs, AppendVertexIDsIfIntersected);
 			OutSelectedIDs.Append(MoveTemp(SelectedIDs));
 		};
 	
@@ -692,10 +801,12 @@ void UMeshSelectionMechanic::OnDragRectangleChanged(const FCameraRectangle& Curr
 	}
 	else if (SelectionMode == EMeshSelectionMechanicMode::Edge)
 	{
-		auto AddEdgeIDsToSelection = [this, &CurrentRectangle]
-			(const FDynamicMesh3& Mesh, const FTransform& Transform, TSet<int32>& OutSelectedIDs)
+		TRACE_CPUPROFILER_EVENT_SCOPE(Edge);
+		
+		auto AddEdgeIDsToSelection = [this, &RectangleXY](const FDynamicMeshAABBTree3& Tree, TSet<int32>& OutSelectedIDs)
 		{
-			TArray<int32> SelectedIDs = MeshSelectionMechanicLocals::FindProjectedEdgesInRectangleXY0(CurrentRectangle, Mesh);
+			TArray<int32> SelectedIDs =
+				FindAllIntersectionsAxisAlignedBox2(Tree, RectangleXY, AppendEdgeIDs, AppendEdgeIDsIfIntersected);
 			OutSelectedIDs.Append(MoveTemp(SelectedIDs));
 		};
 
@@ -703,10 +814,12 @@ void UMeshSelectionMechanic::OnDragRectangleChanged(const FCameraRectangle& Curr
 	}
 	else if (SelectionMode == EMeshSelectionMechanicMode::Triangle)
 	{
-		auto AddTriangleIDsToSelection = [this, &CurrentRectangle]
-			(const FDynamicMesh3& Mesh, const FTransform& Transform, TSet<int32>& OutSelectedIDs)
+		TRACE_CPUPROFILER_EVENT_SCOPE(Triangle);
+		
+		auto AddTriangleIDsToSelection = [this, &RectangleXY] (const FDynamicMeshAABBTree3& Tree, TSet<int32>& OutSelectedIDs)
 		{
-			TArray<int32> SelectedIDs = MeshSelectionMechanicLocals::FindProjectedTrianglesInRectangleXY0(CurrentRectangle, Mesh);
+			TArray<int32> SelectedIDs =
+				FindAllIntersectionsAxisAlignedBox2(Tree, RectangleXY, AppendTriangleID, AppendTriangleIDIfIntersected);
 			OutSelectedIDs.Append(MoveTemp(SelectedIDs));
 		};
 	
@@ -714,11 +827,15 @@ void UMeshSelectionMechanic::OnDragRectangleChanged(const FCameraRectangle& Curr
 	}
 	else if (SelectionMode == EMeshSelectionMechanicMode::Component)
 	{
-		auto AddTriangleIDsToSelection = [this, &CurrentRectangle]
-			(const FDynamicMesh3& Mesh, const FTransform& Transform, TSet<int32>& OutSelectedIDs)
+		TRACE_CPUPROFILER_EVENT_SCOPE(Component);
+		
+		auto AddTriangleIDsToSelection = [this, &RectangleXY](const FDynamicMeshAABBTree3& Tree, TSet<int32>& OutSelectedIDs)
 		{
-			TArray<int32> SeedTriangles = MeshSelectionMechanicLocals::FindProjectedTrianglesInRectangleXY0(CurrentRectangle, Mesh);
-			FMeshConnectedComponents MeshSelectedComponent(&Mesh);
+			TArray<int32> SeedTriangles =
+				FindAllIntersectionsAxisAlignedBox2(Tree, RectangleXY, AppendTriangleID, AppendTriangleIDIfIntersected);
+
+			// TODO :NanitePerformance The following code is MUCH slower than AABB traversal
+			FMeshConnectedComponents MeshSelectedComponent(Tree.GetMesh());
 			MeshSelectedComponent.FindTrianglesConnectedToSeeds(SeedTriangles);
 			for (int ComponentIndex = 0; ComponentIndex < MeshSelectedComponent.Components.Num(); ComponentIndex++)
 			{
@@ -730,13 +847,15 @@ void UMeshSelectionMechanic::OnDragRectangleChanged(const FCameraRectangle& Curr
 	}
 	else if (SelectionMode == EMeshSelectionMechanicMode::Mesh)
 	{
-		auto AddAllMeshTriangleIDsToSelection = [this, &CurrentRectangle]
-			(const FDynamicMesh3& Mesh, const FTransform& Transform, TSet<int32>& OutSelectedIDs)
+		TRACE_CPUPROFILER_EVENT_SCOPE(Mesh);
+		
+		auto AddAllMeshTriangleIDsToSelection = [this, &RectangleXY] (const FDynamicMeshAABBTree3& Tree, TSet<int32>& OutSelectedIDs)
 		{
-			TArray<int32> SelectedIDs = MeshSelectionMechanicLocals::FindProjectedTrianglesInRectangleXY0(CurrentRectangle, Mesh);
+			TArray<int32> SelectedIDs =
+				FindAllIntersectionsAxisAlignedBox2(Tree, RectangleXY, AppendTriangleID, AppendTriangleIDIfIntersected);
 			if (!SelectedIDs.IsEmpty())
 			{
-				for (int32 Tid : Mesh.TriangleIndicesItr())
+				for (int32 Tid : Tree.GetMesh()->TriangleIndicesItr())
 				{
 					OutSelectedIDs.Add(Tid);
 				}
@@ -752,11 +871,14 @@ void UMeshSelectionMechanic::OnDragRectangleChanged(const FCameraRectangle& Curr
 
 	UpdateCurrentSelection(RectangleSelectedIDs, true);
 	
+	// TODO :NanitePerformance With large meshes and selections this call is much slower than AABB traversal (4x)
 	UpdateCentroid();
+	
+	// TODO :NanitePerformance With large meshes and selections this call is MUCH slower than AABB traversal (60x)
 	RebuildDrawnElements(FTransform(GetCurrentSelectionCentroid()));
 }
 
-void UMeshSelectionMechanic::OnDragRectangleFinished(const FCameraRectangle& Rectangle, bool bCancelled)
+void UMeshSelectionMechanic::OnDragRectangleFinished(const FCameraRectangle&, bool bCancelled)
 {
 	if (!bCancelled && (PreDragSelection != CurrentSelection))
 	{
diff --git a/Engine/Plugins/Experimental/UVEditor/Source/UVEditorTools/Private/UVSelectTool.cpp b/Engine/Plugins/Experimental/UVEditor/Source/UVEditorTools/Private/UVSelectTool.cpp
--- a/Engine/Plugins/Experimental/UVEditor/Source/UVEditorTools/Private/UVSelectTool.cpp
+++ b/Engine/Plugins/Experimental/UVEditor/Source/UVEditorTools/Private/UVSelectTool.cpp
@@ -296,6 +296,8 @@ void USelectToolActionPropertySet::PostAction(ESelectToolAction Action)
 
 void UUVSelectTool::Setup()
 {
+	TRACE_CPUPROFILER_EVENT_SCOPE(UVSelectTool_Setup);
+
 	using namespace UVSelectToolLocals;
 
 	check(Targets.Num() > 0);
@@ -373,8 +375,14 @@ void UUVSelectTool::Setup()
 		TSharedPtr<FDynamicMeshAABBTree3> Tree = TreeStore->Get(Target->UnwrapCanonical.Get());
 		if (!Tree)
 		{
+			TRACE_CPUPROFILER_EVENT_SCOPE(BuildAABBTreeForTarget);
 			Tree = MakeShared<FDynamicMeshAABBTree3>();
-			Tree->SetMesh(Target->UnwrapCanonical.Get());
+			Tree->SetMesh(Target->UnwrapCanonical.Get(), false);
+			// For now we split round-robin on the X/Y axes TODO Experiment with better splitting heuristics
+			FDynamicMeshAABBTree3::GetSplitAxisFunc GetSplitAxis = [](int Depth, const FAxisAlignedBox3d&) { return Depth % 2; };
+			// Note: 16 tris/leaf was chosen with data collected by SpatialBenchmarks.cpp in GeometryProcessingUnitTests
+			Tree->SetBuildOptions(16, MoveTemp(GetSplitAxis));
+			Tree->Build();
 			TreeStore->Set(Target->UnwrapCanonical.Get(), Tree, Target);
 		}
 		AABBTrees.Add(Tree);
@@ -443,6 +451,8 @@ void UUVSelectTool::Setup()
 
 void UUVSelectTool::Shutdown(EToolShutdownType ShutdownType)
 {
+	TRACE_CPUPROFILER_EVENT_SCOPE(UVSelectTool_Shutdown);
+	
 	// Clear selection so that it can be restored after undoing back into the select tool
 	if (!SelectionMechanic->GetCurrentSelection().IsEmpty())
 	{
@@ -488,7 +498,7 @@ void UUVSelectTool::Shutdown(EToolShutdownType ShutdownType)
 	ChangeRouter = nullptr;
 }
 
-void UUVSelectTool::SetSelection(const UE::Geometry::FDynamicMeshSelection& NewSelection, bool bBroadcastOnSelectionChanged)
+void UUVSelectTool::SetSelection(const FDynamicMeshSelection& NewSelection, bool bBroadcastOnSelectionChanged)
 {
 	SelectionMechanic->SetSelection(NewSelection, bBroadcastOnSelectionChanged, false);
 }
@@ -547,6 +557,8 @@ void UUVSelectTool::ConfigureSelectionModeFromControls()
 
 void UUVSelectTool::OnSelectionChanged()
 {
+	TRACE_CPUPROFILER_EVENT_SCOPE(UVSelectTool_OnSelectionChanged);
+	
 	using namespace UVSelectToolLocals;
 
 	const FDynamicMeshSelection& Selection = SelectionMechanic->GetCurrentSelection();
@@ -576,6 +588,8 @@ void UUVSelectTool::OnSelectionChanged()
 		TSet<int32> TidSet;
 		if (Selection.Type == FDynamicMeshSelection::EType::Triangle)
 		{
+			TRACE_CPUPROFILER_EVENT_SCOPE(Triangle);
+
 			const FDynamicMesh3* LivePreviewMesh = Targets[SelectionTargetIndex]->AppliedCanonical.Get();
 			for (int32 Tid : Selection.SelectedIDs)
 			{
@@ -612,6 +626,8 @@ void UUVSelectTool::OnSelectionChanged()
 		}
 		else if (Selection.Type == FDynamicMeshSelection::EType::Edge)
 		{
+			TRACE_CPUPROFILER_EVENT_SCOPE(Edge);
+
 			for (int32 Eid : Selection.SelectedIDs)
 			{
 				FIndex2i EdgeVids = Selection.Mesh->GetEdgeV(Eid);
@@ -638,6 +654,8 @@ void UUVSelectTool::OnSelectionChanged()
 		}
 		else if (Selection.Type == FDynamicMeshSelection::EType::Vertex)
 		{
+			TRACE_CPUPROFILER_EVENT_SCOPE(Vertex);
+
 			for (int32 Vid : Selection.SelectedIDs)
 			{
 				if (!VidSet.Contains(Vid))
@@ -674,6 +692,8 @@ void UUVSelectTool::OnSelectionChanged()
 
 void UUVSelectTool::UpdateLivePreviewLines()
 {
+	TRACE_CPUPROFILER_EVENT_SCOPE(UVSelectTool_UpdateLivePreviewLines);
+	
 	LivePreviewLineSet->Clear();
 
 	const FDynamicMeshSelection& Selection = SelectionMechanic->GetCurrentSelection();
@@ -847,6 +867,8 @@ void UUVSelectTool::ApplyAction(ESelectToolAction ActionType)
 	case ESelectToolAction::Sew:
 		if (SewAction)
 		{
+			TRACE_CPUPROFILER_EVENT_SCOPE(ApplyAction_Sew);
+
 			const FText TransactionName(LOCTEXT("SewCompleteTransactionName", "Sew Edges"));
 			EmitChangeAPI->BeginUndoTransaction(TransactionName);
 
@@ -859,6 +881,8 @@ void UUVSelectTool::ApplyAction(ESelectToolAction ActionType)
 	case ESelectToolAction::IslandConformalUnwrap:
 		if (IslandConformalUnwrapAction)
 		{
+			TRACE_CPUPROFILER_EVENT_SCOPE(ApplyAction_IslandConformalUnwrap);
+
 			const FText TransactionName(LOCTEXT("ConformalUnwrapCompleteTransactionName", "Conformal Unwrap Islands"));
 			EmitChangeAPI->BeginUndoTransaction(TransactionName);
 
