ModelingTools: Make baking tools exclusive.

#rnx

diff --git a/Engine/Plugins/Experimental/MeshModelingToolsetExp/Source/MeshModelingToolsExp/Public/BakeMeshAttributeMapsToolBase.h b/Engine/Plugins/Experimental/MeshModelingToolsetExp/Source/MeshModelingToolsExp/Public/BakeMeshAttributeMapsToolBase.h
--- a/Engine/Plugins/Experimental/MeshModelingToolsetExp/Source/MeshModelingToolsExp/Public/BakeMeshAttributeMapsToolBase.h
+++ b/Engine/Plugins/Experimental/MeshModelingToolsetExp/Source/MeshModelingToolsExp/Public/BakeMeshAttributeMapsToolBase.h
@@ -7,6 +7,7 @@
 #include "Materials/MaterialInstanceDynamic.h"
 #include "MultiSelectionTool.h"
 #include "InteractiveToolManager.h"
+#include "InteractiveToolQueryInterfaces.h" // for UInteractiveToolExclusiveToolAPI
 #include "AssetUtils/Texture2DBuilder.h"
 #include "DynamicMesh/DynamicMesh3.h"
 #include "DynamicMesh/DynamicMeshAABBTree3.h"
@@ -116,7 +117,7 @@ enum class EBakeMultisampling
  * Base Bake Maps tool
  */
 UCLASS()
-class MESHMODELINGTOOLSEXP_API UBakeMeshAttributeMapsToolBase : public UMultiSelectionTool, public UE::Geometry::IGenericDataOperatorFactory<UE::Geometry::FMeshMapBaker>
+class MESHMODELINGTOOLSEXP_API UBakeMeshAttributeMapsToolBase : public UMultiSelectionTool, public IInteractiveToolExclusiveToolAPI, public UE::Geometry::IGenericDataOperatorFactory<UE::Geometry::FMeshMapBaker>
 {
 	GENERATED_BODY()
 
diff --git a/Engine/Plugins/Experimental/MeshModelingToolsetExp/Source/MeshModelingToolsExp/Public/BakeMeshAttributeVertexTool.h b/Engine/Plugins/Experimental/MeshModelingToolsetExp/Source/MeshModelingToolsExp/Public/BakeMeshAttributeVertexTool.h
--- a/Engine/Plugins/Experimental/MeshModelingToolsetExp/Source/MeshModelingToolsExp/Public/BakeMeshAttributeVertexTool.h
+++ b/Engine/Plugins/Experimental/MeshModelingToolsetExp/Source/MeshModelingToolsExp/Public/BakeMeshAttributeVertexTool.h
@@ -5,6 +5,7 @@
 #include "CoreMinimal.h"
 #include "MultiSelectionTool.h"
 #include "InteractiveToolBuilder.h"
+#include "InteractiveToolQueryInterfaces.h" // for UInteractiveToolExclusiveToolAPI
 #include "DynamicMesh/DynamicMesh3.h"
 #include "DynamicMesh/DynamicMeshAABBTree3.h"
 #include "ModelingOperators.h"
@@ -165,7 +166,7 @@ public:
  * Vertex Baking Tool
  */
 UCLASS()
-class MESHMODELINGTOOLSEXP_API UBakeMeshAttributeVertexTool : public UMultiSelectionTool, public UE::Geometry::IGenericDataOperatorFactory<UE::Geometry::FMeshVertexBaker>
+class MESHMODELINGTOOLSEXP_API UBakeMeshAttributeVertexTool : public UMultiSelectionTool, public IInteractiveToolExclusiveToolAPI, public UE::Geometry::IGenericDataOperatorFactory<UE::Geometry::FMeshVertexBaker>
 {
 	GENERATED_BODY()
 
