ModelingTools: Update MapPaint properties to use the populated drop down of attribute names rather than the older direct index input methodology.

#rnx

diff --git a/Engine/Plugins/Experimental/MeshModelingToolsetExp/Source/MeshModelingToolsExp/Private/MeshAttributePaintTool.cpp b/Engine/Plugins/Experimental/MeshModelingToolsetExp/Source/MeshModelingToolsExp/Private/MeshAttributePaintTool.cpp
--- a/Engine/Plugins/Experimental/MeshModelingToolsetExp/Source/MeshModelingToolsExp/Private/MeshAttributePaintTool.cpp
+++ b/Engine/Plugins/Experimental/MeshModelingToolsetExp/Source/MeshModelingToolsExp/Private/MeshAttributePaintTool.cpp
@@ -100,8 +100,40 @@ public:
 
 };
 
+void UMeshAttributePaintToolProperties::Initialize(const TArray<FName>& AttributeNames, bool bInitialize)
+{
+	Attributes.Reset(AttributeNames.Num());
+	for (const FName& AttributeName : AttributeNames)
+	{
+		Attributes.Add(AttributeName.ToString());
+	}
+
+	if (bInitialize) {
+		Attribute = (Attributes.Num() > 0) ? Attributes[0] : TEXT("");
+	}
+}
 
 
+bool UMeshAttributePaintToolProperties::ValidateSelectedAttribute(bool bUpdateIfInvalid)
+{
+	int32 FoundIndex = Attributes.IndexOfByKey(Attribute);
+	if (FoundIndex == INDEX_NONE)
+	{
+		if (bUpdateIfInvalid)
+		{
+			Attribute = (Attributes.Num() > 0) ? Attributes[0] : TEXT("");
+		}
+		return false;
+	}
+	return true;
+}
+
+int32 UMeshAttributePaintToolProperties::GetSelectedAttributeIndex()
+{
+	ensure(INDEX_NONE == -1);
+	int32 FoundIndex = Attributes.IndexOfByKey(Attribute);
+	return FoundIndex;
+}
 
 
 void UMeshAttributePaintEditActions::PostAction(EMeshAttributePaintToolActions Action)
@@ -199,7 +231,7 @@ void UMeshAttributePaintTool::Setup()
 	*EditedMesh = *Cast<IMeshDescriptionProvider>(Target)->GetMeshDescription();
 
 	AttributeSource = MakeUnique<FMeshDescriptionVertexAttributeSource>(EditedMesh.Get());
-	AttribProps->Attributes = AttributeSource->GetAttributeList();
+	AttribProps->Initialize(AttributeSource->GetAttributeList(), true);
 
 	if (AttribProps->Attributes.Num() == 0)
 	{
@@ -211,8 +243,8 @@ void UMeshAttributePaintTool::Setup()
 	InitializeAttributes();
 	PendingNewSelectedIndex = 0;
 
-	SelectedAttributeWatcher.Initialize([this]() { return AttribProps->SelectedAttribute; },
-		[this](int32 NewValue) { PendingNewSelectedIndex = NewValue; }, AttribProps->SelectedAttribute);
+	SelectedAttributeWatcher.Initialize([this]() { AttribProps->ValidateSelectedAttribute(true);  return AttribProps->GetSelectedAttributeIndex(); },
+		[this](int32 NewValue) { PendingNewSelectedIndex = NewValue; }, AttribProps->GetSelectedAttributeIndex());
 
 	bVisibleAttributeValid = false;
 }
@@ -407,7 +439,7 @@ void UMeshAttributePaintTool::UpdateVisibleAttribute()
 	// copy current value set back to attribute  (should we just always be doing this??)
 	StoreCurrentAttribute();
 
-	CurrentAttributeIndex = AttribProps->SelectedAttribute;
+	CurrentAttributeIndex = AttribProps->GetSelectedAttributeIndex();
 
 	if (CurrentAttributeIndex >= 0)
 	{
@@ -427,7 +459,7 @@ void UMeshAttributePaintTool::UpdateVisibleAttribute()
 			}
 		});
 
-		AttribProps->AttributeName = AttribData.Name.ToString();
+		AttribProps->Attribute = AttribData.Name.ToString();
 	}
 }
 
@@ -630,9 +662,8 @@ void UMeshAttributePaintTool::ApplyAction(EMeshAttributePaintToolActions ActionT
 
 void UMeshAttributePaintTool::UpdateSelectedAttribute(int32 NewSelectedIndex)
 {
-	AttribProps->SelectedAttribute = 0;
-	AttribProps->Attributes = AttributeSource->GetAttributeList();
-	AttribProps->SelectedAttribute = FMath::Clamp(NewSelectedIndex, 0, AttribProps->Attributes.Num() - 1);
+	AttribProps->Initialize(AttributeSource->GetAttributeList(), false);
+	AttribProps->Attribute = AttribProps->Attributes[FMath::Clamp(NewSelectedIndex, 0, AttribProps->Attributes.Num() - 1)];
 	bVisibleAttributeValid = false;
 }
 
diff --git a/Engine/Plugins/Experimental/MeshModelingToolsetExp/Source/MeshModelingToolsExp/Public/MeshAttributePaintTool.h b/Engine/Plugins/Experimental/MeshModelingToolsetExp/Source/MeshModelingToolsExp/Public/MeshAttributePaintTool.h
--- a/Engine/Plugins/Experimental/MeshModelingToolsetExp/Source/MeshModelingToolsExp/Public/MeshAttributePaintTool.h
+++ b/Engine/Plugins/Experimental/MeshModelingToolsetExp/Source/MeshModelingToolsExp/Public/MeshAttributePaintTool.h
@@ -118,25 +118,39 @@ public:
 
 
 
-
-/**
- * Selected-Attribute settings Attribute Paint Tool
- */
 UCLASS()
 class MESHMODELINGTOOLSEXP_API UMeshAttributePaintToolProperties : public UInteractiveToolPropertySet
 {
 	GENERATED_BODY()
 public:
-	UPROPERTY(VisibleAnywhere, Category = Attribute)
-	TArray<FName> Attributes;
 
-	UPROPERTY(EditAnywhere, Category = Attribute, meta = (ArrayClamp = "Attributes"))
-	int SelectedAttribute = 0;
+	UPROPERTY(EditAnywhere, Category = "Attribute", meta = (DisplayName = "Selected Attribute", GetOptions = GetAttributeNames))
+	FString Attribute;
 
-	UPROPERTY(VisibleAnywhere, Category = Attribute)
-	FString AttributeName;
-};
+	UFUNCTION()
+	const TArray<FString>& GetAttributeNames() { return Attributes; };
+
+	TArray<FString> Attributes;
 
+public:
+	/**
+	* Initialize the internal array of attribute names
+	* @param bInitialize if set, selected Attribute will be reset to the first attribute or empty if there are none.
+	*/
+	void Initialize(const TArray<FName>& AttributeNames, bool bInitialize = false);
+
+	/**
+	 * Verify that the attribute selection is valid
+	 * @param bUpdateIfInvalid if selection is not valid, use attribute at index 0 or empty if there are no attributes
+	 * @return true if selection is in the Attributes array
+	 */
+	bool ValidateSelectedAttribute(bool bUpdateIfInvalid);
+
+	/**
+	 * @return selected attribute index, or -1 if invalid selection
+	 */
+	int32 GetSelectedAttributeIndex();
+};
 
 
 
