ModelingTools: Fix BakeTexture previews not clearing on removing map types.

#rnx

diff --git a/Engine/Plugins/Experimental/MeshModelingToolsetExp/Source/MeshModelingToolsExp/Public/BakeMeshAttributeMapsToolBase.h b/Engine/Plugins/Experimental/MeshModelingToolsetExp/Source/MeshModelingToolsExp/Public/BakeMeshAttributeMapsToolBase.h
--- a/Engine/Plugins/Experimental/MeshModelingToolsetExp/Source/MeshModelingToolsExp/Public/BakeMeshAttributeMapsToolBase.h
+++ b/Engine/Plugins/Experimental/MeshModelingToolsetExp/Source/MeshModelingToolsExp/Public/BakeMeshAttributeMapsToolBase.h
@@ -356,6 +356,7 @@ void UBakeMeshAttributeMapsToolBase::OnMapTypesUpdated(PropertySet& Properties)
 	}
 
 	// Initialize Properties->Result with requested MapTypes
+	Properties->Result.Empty();
 	for (const TTuple<EBakeMapType, TObjectPtr<UTexture2D>>& Map : CachedMaps)
 	{
 		// Only populate map types that were requested. Some map types like
