ModelingTools: Clean up BakeVertex properties display name.

#rnx

diff --git a/Engine/Plugins/Experimental/MeshModelingToolsetExp/Source/MeshModelingToolsExp/Public/BakeMeshAttributeVertexTool.h b/Engine/Plugins/Experimental/MeshModelingToolsetExp/Source/MeshModelingToolsExp/Public/BakeMeshAttributeVertexTool.h
--- a/Engine/Plugins/Experimental/MeshModelingToolsetExp/Source/MeshModelingToolsExp/Public/BakeMeshAttributeVertexTool.h
+++ b/Engine/Plugins/Experimental/MeshModelingToolsetExp/Source/MeshModelingToolsExp/Public/BakeMeshAttributeVertexTool.h
@@ -121,7 +121,7 @@ public:
 	bool bSplitAtNormalSeams = false;
 
 	/** Split vertex colors at UV seams */
-	UPROPERTY(EditAnywhere, Category = BakeSettings, DisplayName = "Split At UV Seams")
+	UPROPERTY(EditAnywhere, Category = BakeSettings)
 	bool bSplitAtUVSeams = false;
 };
 
