MeshModelingToolset: Clenup PathRev and BdryRev tools

diff --git a/Engine/Plugins/Editor/ModelingToolsEditorMode/Source/ModelingToolsEditorMode/Private/ModelingToolsManagerActions.cpp b/Engine/Plugins/Editor/ModelingToolsEditorMode/Source/ModelingToolsEditorMode/Private/ModelingToolsManagerActions.cpp
--- a/Engine/Plugins/Editor/ModelingToolsEditorMode/Source/ModelingToolsEditorMode/Private/ModelingToolsManagerActions.cpp
+++ b/Engine/Plugins/Editor/ModelingToolsEditorMode/Source/ModelingToolsEditorMode/Private/ModelingToolsManagerActions.cpp
@@ -58,8 +58,8 @@ void FModelingToolsManagerCommands::RegisterCommands()
 	// Create
 	REGISTER_MODELING_TOOL_COMMAND(BeginDrawPolygonTool, "PolyExt", "Draw and extrude polygons to create new objects");
 	REGISTER_MODELING_TOOL_COMMAND(BeginDrawPolyPathTool, "PathExt", "Draw and extrude PolyPaths to create new objects");
-	REGISTER_MODELING_TOOL_COMMAND(BeginDrawAndRevolveTool, "PathRev", "Draw and revolve a PolyPaths to create new objects");
-	REGISTER_MODELING_TOOL_COMMAND(BeginRevolveBoundaryTool, "BdryRev", "Revolve mesh boundary loops into new objects");
+	REGISTER_MODELING_TOOL_COMMAND(BeginDrawAndRevolveTool, "PathRev", "Draw and revolve PolyPaths to create new objects");
+	REGISTER_MODELING_TOOL_COMMAND(BeginRevolveBoundaryTool, "BdryRev", "Revolve mesh boundary loops to create new objects");
 	REGISTER_MODELING_TOOL_COMMAND(BeginCombineMeshesTool, "Append", "Combine selections into new objects");
 	REGISTER_MODELING_TOOL_COMMAND(BeginDuplicateMeshesTool, "Dupe", "Duplicate selections into new objects");
 
diff --git a/Engine/Plugins/Experimental/MeshModelingToolsetExp/Source/MeshModelingToolsExp/Private/DrawPolyPathTool.cpp b/Engine/Plugins/Experimental/MeshModelingToolsetExp/Source/MeshModelingToolsExp/Private/DrawPolyPathTool.cpp
--- a/Engine/Plugins/Experimental/MeshModelingToolsetExp/Source/MeshModelingToolsExp/Private/DrawPolyPathTool.cpp
+++ b/Engine/Plugins/Experimental/MeshModelingToolsetExp/Source/MeshModelingToolsExp/Private/DrawPolyPathTool.cpp
@@ -838,7 +838,7 @@ void UDrawPolyPathTool::EmitNewObject(EDrawPolyPathExtrudeMode ExtrudeMode)
 void UDrawPolyPathTool::ShowStartupMessage()
 {
 	GetToolManager()->DisplayMessage(
-		LOCTEXT("OnStartDraw", "Draw a path on the drawing plane, set its width, and extrude it. Left-click to place path vertices. Click on the last or first vertex to complete the path. Hold Shift to ignore snapping."),
+		LOCTEXT("OnStartDraw", "Draw a path on the drawing plane, set its width, and extrude it. Left-click to place path vertices, and click on the last or first vertex to complete the path. Hold Shift to ignore snapping."),
 		EToolMessageLevel::UserNotification);
 }
 
diff --git a/Engine/Plugins/Experimental/MeshModelingToolsetExp/Source/MeshModelingToolsExp/Private/Properties/RevolveProperties.cpp b/Engine/Plugins/Experimental/MeshModelingToolsetExp/Source/MeshModelingToolsExp/Private/Properties/RevolveProperties.cpp
--- a/Engine/Plugins/Experimental/MeshModelingToolsetExp/Source/MeshModelingToolsExp/Private/Properties/RevolveProperties.cpp
+++ b/Engine/Plugins/Experimental/MeshModelingToolsetExp/Source/MeshModelingToolsExp/Private/Properties/RevolveProperties.cpp
@@ -4,7 +4,6 @@
 
 #include "CompositionOps/CurveSweepOp.h"
 #include "Properties/MeshMaterialProperties.h"
-#include "Properties/RevolveProperties.h"
 #include "Util/RevolveUtil.h"
 
 #include "ExplicitUseGeometryMathTypes.h"		// using UE::Geometry::(math types)
@@ -28,27 +27,27 @@ void URevolveProperties::ApplyToCurveSweepOp(const UNewMeshMaterialProperties& M
 		}
 	}
 
-	double TotalRevolutionDegrees = (AlongAxisOffsetPerDegree == 0) ? ClampedRevolutionDegrees : RevolutionDegrees;
+	const double TotalRevolutionDegrees = (HeightOffsetPerDegree == 0) ? RevolveDegreesClamped : RevolveDegrees;
 
-	int32 Steps = bExplicitSteps ? NumExplicitSteps : FMath::CeilToInt(TotalRevolutionDegrees / MaxDegreesPerStep);
+	const int32 Steps = bExplicitSteps ? NumExplicitSteps : FMath::CeilToInt(TotalRevolutionDegrees / StepsMaxDegrees);
 	
 	double DegreesPerStep = TotalRevolutionDegrees / Steps;
-	double DegreesOffset = RevolutionDegreesOffset;
+	double DegreesOffset = RevolveDegreesOffset;
 	if (bReverseRevolutionDirection)
 	{
 		DegreesPerStep *= -1;
 		DegreesOffset *= -1;
 	}
-	double DownAxisOffsetPerStep = TotalRevolutionDegrees * AlongAxisOffsetPerDegree / Steps;
+	const double DownAxisOffsetPerStep = TotalRevolutionDegrees * HeightOffsetPerDegree / Steps;
 
-	if (bProfileIsCrossSectionOfSide && DegreesPerStep != 0 && abs(DegreesPerStep) < 180)
+	if (bPathAtMidpointOfStep && DegreesPerStep != 0 && abs(DegreesPerStep) < 180)
 	{
 		RevolveUtil::MakeProfileCurveMidpointOfFirstStep(CurveSweepOpOut.ProfileCurve, DegreesPerStep, RevolutionAxisOrigin, RevolutionAxisDirection);
 	}
 
 	// Generate the sweep curve
-	CurveSweepOpOut.bSweepCurveIsClosed = bWeldFullRevolution && AlongAxisOffsetPerDegree == 0 && TotalRevolutionDegrees == 360;
-	int32 NumSweepFrames = CurveSweepOpOut.bSweepCurveIsClosed ? Steps : Steps + 1; // If closed, last sweep frame is also first
+	CurveSweepOpOut.bSweepCurveIsClosed = bWeldFullRevolution && HeightOffsetPerDegree == 0 && TotalRevolutionDegrees == 360;
+	const int32 NumSweepFrames = CurveSweepOpOut.bSweepCurveIsClosed ? Steps : Steps + 1; // If closed, last sweep frame is also first
 	CurveSweepOpOut.SweepCurve.Reserve(NumSweepFrames);
 	RevolveUtil::GenerateSweepCurve(RevolutionAxisOrigin, RevolutionAxisDirection, DegreesOffset,
 		DegreesPerStep, DownAxisOffsetPerStep, NumSweepFrames, CurveSweepOpOut.SweepCurve);
@@ -60,9 +59,9 @@ void URevolveProperties::ApplyToCurveSweepOp(const UNewMeshMaterialProperties& M
 			RevolutionAxisDirection, AxisWeldTolerance, CurveSweepOpOut.ProfileVerticesToWeld);
 	}
 	CurveSweepOpOut.bSharpNormals = bSharpNormals;
-	CurveSweepOpOut.SharpNormalAngleTolerance = SharpNormalAngleTolerance;
-	CurveSweepOpOut.DiagonalTolerance = DiagonalProportionTolerance;
-	double UVScale = MaterialProperties.UVScale;
+	CurveSweepOpOut.SharpNormalAngleTolerance = SharpNormalsDegreeThreshold;
+	CurveSweepOpOut.DiagonalTolerance = QuadSplitCompactTolerance;
+	const double UVScale = MaterialProperties.UVScale;
 	CurveSweepOpOut.UVScale = FVector2d(UVScale, UVScale);
 	if (bReverseProfileCurve ^ bFlipVs)
 	{
@@ -74,29 +73,29 @@ void URevolveProperties::ApplyToCurveSweepOp(const UNewMeshMaterialProperties& M
 	CurveSweepOpOut.UnitUVInWorldCoordinates = 100; // This seems to be the case in the AddPrimitiveTool
 	switch (PolygroupMode)
 	{
-	case ERevolvePropertiesPolygroupMode::Single:
+	case ERevolvePropertiesPolygroupMode::PerShape:
 		CurveSweepOpOut.PolygonGroupingMode = EProfileSweepPolygonGrouping::Single;
 		break;
 	case ERevolvePropertiesPolygroupMode::PerFace:
 		CurveSweepOpOut.PolygonGroupingMode = EProfileSweepPolygonGrouping::PerFace;
 		break;
-	case ERevolvePropertiesPolygroupMode::PerStep:
+	case ERevolvePropertiesPolygroupMode::PerRevolveStep:
 		CurveSweepOpOut.PolygonGroupingMode = EProfileSweepPolygonGrouping::PerSweepSegment;
 		break;
-	case ERevolvePropertiesPolygroupMode::AccordingToProfileCurve:
+	case ERevolvePropertiesPolygroupMode::PerPathSegment:
 		CurveSweepOpOut.PolygonGroupingMode = EProfileSweepPolygonGrouping::PerProfileSegment;
 		break;
 	}
 	switch (QuadSplitMode)
 	{
-	case ERevolvePropertiesQuadSplit::ShortestDiagonal:
+	case ERevolvePropertiesQuadSplit::Compact:
 		CurveSweepOpOut.QuadSplitMode = EProfileSweepQuadSplit::ShortestDiagonal;
 		break;
 	case ERevolvePropertiesQuadSplit::Uniform:
 		CurveSweepOpOut.QuadSplitMode = EProfileSweepQuadSplit::Uniform;
 		break;
 	}
-	switch (CapFillMode)
+	switch (GetCapFillMode())
 	{
 	case ERevolvePropertiesCapFillMode::None:
 		CurveSweepOpOut.CapFillMode = FCurveSweepOp::ECapFillMode::None;
diff --git a/Engine/Plugins/Experimental/MeshModelingToolsetExp/Source/MeshModelingToolsExp/Private/RevolveBoundaryTool.cpp b/Engine/Plugins/Experimental/MeshModelingToolsetExp/Source/MeshModelingToolsExp/Private/RevolveBoundaryTool.cpp
--- a/Engine/Plugins/Experimental/MeshModelingToolsetExp/Source/MeshModelingToolsExp/Private/RevolveBoundaryTool.cpp
+++ b/Engine/Plugins/Experimental/MeshModelingToolsetExp/Source/MeshModelingToolsExp/Private/RevolveBoundaryTool.cpp
@@ -98,22 +98,22 @@ void URevolveBoundaryTool::Setup()
 	PlaneMechanic = NewObject<UConstructionPlaneMechanic>(this);
 	PlaneMechanic->Setup(this);
 	PlaneMechanic->Initialize(TargetWorld, FFrame3d(Settings->AxisOrigin, 
-		FRotator(Settings->AxisPitch, Settings->AxisYaw, 0).Quaternion()));
+		FRotator(Settings->AxisOrientation.X, Settings->AxisOrientation.Y, 0).Quaternion()));
 	PlaneMechanic->UpdateClickPriority(ClickBehavior->GetPriority().MakeLower());
 	PlaneMechanic->bShowGrid = false;
 	PlaneMechanic->OnPlaneChanged.AddLambda([this]() {
 		Settings->AxisOrigin = (FVector)PlaneMechanic->Plane.Origin;
 		FRotator AxisOrientation = ((FQuat)PlaneMechanic->Plane.Rotation).Rotator();
-		Settings->AxisPitch = AxisOrientation.Pitch;
-		Settings->AxisYaw = AxisOrientation.Yaw;
+		Settings->AxisOrientation.X = AxisOrientation.Pitch;
+		Settings->AxisOrientation.Y = AxisOrientation.Yaw;
 		UpdateRevolutionAxis();
 		});
 
-	Cast<IPrimitiveComponentBackedTarget>(Target)->SetOwnerVisibility(Settings->bDisplayOriginalMesh);
+	Cast<IPrimitiveComponentBackedTarget>(Target)->SetOwnerVisibility(Settings->bDisplayInputMesh);
 
 	SetToolDisplayName(LOCTEXT("ToolName", "Revolve Boundary"));
 	GetToolManager()->DisplayMessage(
-		LOCTEXT("OnStartRevolveBoundaryTool", "Revolve an open mesh boundary loop around an axis to create a new mesh. Ctrl+click will reposition the revolution axis, potentially aligning it with an edge."),
+		LOCTEXT("OnStartRevolveBoundaryTool", "Revolve an open mesh boundary loop around an axis to create a new mesh."),
 		EToolMessageLevel::UserNotification);
 	if (Topology->Edges.Num() == 1)
 	{
@@ -125,13 +125,13 @@ void URevolveBoundaryTool::Setup()
 	else if (Topology->Edges.Num() == 0)
 	{
 		GetToolManager()->DisplayMessage(
-			LOCTEXT("NoBoundaryLoops", "This mesh does not have any boundary loops to display and revolve. Delete some faces or use a different mesh."),
+			LOCTEXT("NoBoundaryLoops", "This mesh does not have any boundary loops to display and revolve. Delete faces to create a boundary or use a different mesh."),
 			EToolMessageLevel::UserWarning);
 	}
 	else
 	{
 		GetToolManager()->DisplayMessage(
-			LOCTEXT("OnStartRevolveBoundaryToolMultipleBoundaries", "Your mesh has multiple boundaries- Click the one you wish to use"),
+			LOCTEXT("OnStartRevolveBoundaryToolMultipleBoundaries", "Your mesh has multiple boundaries. Select the one you wish to use."),
 			EToolMessageLevel::UserWarning);
 	}
 }
@@ -187,8 +187,8 @@ void URevolveBoundaryTool::OnClicked(const FInputDeviceRay& ClickPos)
 
 			Settings->AxisOrigin = (FVector)RevolutionAxisFrame.Origin;
 			FRotator AxisOrientation = ((FQuat)RevolutionAxisFrame.Rotation).Rotator();
-			Settings->AxisPitch = AxisOrientation.Pitch;
-			Settings->AxisYaw = AxisOrientation.Yaw;
+			Settings->AxisOrientation.X = AxisOrientation.Pitch;
+			Settings->AxisOrientation.Y = AxisOrientation.Yaw;
 			UpdateRevolutionAxis();
 		}
 
@@ -215,7 +215,7 @@ bool URevolveBoundaryTool::CanAccept() const
 void URevolveBoundaryTool::UpdateRevolutionAxis()
 {
 	RevolutionAxisOrigin = (FVector3d)Settings->AxisOrigin;
-	RevolutionAxisDirection = (FVector3d)FRotator(Settings->AxisPitch, Settings->AxisYaw, 0).RotateVector(FVector(1, 0, 0));
+	RevolutionAxisDirection = (FVector3d)FRotator(Settings->AxisOrientation.X, Settings->AxisOrientation.Y, 0).RotateVector(FVector(1, 0, 0));
 	if (Preview)
 	{
 		Preview->InvalidateResult();
@@ -278,7 +278,7 @@ void URevolveBoundaryTool::GenerateAsset(const FDynamicMeshOpResult& OpResult)
 		return;
 	}
 
-	GetToolManager()->BeginUndoTransaction(LOCTEXT("RevolveBoundaryToolTransactionName", "Revolve Tool"));
+	GetToolManager()->BeginUndoTransaction(LOCTEXT("RevolveBoundaryToolTransactionName", "Boundary Revolve Tool"));
 
 	FCreateMeshObjectParams NewMeshObjectParams;
 	NewMeshObjectParams.TargetWorld = TargetWorld;
@@ -338,10 +338,10 @@ void URevolveBoundaryTool::Render(IToolsContextRenderAPI* RenderAPI)
 void URevolveBoundaryTool::OnPropertyModified(UObject* PropertySet, FProperty* Property)
 {
 	PlaneMechanic->SetPlaneWithoutBroadcast(FFrame3d(Settings->AxisOrigin, 
-		FRotator(Settings->AxisPitch, Settings->AxisYaw, 0).Quaternion()));
+		FRotator(Settings->AxisOrientation.X, Settings->AxisOrientation.Y, 0).Quaternion()));
 	UpdateRevolutionAxis();
 
-	Cast<IPrimitiveComponentBackedTarget>(Target)->SetOwnerVisibility(Settings->bDisplayOriginalMesh);
+	Cast<IPrimitiveComponentBackedTarget>(Target)->SetOwnerVisibility(Settings->bDisplayInputMesh);
 
 	if (Preview)
 	{
diff --git a/Engine/Plugins/Experimental/MeshModelingToolsetExp/Source/MeshModelingToolsExp/Public/Properties/RevolveProperties.h b/Engine/Plugins/Experimental/MeshModelingToolsetExp/Source/MeshModelingToolsExp/Public/Properties/RevolveProperties.h
--- a/Engine/Plugins/Experimental/MeshModelingToolsetExp/Source/MeshModelingToolsExp/Public/Properties/RevolveProperties.h
+++ b/Engine/Plugins/Experimental/MeshModelingToolsetExp/Source/MeshModelingToolsExp/Public/Properties/RevolveProperties.h
@@ -12,48 +12,40 @@
 class UNewMeshMaterialProperties;
 PREDECLARE_GEOMETRY(class FCurveSweepOp);
 
-UENUM()
 enum class ERevolvePropertiesCapFillMode : uint8
 {
-	/** No cap. */
 	None,
-	/** Cap is triangulated to maximize the minimal angle in the triangles (if they were to be
-	   projected onto a best-fit plane). */
+	CenterFan,
 	Delaunay,
-	/** Cap is triangualted using a standard ear clipping approach. This could result in some
-	   very thin triangles. */
-	EarClipping,
-	/** A vertex is placed in the center and a fan is created to the boundary. This is nice if
-	   the cross section is convex, but creates invalid geometry if it isn't. */
-	CenterFan
+	EarClipping
 };
 
 
 UENUM()
 enum class ERevolvePropertiesPolygroupMode : uint8
 {
-	/** One polygroup for body of output mesh */
-	Single,
-	/** One polygroup per generated quad/triangle. */
+	/** One Polygroup for the entire shape */
+	PerShape,
+	/** One Polygroup for each geometric face */
 	PerFace,
-	/** Groups will be arranged in strips running in the profile curve direction, one per revolution step. */
-	PerStep,
-	/** Groups will be arranged in strips running along in the revolution direction according to profile curve. */
-	AccordingToProfileCurve
+	/** One PolyGroup along the path for each revolution step */
+	PerRevolveStep,
+	/** One PolyGroup along the revolution steps for each path segment */
+	PerPathSegment
 };
 
 UENUM()
 enum class ERevolvePropertiesQuadSplit : uint8
 {
-	/** Quads will always be split the same way relative to an unrolled mesh, regardless of quad shape. */
+	/** Quads will always be split into triangles the same way regardless of quad shape. */
 	Uniform,
 	
-	/** Quads will be split such that the shortest diagonal is connected. */
-	ShortestDiagonal
+	/** Quads will be split into triangles by connecting the shortest diagonal. */
+	Compact
 };
 
 /**
- * Common properties for revolving a polyline to create a mesh.
+ * Common properties for revolving a PolyPath to create a mesh.
  */
 UCLASS()
 class MESHMODELINGTOOLSEXP_API URevolveProperties : public UInteractiveToolPropertySet
@@ -65,126 +57,117 @@ public:
 	{
 		// We want the revolution degrees to be clamped to 360 when there is no offset along axis, and extendable
 		// beyond 360 when there is one (to make springs, etc). There's not currently a way to have a conditional
-		// clamp value, so instead, we swap the variables that we make visible to the user (RevolutionDegrees and
-		// ClampedRevolutionDegrees). However, we have to do some work to keep the values consistent so that they
+		// clamp value, so instead, we swap the variables that we make visible to the user (RevolveDegrees and
+		// RevolveDegreesClamped). However, we have to do some work to keep the values consistent so that they
 		// do not jump when they are swapped.
-		ClampedRevolutionDegreesWatcherIndex = WatchProperty(ClampedRevolutionDegrees,
+		ClampedRevolutionDegreesWatcherIndex = WatchProperty(RevolveDegreesClamped,
 			[this](const double& NewClamped) 
 			{ 
-				RevolutionDegrees = NewClamped; 
+				RevolveDegrees = NewClamped; 
 				SilentUpdateWatcherAtIndex(RevolutionDegreesWatcherIndex); //avoid triggering that watcher
 			});
 
-		RevolutionDegreesWatcherIndex = WatchProperty(RevolutionDegrees,
+		RevolutionDegreesWatcherIndex = WatchProperty(RevolveDegrees,
 			[this](const double& NewNonClamped)
 			{
-				ClampedRevolutionDegrees = FMath::Min(RevolutionDegrees, 360.0);
+				RevolveDegreesClamped = FMath::Min(RevolveDegrees, 360.0);
 				SilentUpdateWatcherAtIndex(ClampedRevolutionDegreesWatcherIndex); //avoid triggering that watcher
 			});
 	}
 
-	/** Revolution extent. Clamped to a max of 360 when not offsetting along axis.*/
-	UPROPERTY(EditAnywhere, Category = RevolveSettings, meta = (UIMin = "0", UIMax = "360", ClampMin = "0", ClampMax = "360", 
-		EditCondition = "AlongAxisOffsetPerDegree == 0", EditConditionHides, DisplayName = "Revolution Degrees"))
-	double ClampedRevolutionDegrees = 360;
+	/** Revolution extent in degrees. Clamped to a maximum of 360 if Height Offset Per Degree is set to 0. */
+	UPROPERTY(EditAnywhere, Category = Revolve, meta = (DisplayName = "Degrees", UIMin = "0", UIMax = "360", ClampMin = "0", ClampMax = "360", 
+		EditCondition = "HeightOffsetPerDegree == 0", EditConditionHides))
+	double RevolveDegreesClamped = 360;
 
-	/** Revolution extent. */
-	UPROPERTY(EditAnywhere, Category = RevolveSettings, meta = (UIMin = "0", UIMax = "3600", ClampMin = "0", ClampMax = "360000",
-		EditCondition = "AlongAxisOffsetPerDegree != 0", EditConditionHides))
-	double RevolutionDegrees = 360;
+	/** Revolution extent in degrees. Clamped to a maximum of 360 if Height Offset Per Degree is set to 0. */
+	UPROPERTY(EditAnywhere, Category = Revolve, meta = (DisplayName = "Degrees", UIMin = "0", UIMax = "3600", ClampMin = "0", ClampMax = "360000",
+		EditCondition = "HeightOffsetPerDegree != 0", EditConditionHides))
+	double RevolveDegrees = 360;
 
-	/** When true, the number of steps can be specified explicitly. When false, the number of steps is adjusted automatically. */
-	UPROPERTY(EditAnywhere, Category = RevolveSettings)
-	bool bExplicitSteps = false;
+	/** The angle by which to rotate the path around the axis before beginning the revolve. */
+	UPROPERTY(EditAnywhere, Category = Revolve, meta = (DisplayName = "Degrees Offset", UIMin = "-360", UIMax = "360", ClampMin = "-36000", ClampMax = "36000"))
+	double RevolveDegreesOffset = 0;
 
-	/** Number of steps to take while revolving. */
-	UPROPERTY(EditAnywhere, Category = RevolveSettings, meta = (UIMin = "1", ClampMin = "1", UIMax = "100", ClampMax = "5000",
-		Display = "Steps", EditCondition = "bExplicitSteps", EditConditionHides))
-	int NumExplicitSteps = 24;
+	/** Implicitly defines the number of steps in the revolution such that each step moves the revolution no more than the given number of degrees. This is only available if Explicit Steps is disabled. */
+	UPROPERTY(EditAnywhere, Category = Revolve, meta = (UIMin = "1", ClampMin = "1", UIMax = "120", ClampMax = "180", EditCondition = "!bExplicitSteps"))
+	double StepsMaxDegrees = 15;
 
-	/** The revolution is split into a number of steps such that each step moves the revolution no more than this number of degrees.  */
-	UPROPERTY(EditAnywhere, Category = RevolveSettings, meta = (UIMin = "1", ClampMin = "1", UIMax = "120", ClampMax = "180",
-		EditCondition = "!bExplicitSteps", EditConditionHides))
-	double MaxDegreesPerStep = 15;
+	/** If true, the number of steps can be specified explicitly via Steps. If false, the number of steps is adjusted automatically based on Steps Max Degrees. */
+	UPROPERTY(EditAnywhere, Category = Revolve)
+	bool bExplicitSteps = false;
 
-	/** How far to move each step along the axis (per degree). Used to create spirals. */
-	UPROPERTY(EditAnywhere, Category = RevolveSettings, meta = (UIMin = "-1", UIMax = "1", ClampMin = "-100000", ClampMax = "100000"))
-	double AlongAxisOffsetPerDegree = 0;
+	/** Number of steps in the revolution. This is only available if Explicit Steps is enabled. */
+	UPROPERTY(EditAnywhere, Category = Revolve, meta = (DisplayName = "Steps", UIMin = "1", ClampMin = "1", UIMax = "100", ClampMax = "5000",
+		EditCondition = "bExplicitSteps"))
+	int NumExplicitSteps = 24;
 
-	/** The angle by which to shift the profile curve around the axis before beginning the revolve */
-	UPROPERTY(EditAnywhere, Category = RevolveSettings, meta = (UIMin = "-360", UIMax = "360", ClampMin = "-36000", ClampMax = "36000"), AdvancedDisplay)
-	double RevolutionDegreesOffset = 0;
+	/** How far to move each step along the revolution axis per degree. Non-zero values are useful for creating spirals. */
+	UPROPERTY(EditAnywhere, Category = Revolve, meta = (DisplayName = "Height Offset", UIMin = "-1", UIMax = "1", ClampMin = "-100000", ClampMax = "100000"))
+	double HeightOffsetPerDegree = 0;
 
-	/** By default, revolution is done counterclockwise if looking down the revolution axis. This reverses the direction.*/
-	UPROPERTY(EditAnywhere, Category = RevolveSettings)
+	/** By default, revolution is done counterclockwise if looking down the revolution axis. This reverses the revolution direction to clockwise.*/
+	UPROPERTY(EditAnywhere, Category = Revolve, meta = (DisplayName = "Reverse"))
 	bool bReverseRevolutionDirection = false;
 
 	/** Flips the mesh inside out. */
-	UPROPERTY(EditAnywhere, Category = RevolveSettings)
+	UPROPERTY(EditAnywhere, Category = Revolve)
 	bool bFlipMesh = false;
 
-	/** If true, then rather than revolving the profile directly, it is interpreted as the midpoint cross section of
-	 the first rotation step. Useful, for instance, for using the tool to create square columns. */
-	UPROPERTY(EditAnywhere, Category = RevolveSettings, AdvancedDisplay)
-	bool bProfileIsCrossSectionOfSide = false;
+	/** If true, normals are not averaged or shared between triangles beyond the Sharp Normals Degree Threshold. */
+	UPROPERTY(EditAnywhere, Category = Revolve)
+	bool bSharpNormals = false;
+
+	/** The threshold in degrees beyond which normals are not averaged or shared between triangles anymore. This is only available if Sharp Normals is enabled. */
+	UPROPERTY(EditAnywhere, Category = Revolve,	meta = (DisplayName = "Sharp Normals Threshold", ClampMin = "0.0", ClampMax = "90.0",
+		EditCondition = "bSharpNormals"))
+	double SharpNormalsDegreeThreshold = 0.1;
 
-	/** Determines grouping of generated triangles into polygroups. 
-	   Caps (if present) will always be separate groups. */
-	UPROPERTY(EditAnywhere, Category = RevolveSettings, AdvancedDisplay)
+	/** If true, the path is placed at the midpoint of each step instead of at the start and/or end of a step. For example, this is useful for creating square columns. */
+	UPROPERTY(EditAnywhere, Category = Revolve, AdvancedDisplay, meta = (DisplayName = "Path at Midpoint"))
+	bool bPathAtMidpointOfStep = false;
+
+	/** How Polygroups are assigned to shape primitives. If caps are generated, they will always be in separate groups. */
+	UPROPERTY(EditAnywhere, Category = Revolve, AdvancedDisplay)
 	ERevolvePropertiesPolygroupMode PolygroupMode = ERevolvePropertiesPolygroupMode::PerFace;
 
-	/** Determines how any generated quads are split into triangles. */
-	UPROPERTY(EditAnywhere, Category = RevolveSettings, AdvancedDisplay)
-	ERevolvePropertiesQuadSplit QuadSplitMode = ERevolvePropertiesQuadSplit::ShortestDiagonal;
-	
-	/** When quads are generated using "shortest" diagonal, this biases the diagonal length comparison
-	 to prefer one slightly in the case of similar diagonals (for example, a value of 0.01 allows a
-	 1% difference in lengths before the triangulation is flipped). Helps symmetric quads be uniformly
-	 triangulated. */
-	UPROPERTY(EditAnywhere, Category = RevolveSettings, AdvancedDisplay, meta = (ClampMin = "0.0",  ClampMax = "2.0", 
-		EditCondition = "QuadSplitMode == ERevolvePropertiesQuadSplit::ShortestDiagonal", EditConditionHides))
-	double DiagonalProportionTolerance = 0.01;
-
-	/** Determines how caps are created if the revolution is partial. Not relevant if the
-	  revolution is full and welded. */
-	UPROPERTY(EditAnywhere, Category = RevolveSettings, AdvancedDisplay, meta = (
-		EditCondition = "AlongAxisOffsetPerDegree != 0 || RevolutionDegrees != 360", EditConditionHides))
-	ERevolvePropertiesCapFillMode CapFillMode = ERevolvePropertiesCapFillMode::Delaunay;
-
-	/** If true, the ends of a fully revolved profile are welded together, rather than duplicating
-	  vertices at the seam. Not relevant if the revolution is not full. */
-	UPROPERTY(EditAnywhere, Category = RevolveSettings, AdvancedDisplay, meta = (
-		EditCondition = "AlongAxisOffsetPerDegree == 0 && RevolutionDegrees == 360", EditConditionHides))
+	/** How generated quads are split into triangles. */
+	UPROPERTY(EditAnywhere, Category = Revolve, AdvancedDisplay)
+	ERevolvePropertiesQuadSplit QuadSplitMode = ERevolvePropertiesQuadSplit::Compact;
+
+	/**
+	 * For quad split mode Compact, this biases the length comparison to prefer one diagonal over the other.
+	 * For example, a value of 0.01 allows a 1% difference in lengths before the triangulation is flipped. This helps symmetric quads to
+	 * be triangulated more uniformly.
+	 */
+	// UPROPERTY(EditAnywhere, Category = Revolve, AdvancedDisplay, meta = (DisplayName = "Quad Split Tolerance", ClampMin = "0.0",  ClampMax = "2.0",
+	// 	EditCondition = "QuadSplitMode == ERevolvePropertiesQuadSplit::Compact"))
+	double QuadSplitCompactTolerance = 0.01;
+
+	// /** Determines how end caps are created. This is not relevant if the end caps are not visible or if the path is not closed. */
+	// UPROPERTY(EditAnywhere, Category = Revolve, AdvancedDisplay, meta = (EditCondition = "HeightOffsetPerDegree != 0 || RevolveDegrees != 360"))
+	// ERevolvePropertiesCapFillMode CapFillMode = ERevolvePropertiesCapFillMode::Delaunay;
+
+	/** If true, the ends of a fully revolved path are welded together, rather than duplicating vertices at the seam.
+	 * This is not relevant if the revolution is not full and/or there is a height offset. */
+	// UPROPERTY(EditAnywhere, Category = Revolve, AdvancedDisplay, meta = (EditCondition = "HeightOffsetPerDegree == 0 && RevolveDegrees == 360"))
 	bool bWeldFullRevolution = true;
 
-	/** If true, vertices sufficiently close to the axis will not be replicated, instead reusing
-	  the same vertex for any adjacent triangles. */
-	UPROPERTY(EditAnywhere, Category = RevolveSettings, AdvancedDisplay)
+	/** If true, vertices close to the axis will not be replicated, instead reusing the same vertex for any adjacent triangles. */
+	// UPROPERTY(EditAnywhere, Category = Revolve, AdvancedDisplay)
 	bool bWeldVertsOnAxis = true;
 
-	/** If welding vertices on the axis, the distance that a vertex can be from the axis and still be welded */
-	UPROPERTY(EditAnywhere, Category = RevolveSettings, AdvancedDisplay, meta = (ClampMin = "0.0", ClampMax = "20.0", 
-		EditCondition = "bWeldVertsOnAxis", EditConditionHides))
+	/** If welding vertices on the axis, the distance that a vertex can be from the axis and still be welded. */
+	// UPROPERTY(EditAnywhere, Category = Revolve, AdvancedDisplay, meta = (ClampMin = "0.0", ClampMax = "20.0", EditCondition = "bWeldVertsOnAxis"))
 	double AxisWeldTolerance = 0.1;
 
-	/** If true, normals are not averaged or shared between triangles with sufficient angle difference. */
-	UPROPERTY(EditAnywhere, Category = RevolveSettings)
-	bool bSharpNormals = false;
-
-	/** When using sharp normals, the degree difference to accept between adjacent triangle normals to allow them to share
-	 normals at their vertices. */
-	UPROPERTY(EditAnywhere, Category = RevolveSettings, meta = (ClampMin = "0.0", ClampMax = "90.0", 
-		EditCondition = "bSharpNormals", EditConditionHides))
-	double SharpNormalAngleTolerance = 0.1;
-
 	/** If true, UV coordinates will be flipped in the V direction. */
-	UPROPERTY(EditAnywhere, Category = RevolveSettings, AdvancedDisplay)
+	// UPROPERTY(EditAnywhere, Category = Revolve, AdvancedDisplay)
 	bool bFlipVs = false;
 
-	/* If true, UV layout is not affected by segments of the profile curve that 
-	  do not result in any triangles (i.e., when both ends of the segment are welded
-	  due to being on the revolution axis).*/
-	UPROPERTY(EditAnywhere, Category = RevolveSettings, AdvancedDisplay)
+	/* If true, UV layout is not affected by segments of the path that do not result in any triangles. For example, when both ends of
+	 * a segment are welded due to being on the revolution axis.*/
+	// UPROPERTY(EditAnywhere, Category = Revolve, AdvancedDisplay)
 	bool bUVsSkipFullyWeldedEdges = true;
 
 
@@ -199,6 +182,11 @@ public:
 		UE::Geometry::FCurveSweepOp& CurveSweepOpOut) const;
 
 	protected:
+		virtual ERevolvePropertiesCapFillMode GetCapFillMode() const
+		{
+			return ERevolvePropertiesCapFillMode::Delaunay;
+		}
+	
 		int32 ClampedRevolutionDegreesWatcherIndex;
 		int32 RevolutionDegreesWatcherIndex;
 };
diff --git a/Engine/Plugins/Experimental/MeshModelingToolsetExp/Source/MeshModelingToolsExp/Public/RevolveBoundaryTool.h b/Engine/Plugins/Experimental/MeshModelingToolsetExp/Source/MeshModelingToolsExp/Public/RevolveBoundaryTool.h
--- a/Engine/Plugins/Experimental/MeshModelingToolsetExp/Source/MeshModelingToolsExp/Public/RevolveBoundaryTool.h
+++ b/Engine/Plugins/Experimental/MeshModelingToolsetExp/Source/MeshModelingToolsExp/Public/RevolveBoundaryTool.h
@@ -5,14 +5,11 @@
 #include "CoreMinimal.h"
 
 #include "BaseBehaviors/BehaviorTargetInterfaces.h"
-#include "InteractiveToolBuilder.h"
-#include "SingleSelectionTool.h"
 #include "Mechanics/ConstructionPlaneMechanic.h"
 #include "MeshBoundaryToolBase.h"
 #include "MeshOpPreviewHelpers.h" //UMeshOpPreviewWithBackgroundCompute
 #include "ModelingOperators.h" //IDynamicMeshOperatorFactory
 #include "InteractiveTool.h"
-#include "InteractiveToolBuilder.h" //UInteractiveToolBuilder
 #include "Properties/MeshMaterialProperties.h"
 #include "Properties/RevolveProperties.h"
 #include "ToolContextInterfaces.h" // FToolBuilderState
@@ -43,28 +40,49 @@ public:
 	TObjectPtr<URevolveBoundaryTool> RevolveBoundaryTool;
 };
 
+UENUM()
+enum class ERevolvePropertiesBoundaryCapFillMode : uint8
+{
+	/** No caps will be generated. */
+	None = static_cast<uint8>(ERevolvePropertiesCapFillMode::None),
+	/** Caps are triangulated by placing a vertex in the center and creating a fan to the boundary. This works well if the path is convex,
+	  * but can create invalid geometry if it is concave. */
+	CenterFan = static_cast<uint8>(ERevolvePropertiesCapFillMode::CenterFan),
+	/** Caps are triangulated to maximize the minimal angle in the triangles using Delaunay triangulation. */
+	Delaunay = static_cast<uint8>(ERevolvePropertiesCapFillMode::Delaunay),
+	/** Caps are triangulated using a standard ear clipping approach. This could result in some very thin triangles. */
+	EarClipping = static_cast<uint8>(ERevolvePropertiesCapFillMode::EarClipping)
+};
+
 UCLASS()
 class MESHMODELINGTOOLSEXP_API URevolveBoundaryToolProperties : public URevolveProperties
 {
 	GENERATED_BODY()
 
 public:
-	
-	UPROPERTY(EditAnywhere, Category = RevolveSettings, AdvancedDisplay)
-	bool bDisplayOriginalMesh = false;
 
-	UPROPERTY(EditAnywhere, Category = RevolutionAxis)
+	/** Determines how end caps are created. This is not relevant if the end caps are not visible or if the path is not closed. */
+	UPROPERTY(EditAnywhere, Category = Revolve, AdvancedDisplay, meta = (DisplayAfter = "QuadSplitMode",
+		EditCondition = "HeightOffsetPerDegree != 0 || RevolveDegrees != 360"))
+	ERevolvePropertiesBoundaryCapFillMode CapFillMode = ERevolvePropertiesBoundaryCapFillMode::Delaunay;
+
+	/** If true, displays the original mesh in addition to the revolved boundary. */
+	UPROPERTY(EditAnywhere, Category = Revolve, AdvancedDisplay)
+	bool bDisplayInputMesh = false;
+
+	/** Sets the revolution axis origin. */
+	UPROPERTY(EditAnywhere, Category = RevolutionAxis, meta = (DisplayName = "Origin"))
 	FVector AxisOrigin = FVector(0, 0, 0);
 
-	//~ We don't use a rotator for axis orientation because one of the components (roll) 
-	//~ will never do anything in the case of our axis.
-	UPROPERTY(EditAnywhere, Category = RevolutionAxis, meta = (
-		UIMin = -180, UIMax = 180, ClampMin = -180000, ClampMax = 18000))
-	float AxisYaw = 0;
+	/** Sets the revolution axis pitch and yaw. */
+	UPROPERTY(EditAnywhere, Category = RevolutionAxis, meta = (DisplayName = "Orientation"))
+	FVector2D AxisOrientation;
 
-	UPROPERTY(EditAnywhere, Category = RevolutionAxis, meta = (
-		UIMin = -180, UIMax = 180, ClampMin = -180000, ClampMax = 18000))
-	float AxisPitch = 0;
+protected:
+	virtual ERevolvePropertiesCapFillMode GetCapFillMode() const override
+	{
+		return static_cast<ERevolvePropertiesCapFillMode>(CapFillMode);
+	}
 };
 
 /** 
diff --git a/Engine/Plugins/Runtime/MeshModelingToolset/Source/MeshModelingTools/Private/DrawAndRevolveTool.cpp b/Engine/Plugins/Runtime/MeshModelingToolset/Source/MeshModelingTools/Private/DrawAndRevolveTool.cpp
--- a/Engine/Plugins/Runtime/MeshModelingToolset/Source/MeshModelingTools/Private/DrawAndRevolveTool.cpp
+++ b/Engine/Plugins/Runtime/MeshModelingToolset/Source/MeshModelingTools/Private/DrawAndRevolveTool.cpp
@@ -2,7 +2,6 @@
 
 #include "DrawAndRevolveTool.h"
 
-#include "BaseGizmos/GizmoRenderingUtil.h"
 #include "BaseBehaviors/SingleClickBehavior.h"
 #include "BaseBehaviors/MouseHoverBehavior.h"
 #include "CompositionOps/CurveSweepOp.h"
@@ -21,8 +20,8 @@ using namespace UE::Geometry;
 
 #define LOCTEXT_NAMESPACE "UDrawAndRevolveTool"
 
-const FText InitializationModeMessage = LOCTEXT("CurveInitialization", "Draw a profile curve and revolve it around purple axis. Ctrl+click repositions draw plane, Ctrl+Shift+click repositions without aligning axis. End curve by clicking the end again or connecting to start. Backspace deletes points, shift inverts snapping behavior.");
-const FText EditModeMessage = LOCTEXT("CurveEditing", "Click points to select them, Shift+click to add/remove points to selection. Ctrl+click a segment to add a point, or select an endpoint and Ctrl+click somewhere on the plane to add to the ends. Backspace deletes selected points. Holding Shift toggles snapping to be opposite the EnableSnapping setting.");
+const FText InitializationModeMessage = LOCTEXT("CurveInitialization", "Draw a path and revolve it around the purple axis. Left-click to place path vertices, and click on the last or first vertex to complete the path. Hold Shift to invert snapping behavior.");
+const FText EditModeMessage = LOCTEXT("CurveEditing", "Click points to select them, Shift+click to modify selection. Ctrl+click a segment to add a point, or select an endpoint and Ctrl+click to add new endpoint. Backspace deletes selected points. Hold Shift to invert snapping behavior.");
 
 
 // Tool builder
@@ -53,7 +52,7 @@ TUniquePtr<FDynamicMeshOperator> URevolveOperatorFactory::MakeNewOperator()
 	CurveSweepOp->bProfileCurveIsClosed = RevolveTool->ControlPointsMechanic->GetIsLoop();
 
 	// If we are capping the top and bottom, we just add a couple extra vertices and mark the curve as being closed
-	if (!CurveSweepOp->bProfileCurveIsClosed && RevolveTool->Settings->bConnectOpenProfileToAxis)
+	if (!CurveSweepOp->bProfileCurveIsClosed && RevolveTool->Settings->bClosePathToAxis)
 	{
 		// Project first and last points onto the revolution axis.
 		FVector3d FirstPoint = CurveSweepOp->ProfileCurve[0];
@@ -241,7 +240,7 @@ void UDrawAndRevolveTool::GenerateAsset(const FDynamicMeshOpResult& OpResult)
 		return;
 	}
 
-	GetToolManager()->BeginUndoTransaction(LOCTEXT("RevolveToolTransactionName", "Revolve Tool"));
+	GetToolManager()->BeginUndoTransaction(LOCTEXT("RevolveToolTransactionName", "Path Revolve Tool"));
 
 	FCreateMeshObjectParams NewMeshObjectParams;
 	NewMeshObjectParams.TargetWorld = TargetWorld;
diff --git a/Engine/Plugins/Runtime/MeshModelingToolset/Source/MeshModelingTools/Public/DrawAndRevolveTool.h b/Engine/Plugins/Runtime/MeshModelingToolset/Source/MeshModelingTools/Public/DrawAndRevolveTool.h
--- a/Engine/Plugins/Runtime/MeshModelingToolset/Source/MeshModelingTools/Public/DrawAndRevolveTool.h
+++ b/Engine/Plugins/Runtime/MeshModelingToolset/Source/MeshModelingTools/Public/DrawAndRevolveTool.h
@@ -4,11 +4,9 @@
 
 #include "CoreMinimal.h"
 
-#include "BaseBehaviors/BehaviorTargetInterfaces.h"
 #include "ModelingOperators.h" //IDynamicMeshOperatorFactory
 #include "InteractiveTool.h" //UInteractiveToolPropertySet
 #include "InteractiveToolBuilder.h" //UInteractiveToolBuilder
-#include "InteractiveToolChange.h" //FToolCommandChange
 #include "MeshOpPreviewHelpers.h" //FDynamicMeshOpResult
 #include "Properties/MeshMaterialProperties.h"
 #include "PropertySets/CreateMeshObjectTypeProperties.h"
@@ -30,6 +28,17 @@ public:
 	virtual UInteractiveTool* BuildTool(const FToolBuilderState& SceneState) const override;
 };
 
+UENUM()
+enum class ERevolvePropertiesPathCapFillMode : uint8
+{
+	/** No caps will be generated. */
+	None = static_cast<uint8>(ERevolvePropertiesCapFillMode::None),
+	/** Caps are triangulated by placing a vertex in the center and creating a fan to the boundary. This works well if the path is convex,
+	  * but can create invalid geometry if it is concave. */
+	CenterFan = static_cast<uint8>(ERevolvePropertiesCapFillMode::CenterFan),
+	/** Caps are triangulated to maximize the minimal angle in the triangles using Delaunay triangulation. */
+	Delaunay = static_cast<uint8>(ERevolvePropertiesCapFillMode::Delaunay)
+};
 
 UCLASS()
 class MESHMODELINGTOOLS_API URevolveToolProperties : public URevolveProperties
@@ -38,28 +47,38 @@ class MESHMODELINGTOOLS_API URevolveToolProperties : public URevolveProperties
 
 public:
 
-	/** Connect the ends of an open profile to the axis to close the top and bottom of the revolved result. Not relevant if profile curve is closed. */
-	UPROPERTY(EditAnywhere, Category = RevolveSettings, AdvancedDisplay)
-	bool bConnectOpenProfileToAxis = true;
+	/** Determines how end caps are created. This is not relevant if the end caps are not visible or if the path is not closed. */
+	UPROPERTY(EditAnywhere, Category = Revolve, AdvancedDisplay, meta = (DisplayAfter = "QuadSplitMode",
+		EditCondition = "HeightOffsetPerDegree != 0 || RevolveDegrees != 360"))
+	ERevolvePropertiesPathCapFillMode CapFillMode = ERevolvePropertiesPathCapFillMode::Delaunay;
+
+	/** Connect the ends of an open path to the axis to add caps to the top and bottom of the revolved result.
+	  * This is not relevant for paths that are already closed. */
+	UPROPERTY(EditAnywhere, Category = Revolve, AdvancedDisplay)
+	bool bClosePathToAxis = true;
 	
 	/** Sets the draw plane origin. The revolution axis is the X axis in the plane. */
-	UPROPERTY(EditAnywhere, Category = DrawPlane, meta = (
-		EditCondition = "bAllowedToEditDrawPlane", HideEditConditionToggle))
+	UPROPERTY(EditAnywhere, Category = DrawPlane, meta = (DisplayName = "Origin", EditCondition = "bAllowedToEditDrawPlane", HideEditConditionToggle))
 	FVector DrawPlaneOrigin = FVector(0, 0, 0);
 
 	/** Sets the draw plane orientation. The revolution axis is the X axis in the plane. */
-	UPROPERTY(EditAnywhere, Category = DrawPlane, meta = (
-		EditCondition = "bAllowedToEditDrawPlane", HideEditConditionToggle, 
+	UPROPERTY(EditAnywhere, Category = DrawPlane, meta = (DisplayName = "Orientation", EditCondition = "bAllowedToEditDrawPlane", HideEditConditionToggle, 
 		UIMin = -180, UIMax = 180, ClampMin = -180000, ClampMax = 18000))
 	FRotator DrawPlaneOrientation = FRotator(90, 0, 0);
 
-	/** Enables/disables snapping while editing the profile curve. Also toggled with Shift. */
-	UPROPERTY(EditAnywhere, Category = ProfileCurve)
+	/** Enables snapping while editing the path. */
+	UPROPERTY(EditAnywhere, Category = Snapping)
 	bool bEnableSnapping = true;
 
 	// Not user visible- used to disallow draw plane modification.
 	UPROPERTY(meta = (TransientToolProperty))
 	bool bAllowedToEditDrawPlane = true;
+
+protected:
+	virtual ERevolvePropertiesCapFillMode GetCapFillMode() const override
+	{
+		return static_cast<ERevolvePropertiesCapFillMode>(CapFillMode);
+	}
 };
 
 UCLASS()
