ModelingTools: Use actor label for output texture names when available in BakeTexture & BakeAll.

#rnx

diff --git a/Engine/Plugins/Experimental/MeshModelingToolsetExp/Source/MeshModelingToolsExp/Private/BakeMeshAttributeMapsTool.cpp b/Engine/Plugins/Experimental/MeshModelingToolsetExp/Source/MeshModelingToolsExp/Private/BakeMeshAttributeMapsTool.cpp
--- a/Engine/Plugins/Experimental/MeshModelingToolsetExp/Source/MeshModelingToolsExp/Private/BakeMeshAttributeMapsTool.cpp
+++ b/Engine/Plugins/Experimental/MeshModelingToolsetExp/Source/MeshModelingToolsExp/Private/BakeMeshAttributeMapsTool.cpp
@@ -471,7 +471,7 @@ void UBakeMeshAttributeMapsTool::Shutdown(EToolShutdownType ShutdownType)
 			SourceAsset = SkeletalMeshTarget ? SkeletalMeshTarget->GetSkeletalMesh() : nullptr;
 		}
 		const UPrimitiveComponent* SourceComponent = UE::ToolTarget::GetTargetComponent(Targets[0]);
-		const FString BaseName = UE::ToolTarget::GetTargetActor(Targets[0])->GetName();
+		const FString BaseName = UE::ToolTarget::GetTargetActor(Targets[0])->GetActorNameOrLabel();
 
 		bool bCreatedAssetOK = true;
 		for (const TTuple<EBakeMapType, TObjectPtr<UTexture2D>>& Result : Settings->Result)
diff --git a/Engine/Plugins/Experimental/MeshModelingToolsetExp/Source/MeshModelingToolsExp/Private/BakeMultiMeshAttributeMapsTool.cpp b/Engine/Plugins/Experimental/MeshModelingToolsetExp/Source/MeshModelingToolsExp/Private/BakeMultiMeshAttributeMapsTool.cpp
--- a/Engine/Plugins/Experimental/MeshModelingToolsetExp/Source/MeshModelingToolsExp/Private/BakeMultiMeshAttributeMapsTool.cpp
+++ b/Engine/Plugins/Experimental/MeshModelingToolsetExp/Source/MeshModelingToolsExp/Private/BakeMultiMeshAttributeMapsTool.cpp
@@ -547,7 +547,7 @@ void UBakeMultiMeshAttributeMapsTool::Shutdown(EToolShutdownType ShutdownType)
 		UStaticMeshComponent* StaticMeshComponent = CastChecked<UStaticMeshComponent>(UE::ToolTarget::GetTargetComponent(Targets[0]));
 		UStaticMesh* StaticMeshAsset = StaticMeshComponent->GetStaticMesh();
 		check(StaticMeshAsset);
-		FString BaseName = UE::ToolTarget::GetTargetActor(Targets[0])->GetName();
+		const FString BaseName = UE::ToolTarget::GetTargetActor(Targets[0])->GetActorNameOrLabel();
 
 		bool bCreatedAssetOK = true;
 		for (const TTuple<EBakeMapType, TObjectPtr<UTexture2D>>& Result : Settings->Result)
