ModelingTools: Mark MultiTexture texture properties as transient.

#rnx

diff --git a/Engine/Plugins/Experimental/MeshModelingToolsetExp/Source/MeshModelingToolsExp/Public/BakeMeshAttributeToolCommon.h b/Engine/Plugins/Experimental/MeshModelingToolsetExp/Source/MeshModelingToolsExp/Public/BakeMeshAttributeToolCommon.h
--- a/Engine/Plugins/Experimental/MeshModelingToolsetExp/Source/MeshModelingToolsExp/Public/BakeMeshAttributeToolCommon.h
+++ b/Engine/Plugins/Experimental/MeshModelingToolsetExp/Source/MeshModelingToolsExp/Public/BakeMeshAttributeToolCommon.h
@@ -176,7 +176,7 @@ class MESHMODELINGTOOLSEXP_API UBakedMultiTexture2DImageProperties : public UInt
 public:
 
 	/** For each material ID, the source texture that will be resampled in that material's region*/
-	UPROPERTY(EditAnywhere, Category = MultiTexture, meta = (DisplayName = "Material IDs / Source Textures"))
+	UPROPERTY(EditAnywhere, Category = MultiTexture, meta = (DisplayName = "Material IDs / Source Textures", TransientToolProperty))
 	TMap<int32, TObjectPtr<UTexture2D>> MaterialIDSourceTextureMap;
 
 	/** UV layer to sample from on the input mesh */
@@ -184,7 +184,7 @@ public:
 	int32 UVLayer = 0;
 
 	/** The set of all source textures from all input materials */
-	UPROPERTY(VisibleAnywhere, BlueprintReadOnly, AdvancedDisplay, Category = MultiTexture, meta = (DisplayName = "Source Textures"))
+	UPROPERTY(VisibleAnywhere, BlueprintReadOnly, AdvancedDisplay, Category = MultiTexture, meta = (DisplayName = "Source Textures", TransientToolProperty))
 	TArray<TObjectPtr<UTexture2D>> AllSourceTextures;
 
 };
