ModelingToolset: Add ModelingEditorUI module for custom slate widgets and other Modeling UI customization/etc shared code
#rnx

diff --git a/Engine/Plugins/Experimental/MeshModelingToolsetExp/MeshModelingToolsetExp.uplugin b/Engine/Plugins/Experimental/MeshModelingToolsetExp/MeshModelingToolsetExp.uplugin
--- a/Engine/Plugins/Experimental/MeshModelingToolsetExp/MeshModelingToolsetExp.uplugin
+++ b/Engine/Plugins/Experimental/MeshModelingToolsetExp/MeshModelingToolsetExp.uplugin
@@ -53,6 +53,11 @@
 			"Type": "Editor",
 			"LoadingPhase": "Default",
 			"PlatformAllowList": [ "Win64" ]
+		},
+		{
+			"Name": "ModelingEditorUI",
+			"Type": "Editor",
+			"LoadingPhase": "Default"
 		}
 	],
 	"Plugins": [
diff --git a/Engine/Plugins/Experimental/MeshModelingToolsetExp/Source/ModelingEditorUI/ModelingEditorUI.build.cs b/Engine/Plugins/Experimental/MeshModelingToolsetExp/Source/ModelingEditorUI/ModelingEditorUI.build.cs
new file mode 100644
--- /dev/null
+++ b/Engine/Plugins/Experimental/MeshModelingToolsetExp/Source/ModelingEditorUI/ModelingEditorUI.build.cs
@@ -0,0 +1,52 @@
+// Copyright Epic Games, Inc. All Rights Reserved.
+
+using UnrealBuildTool;
+
+public class ModelingEditorUI : ModuleRules
+{
+	public ModelingEditorUI(ReadOnlyTargetRules Target) : base(Target)
+	{
+		PCHUsage = ModuleRules.PCHUsageMode.UseExplicitOrSharedPCHs;
+		
+		PublicDependencyModuleNames.AddRange(
+			new string[]
+			{
+				"Core",
+				// ... add other public dependencies that you statically link with here ...
+			}
+            );
+			
+		
+		PrivateDependencyModuleNames.AddRange(
+			new string[]
+			{
+                "Core",
+                "CoreUObject",
+                "ApplicationCore",
+                "Slate",
+                "SlateCore",
+                "Engine",
+                "InputCore",
+				"EditorFramework",
+				"UnrealEd",
+                "ContentBrowser",
+				"ContentBrowserData",
+                "LevelEditor",
+				"StatusBar",
+				"EditorStyle",
+                "Projects",
+				"ToolWidgets",
+				"EditorWidgets",
+				"DeveloperSettings"
+			}
+			);
+
+
+		DynamicallyLoadedModuleNames.AddRange(
+			new string[]
+			{
+				// ... add any modules that your module loads dynamically here ...
+			}
+			);
+	}
+}
diff --git a/Engine/Plugins/Experimental/MeshModelingToolsetExp/Source/ModelingEditorUI/Private/ModelingEditorUIModule.cpp b/Engine/Plugins/Experimental/MeshModelingToolsetExp/Source/ModelingEditorUI/Private/ModelingEditorUIModule.cpp
new file mode 100644
--- /dev/null
+++ b/Engine/Plugins/Experimental/MeshModelingToolsetExp/Source/ModelingEditorUI/Private/ModelingEditorUIModule.cpp
@@ -0,0 +1,20 @@
+// Copyright Epic Games, Inc. All Rights Reserved.
+
+#include "ModelingEditorUIModule.h"
+
+#define LOCTEXT_NAMESPACE "FModelingEditorUIModule"
+
+void FModelingEditorUIModule::StartupModule()
+{
+	// This code will execute after your module is loaded into memory; the exact timing is specified in the .uplugin file per-module
+}
+
+void FModelingEditorUIModule::ShutdownModule()
+{
+	// This function may be called during shutdown to clean up your module.  For modules that support dynamic reloading,
+	// we call this function before unloading the module.
+}
+
+#undef LOCTEXT_NAMESPACE
+	
+IMPLEMENT_MODULE(FModelingEditorUIModule, ModelingEditorUI)
\ No newline at end of file
diff --git a/Engine/Plugins/Experimental/MeshModelingToolsetExp/Source/ModelingEditorUI/Private/ModelingWidgets/ModelingCustomizationUtil.cpp b/Engine/Plugins/Experimental/MeshModelingToolsetExp/Source/ModelingEditorUI/Private/ModelingWidgets/ModelingCustomizationUtil.cpp
new file mode 100644
--- /dev/null
+++ b/Engine/Plugins/Experimental/MeshModelingToolsetExp/Source/ModelingEditorUI/Private/ModelingWidgets/ModelingCustomizationUtil.cpp
@@ -0,0 +1,146 @@
+// Copyright Epic Games, Inc. All Rights Reserved.
+
+#include "ModelingWidgets/ModelingCustomizationUtil.h"
+
+#include "Widgets/Layout/SBox.h"
+#include "Widgets/Input/SCheckBox.h"
+
+// necessary for export of static constexpr to work (?)
+constexpr int UE::ModelingUI::ModelingUIConstants::DetailRowVertPadding;
+constexpr int UE::ModelingUI::ModelingUIConstants::LabelWidgetMinPadding;
+constexpr int UE::ModelingUI::ModelingUIConstants::MultiWidgetRowHorzPadding;
+
+
+
+TSharedRef<SBox> UE::ModelingUI::WrapInFixedWidthBox(TSharedRef<SWidget> SubWidget, int32 Width)
+{
+	return SNew(SBox)
+		.WidthOverride(Width)
+		[
+			SubWidget
+		];
+}
+
+
+
+TSharedRef<SHorizontalBox> UE::ModelingUI::MakeFixedWidthLabelSliderHBox(
+	TSharedPtr<IPropertyHandle> LabelHandle,
+	TSharedPtr<SDynamicNumericEntry::FDataSource> SliderDataSource,
+	int32 LabelFixedWidth)
+{
+	return SNew(SHorizontalBox)
+		+ SHorizontalBox::Slot()
+		.AutoWidth()
+		[
+			UE::ModelingUI::WrapInFixedWidthBox( 
+				LabelHandle->CreatePropertyNameWidget(), LabelFixedWidth )
+		]
+		+ SHorizontalBox::Slot()
+		.Padding(ModelingUIConstants::LabelWidgetMinPadding, 0,0,0)
+		.FillWidth(3.0f)		// seems arbitrary....
+		[
+			SNew(SDynamicNumericEntry)
+			.Source(SliderDataSource)
+		];
+}
+
+
+
+
+TSharedRef<SHorizontalBox> UE::ModelingUI::MakeToggleSliderHBox(
+	TSharedPtr<IPropertyHandle> BoolToggleHandle,
+	FText ToggleLabelText,
+	TSharedPtr<SDynamicNumericEntry::FDataSource> SliderDataSource,
+	int32 LabelFixedWidth)
+{
+	TSharedRef<SDynamicNumericEntry> SliderWidget = SNew(SDynamicNumericEntry)
+		.Source(SliderDataSource);
+
+	TFunction<void(bool)> ToggledFunc = [SliderWidget](bool bToggled)
+	{
+		SliderWidget->SetEnabled(bToggled);
+	};
+
+	TSharedRef<SCheckBox> ToggleWidget = MakeBoolToggleButton(BoolToggleHandle, ToggleLabelText, ToggledFunc);
+
+	bool bInitialValue = true;
+	BoolToggleHandle->GetValue(bInitialValue);
+	SliderWidget->SetEnabled(bInitialValue);
+
+	return SNew(SHorizontalBox)
+		+ SHorizontalBox::Slot()
+		.AutoWidth()
+		[
+			UE::ModelingUI::WrapInFixedWidthBox( ToggleWidget, LabelFixedWidth )
+		]
+		+ SHorizontalBox::Slot()
+		.Padding(ModelingUIConstants::LabelWidgetMinPadding, 0,0,0)
+		.FillWidth(3.0f)		// seems arbitrary....
+		[
+			SliderWidget
+		];
+}
+
+
+
+TSharedRef<SHorizontalBox> UE::ModelingUI::MakeTwoWidgetDetailRowHBox(
+	TSharedRef<SWidget> Widget1,
+	TSharedRef<SWidget> Widget2,
+	float FillWidth1, float FillWidth2)
+{
+	return SNew(SHorizontalBox)
+		+ SHorizontalBox::Slot()
+		.Padding(FMargin(0, ModelingUIConstants::DetailRowVertPadding))
+		.FillWidth(FillWidth1)
+		[
+			Widget1
+		]
+
+		+ SHorizontalBox::Slot()
+		.Padding(FMargin(ModelingUIConstants::MultiWidgetRowHorzPadding, ModelingUIConstants::DetailRowVertPadding, 0, ModelingUIConstants::DetailRowVertPadding))
+		.FillWidth(FillWidth2)
+		[
+			Widget2
+		];
+}
+
+
+
+TSharedRef<SCheckBox> UE::ModelingUI::MakeBoolToggleButton(	
+	TSharedPtr<IPropertyHandle> BoolToggleHandle,
+	FText ButtonLabelText,
+	TFunction<void(bool)> OnToggledCallback,
+	int HorzPadding)
+{
+	return SNew(SCheckBox)
+		.Style(FAppStyle::Get(), "DetailsView.SectionButton")
+		.Padding(FMargin(HorzPadding, 2))
+		.ToolTipText(BoolToggleHandle->GetToolTipText())
+		.HAlign(HAlign_Center)
+		.OnCheckStateChanged_Lambda([BoolToggleHandle, LocalToggledFunc=OnToggledCallback](ECheckBoxState NewState)
+		{
+			bool bSet = (NewState == ECheckBoxState::Checked) ? true : false;
+			BoolToggleHandle->SetValue(bSet);
+			LocalToggledFunc(bSet);
+		})
+		.IsChecked_Lambda([BoolToggleHandle]() -> ECheckBoxState
+		{
+			bool bSet;
+			BoolToggleHandle->GetValue(bSet);
+			return bSet ? ECheckBoxState::Checked : ECheckBoxState::Unchecked;
+		})
+		[
+			SNew(SHorizontalBox)
+			+ SHorizontalBox::Slot()
+			.VAlign(VAlign_Center)
+			.Padding(FMargin(0))
+			.AutoWidth()
+			[
+				SNew(STextBlock)
+				.Justification(ETextJustify::Center)
+				.TextStyle(FAppStyle::Get(), "DetailsView.CategoryTextStyle")
+				.Text(ButtonLabelText)
+			]
+		];
+}
+
diff --git a/Engine/Plugins/Experimental/MeshModelingToolsetExp/Source/ModelingEditorUI/Private/ModelingWidgets/SComboPanel.cpp b/Engine/Plugins/Experimental/MeshModelingToolsetExp/Source/ModelingEditorUI/Private/ModelingWidgets/SComboPanel.cpp
new file mode 100644
--- /dev/null
+++ b/Engine/Plugins/Experimental/MeshModelingToolsetExp/Source/ModelingEditorUI/Private/ModelingWidgets/SComboPanel.cpp
@@ -0,0 +1,344 @@
+// Copyright Epic Games, Inc. All Rights Reserved.
+
+#include "ModelingWidgets/SComboPanel.h"
+
+#include "Widgets/SOverlay.h"
+#include "Widgets/Input/SButton.h"
+#include "Widgets/Input/SCheckBox.h"
+#include "Widgets/Text/STextBlock.h"
+#include "Widgets/Input/SComboButton.h"
+#include "Widgets/Views/STileView.h"
+#include "Widgets/Images/SImage.h"
+#include "Widgets/Layout/SSpacer.h"
+
+#include "Internationalization/BreakIterator.h"
+#include "Framework/MultiBox/MultiBoxBuilder.h"
+#include "Styling/SlateStyleMacros.h"
+
+#include "ModelingWidgets/ModelingCustomizationUtil.h"
+
+#define LOCTEXT_NAMESPACE "SComboPanel"
+
+
+/**
+ * Internal class used for icons in the ComboPanel STileView and in the ComboButton
+ */
+class SComboPanelIconTile : public SCompoundWidget
+{
+public:
+	SLATE_BEGIN_ARGS(SComboPanelIconTile)
+		: _ThumbnailSize(102, 102)
+		, _ExteriorPadding(2.0f)
+		, _ImagePadding(FMargin(0))
+		, _ThumbnailPadding(FMargin(5.0f, 0))
+		, _bShowLabel(true)
+		, _IsSelected(false)
+		{
+		}
+
+		SLATE_ATTRIBUTE(const FSlateBrush*, Image)
+		SLATE_ATTRIBUTE(FText, DisplayName)
+		/** The size of the thumbnail */
+		SLATE_ARGUMENT(FVector2D, ThumbnailSize)
+		/** The size of the thumbnail */
+		SLATE_ARGUMENT(float, ExteriorPadding)
+		/** The size of the thumbnail */
+		SLATE_ARGUMENT(FMargin, ImagePadding)
+		/** The size of the image in the thumbnail. Used to make the thumbnail not fill the entire thumbnail area. If this is not set the Image brushes desired size is used */
+		SLATE_ARGUMENT(TOptional<FVector2D>, ImageSize)
+		SLATE_ARGUMENT(FMargin, ThumbnailPadding)
+		SLATE_ARGUMENT(bool, bShowLabel)
+		SLATE_ATTRIBUTE(bool, IsSelected)
+	SLATE_END_ARGS()
+
+public:
+	/**
+	 * Construct this widget
+	 */
+	void Construct( const FArguments& InArgs );
+
+private:
+	TAttribute<const FSlateBrush*> Image;
+	TAttribute<FText> Text;
+	TAttribute<bool> IsSelected;
+};
+
+
+
+void SComboPanelIconTile::Construct(const FArguments& InArgs)
+{
+	Image = InArgs._Image;
+	Text = InArgs._DisplayName;
+	IsSelected = InArgs._IsSelected;
+
+
+	TSharedPtr<SWidget> LabelContent;
+	if (InArgs._bShowLabel)
+	{
+		LabelContent = SNew(SVerticalBox)
+			+ SVerticalBox::Slot()
+			[
+				SNew(SBorder)
+				.Padding(InArgs._ThumbnailPadding)
+				.VAlign(VAlign_Bottom)
+				.HAlign(HAlign_Right)
+				.Padding(FMargin(3.0f, 3.0f))
+				//.BorderImage(FAppStyle::Get().GetBrush("ProjectBrowser.ProjectTile.NameAreaBackground"))
+				[
+					SNew(STextBlock)
+					//.Font(FAppStyle::Get().GetFontStyle("ProjectBrowser.ProjectTile.Font"))
+					.Font( DEFAULT_FONT("Regular", 8) )
+					.AutoWrapText(false)
+					.LineBreakPolicy(FBreakIterator::CreateWordBreakIterator())
+					.Text(InArgs._DisplayName)
+					.ColorAndOpacity( FAppStyle::Get().GetSlateColor("Colors.Foreground") )
+				]
+			];
+	}
+	else
+	{
+		LabelContent = SNew(SSpacer);
+	}
+
+	ChildSlot
+	[
+		SNew(SBorder)
+		.Padding(FMargin(InArgs._ExteriorPadding))
+		.BorderImage(FAppStyle::Get().GetBrush("ProjectBrowser.ProjectTile.DropShadow"))
+		[
+			SNew(SOverlay)
+			+ SOverlay::Slot()
+			[
+				SNew(SVerticalBox)
+				// Thumbnail
+				+ SVerticalBox::Slot()
+				.AutoHeight()
+				.HAlign(HAlign_Center)
+				.VAlign(VAlign_Center)
+				[
+					SNew(SBox)
+					.WidthOverride(InArgs._ThumbnailSize.X)
+					.HeightOverride(InArgs._ThumbnailSize.Y)
+					[
+						SNew(SBorder)
+						.Padding(InArgs._ImagePadding)
+						.BorderImage(FAppStyle::Get().GetBrush("ProjectBrowser.ProjectTile.ThumbnailAreaBackground"))
+						.HAlign(HAlign_Center)
+						.VAlign(VAlign_Top)
+						[
+							SNew(SImage)
+							.Image(Image)
+							.DesiredSizeOverride(InArgs._ImageSize)
+						]
+					]
+				]
+			]
+			+ SOverlay::Slot()
+			[
+				LabelContent->AsShared()
+			]
+			+ SOverlay::Slot()
+			[
+				SNew(SImage)
+				.Visibility(EVisibility::HitTestInvisible)
+				.Image_Lambda
+				(
+					[this]()
+					{
+						const bool bSelected = IsSelected.Get();
+						const bool bHovered = IsHovered();
+
+						if (bSelected && bHovered)
+						{
+							static const FName SelectedHover("ProjectBrowser.ProjectTile.SelectedHoverBorder");
+							return FAppStyle::Get().GetBrush(SelectedHover);
+						}
+						else if (bSelected)
+						{
+							static const FName Selected("ProjectBrowser.ProjectTile.SelectedBorder");
+							return FAppStyle::Get().GetBrush(Selected);
+						}
+						else if (bHovered)
+						{
+							static const FName Hovered("ProjectBrowser.ProjectTile.HoverBorder");
+							return FAppStyle::Get().GetBrush(Hovered);
+						}
+
+						return FStyleDefaults::GetNoBrush();
+					}
+				)
+			]
+		]
+	];
+}
+
+
+
+
+
+void SComboPanel::Construct(const FArguments& InArgs)
+{
+	this->ComboButtonTileSize = InArgs._ComboButtonTileSize;
+	this->FlyoutTileSize = InArgs._FlyoutTileSize;
+	this->FlyoutSize = InArgs._FlyoutSize;
+	this->Items = InArgs._ListItems;
+	this->OnSelectionChanged = InArgs._OnSelectionChanged;
+
+	MissingIcon = MakeShared<FSlateBrush>();
+
+	// Check for any parameters that the coder forgot to specify.
+	if ( Items.Num() == 0 )
+	{
+		UE::ModelingUI::SetCustomWidgetErrorString(LOCTEXT("MissingListItemsError", "Please specify a ListItems with at least one Item."), this->ChildSlot);
+		return;
+	}
+
+	TileView = SNew(STileView<TSharedPtr<FComboPanelItem>>)
+		.ListItemsSource( &Items )
+		.OnGenerateTile(this, &SComboPanel::CreateFlyoutIconTile)
+		.OnSelectionChanged(this, &SComboPanel::FlyoutSelectionChanged)
+		.ClearSelectionOnClick(false)
+		.ItemAlignment(EListItemAlignment::LeftAligned)
+		.ItemWidth(InArgs._FlyoutTileSize.X)
+		.ItemHeight(InArgs._FlyoutTileSize.Y)
+		.SelectionMode(ESelectionMode::Single);
+
+	TSharedPtr<SWidget> ComboButtonContent;
+	if (InArgs._ComboDisplayType == EComboDisplayType::IconAndLabel)
+	{
+		ComboButtonContent = SNew(SHorizontalBox)
+			+ SHorizontalBox::Slot()
+			.AutoWidth()
+			[
+				SNew(SBox)
+				.MinAspectRatio(1.0f)
+				.MaxAspectRatio(1.0f)
+				[
+					SNew(SComboPanelIconTile)
+					.Image_Lambda([this]() { return (SelectedItem->Icon != nullptr) ? SelectedItem->Icon : MissingIcon.Get(); })
+					.DisplayName_Lambda([this]() { return SelectedItem->Name; })
+					.bShowLabel(false)
+					.ThumbnailSize(ComboButtonTileSize)
+				]
+			]
+			+ SHorizontalBox::Slot()
+			.VAlign(VAlign_Center)
+			.Padding(4,0,0,0)
+			.FillWidth(1.0f)
+			[
+				SNew(STextBlock)
+				.TextStyle(FAppStyle::Get(), "DetailsView.CategoryTextStyle")
+				.AutoWrapText(false)
+				.Text_Lambda([this]() { return SelectedItem->Name; })
+			];
+	}
+	else
+	{
+		ComboButtonContent = SNew(SComboPanelIconTile)
+			.Image_Lambda([this]() { return (SelectedItem->Icon != nullptr) ? SelectedItem->Icon : MissingIcon.Get(); })
+			.DisplayName_Lambda([this]() { return SelectedItem->Name; })
+			.ExteriorPadding(0.0f)
+			.ImagePadding(FMargin(0,2,0,0))
+			.ImageSize( FVector2D(ComboButtonTileSize.X-15, ComboButtonTileSize.Y-15) )
+			.ThumbnailSize(ComboButtonTileSize);
+	}
+
+	ComboButton = SNew(SComboButton)
+		//.ToolTipText(BrushTypeHandle->GetToolTipText())		// how to do this? SWidget has TooltipText...
+		.HasDownArrow(false)
+		.OnGetMenuContent( this, &SComboPanel::OnGetMenuContent )
+		.OnMenuOpenChanged( this, &SComboPanel::OnMenuOpenChanged )
+		//.IsEnabled(IsEnabledAttribute)
+		.ContentPadding(FMargin(0))
+		.ButtonContent()
+		[
+			ComboButtonContent->AsShared()
+		];
+
+	int InitialSelectionIndex = FMath::Clamp(InArgs._InitialSelectionIndex, 0, Items.Num());
+	SelectedItem = Items[InitialSelectionIndex];
+	TileView->SetSelection(SelectedItem);
+
+	ChildSlot
+	[
+		ComboButton->AsShared()
+	];
+}
+
+
+
+TSharedRef<SWidget> SComboPanel::OnGetMenuContent()
+{
+	bool bInShouldCloseWindowAfterMenuSelection = true;
+	bool bCloseSelfOnly = true;
+	FMenuBuilder MenuBuilder(bInShouldCloseWindowAfterMenuSelection, nullptr, nullptr, bCloseSelfOnly);
+	MenuBuilder.BeginSection(NAME_None, LOCTEXT("BrushesHeader", "Brush Types"));
+	{
+		TSharedPtr<SWidget> MenuContent;
+		MenuContent =
+			SNew(SVerticalBox)
+			+ SVerticalBox::Slot()
+			.FillHeight(1.f)
+			[
+				SAssignNew(TileViewContainer, SBox)
+				.Padding(6.0f)
+				.MinDesiredWidth( FlyoutSize.X )
+			];
+
+		TileViewContainer->SetContent(TileView->AsShared());
+
+		MenuBuilder.AddWidget(MenuContent.ToSharedRef(), FText::GetEmpty(), true);
+	}
+	MenuBuilder.EndSection();
+
+	return MenuBuilder.MakeWidget();
+}
+
+void SComboPanel::OnMenuOpenChanged(bool bOpen)
+{
+	if (bOpen == false)
+	{
+		ComboButton->SetMenuContent(SNullWidget::NullWidget);
+	}
+}
+
+
+TSharedRef<ITableRow> SComboPanel::CreateFlyoutIconTile(
+	TSharedPtr<FComboPanelItem> Item,
+	const TSharedRef<STableViewBase>& OwnerTable)
+{
+	return SNew(STableRow< TSharedPtr<FComboPanelItem> >, OwnerTable)
+		.Style(FAppStyle::Get(), "ProjectBrowser.TableRow")
+		.ToolTipText( Item->Name )
+		.Padding(2.0f)
+		[
+			SNew(SComboPanelIconTile)
+			.Image( (Item->Icon != nullptr) ? Item->Icon : MissingIcon.Get() )
+			.DisplayName( Item->Name )
+			.ThumbnailSize( FlyoutTileSize )
+			.ImagePadding(FMargin(0,4,0,0))
+			.ImageSize( FVector2D(FlyoutTileSize.X-25, FlyoutTileSize.Y-25) )
+			.IsSelected_Lambda( [this, Item] { return Item.Get() == SelectedItem.Get(); })
+		];
+}
+
+
+void SComboPanel::FlyoutSelectionChanged(
+	TSharedPtr<FComboPanelItem> SelectedItemIn,
+	ESelectInfo::Type SelectInfo)
+{
+	SelectedItem = SelectedItemIn;
+
+	if (ComboButton.IsValid())
+	{
+		ComboButton->SetIsOpen(false);
+	}
+
+	if ( OnSelectionChanged.IsBound() )
+	{
+		OnSelectionChanged.ExecuteIfBound(SelectedItem);
+	}	
+}
+
+
+#undef LOCTEXT_NAMESPACE
\ No newline at end of file
diff --git a/Engine/Plugins/Experimental/MeshModelingToolsetExp/Source/ModelingEditorUI/Private/ModelingWidgets/SDynamicNumericEntry.cpp b/Engine/Plugins/Experimental/MeshModelingToolsetExp/Source/ModelingEditorUI/Private/ModelingWidgets/SDynamicNumericEntry.cpp
new file mode 100644
--- /dev/null
+++ b/Engine/Plugins/Experimental/MeshModelingToolsetExp/Source/ModelingEditorUI/Private/ModelingWidgets/SDynamicNumericEntry.cpp
@@ -0,0 +1,162 @@
+// Copyright Epic Games, Inc. All Rights Reserved.
+
+#include "ModelingWidgets/SDynamicNumericEntry.h"
+
+#include "Editor.h"
+#include "EditorStyleSet.h"
+#include "PropertyHandle.h"
+#include "IPropertyTypeCustomization.h"
+#include "IPropertyUtilities.h"
+#include "DetailWidgetRow.h"
+#include "EditorFontGlyphs.h"
+#include "Widgets/Input/SCheckBox.h"
+#include "Widgets/Input/SNumericEntryBox.h"
+#include "Math/Interval.h"
+
+#define LOCTEXT_NAMESPACE "SDynamicNumericEntry"
+
+
+void SDynamicNumericEntry::Construct(const FArguments& InArgs)
+{
+	this->Source = InArgs._Source;
+
+	this->TypeInterface = MakeShared<TDefaultNumericTypeInterface<float>>();
+
+	SAssignNew(NumericEntry, SNumericEntryBox<float>)
+	.AllowSpin(true)
+	.Value_Lambda([this]() { return this->Source->GetValue(); })
+	//.Font(InArgs._Font)
+	.Font( FEditorStyle::GetFontStyle( TEXT("PropertyWindow.NormalFont") ) )		// standard details panel font
+	.MinValue_Lambda([this]() { return this->Source->GetValueRange().Min; })
+	.MaxValue_Lambda([this]() { return this->Source->GetValueRange().Max; })
+	.MinSliderValue_Lambda([this]() { return this->Source->GetUIRange().Min; })
+	.MaxSliderValue_Lambda([this]() { return this->Source->GetUIRange().Max; })
+	//.SliderExponent(SliderExponent)
+	//.Delta(Delta)
+	.UndeterminedString(LOCTEXT("MultipleValues", "Multiple Values"))
+	.OnValueChanged(this, &SDynamicNumericEntry::OnValueChanged)
+	.OnValueCommitted(this, &SDynamicNumericEntry::OnValueCommitted)
+	//.OnUndeterminedValueCommitted(this, &SPropertyEditorNumeric<NumericType>::OnUndeterminedValueCommitted)
+	.OnBeginSliderMovement(this, &SDynamicNumericEntry::OnBeginSliderMovement)
+	.OnEndSliderMovement(this, &SDynamicNumericEntry::OnEndSliderMovement)
+	.TypeInterface(this->TypeInterface)
+	;
+
+	this->TypeInterface->SetMaxFractionalDigits(InArgs._MaxNumFloatDigits);
+
+	ChildSlot
+	[
+		NumericEntry->AsShared()
+	];
+}
+
+
+
+TSharedPtr<SDynamicNumericEntry::FDataSource> SDynamicNumericEntry::MakeSimpleDataSource(
+	TSharedPtr<IPropertyHandle> PropertyHandle, 
+	TInterval<float> ValueRange, 
+	TInterval<float> UIRange)
+{
+	TSharedPtr<SDynamicNumericEntry::FDataSource> NumericSource = MakeShared<SDynamicNumericEntry::FDataSource>();
+	NumericSource->SetValue = [PropertyHandle](float NewSize, EPropertyValueSetFlags::Type Flags)
+	{
+		PropertyHandle->SetValue(NewSize, Flags);
+	};
+	NumericSource->GetValue = [PropertyHandle]() -> float
+	{
+		float Size;
+		PropertyHandle->GetValue(Size);
+		return Size;
+	};
+	NumericSource->GetValueRange = [ValueRange]() -> TInterval<float>
+	{
+		return ValueRange;
+	};
+	NumericSource->GetUIRange = [UIRange]() -> TInterval<float>
+	{
+		return UIRange;
+	};
+
+	return NumericSource;
+}
+
+
+
+void SDynamicNumericEntry::OnValueChanged( float NewValue )
+{
+	if( bIsUsingSlider )
+	{
+		// Value hasn't changed, so lets return now
+		if (Source->GetValue() == NewValue)
+		{ 
+			return;
+		}
+
+		EPropertyValueSetFlags::Type Flags = (EPropertyValueSetFlags::InteractiveChange | EPropertyValueSetFlags::NotTransactable);
+		Source->SetValue(NewValue, Flags);
+
+		//if (TypeInterface.IsValid() && !TypeInterface->FixedDisplayUnits.IsSet())
+		//{
+		//	TypeInterface->SetupFixedDisplay(NewValue);
+		//}
+
+	}
+}
+
+void SDynamicNumericEntry::OnValueCommitted( float NewValue, ETextCommit::Type CommitInfo )
+{
+	if (bIsUsingSlider || Source->GetValue() != NewValue)
+	{
+		EPropertyValueSetFlags::Type Flags = EPropertyValueSetFlags::DefaultFlags;
+		Source->SetValue(NewValue, Flags);
+		LastSliderCommittedValue = NewValue;
+	}
+
+	//if (TypeInterface.IsValid() && !TypeInterface->FixedDisplayUnits.IsSet())
+	//{
+	//	TypeInterface->SetupFixedDisplay(NewValue);
+	//}
+
+}
+
+
+/**
+* Called when the slider begins to move.  We create a transaction here to undo the property
+*/
+void SDynamicNumericEntry::OnBeginSliderMovement()
+{
+	bIsUsingSlider = true;
+
+	LastSliderCommittedValue = Source->GetValue();
+
+	//GEditor->BeginTransaction(TEXT("PropertyEditor"), FText::Format(NSLOCTEXT("PropertyEditor", "SetNumericPropertyTransaction", "Edit {0}"), PropertyEditor->GetDisplayName()), nullptr );
+	GEditor->BeginTransaction(LOCTEXT("EditNumericEntry", "Edit Number"));
+}
+
+
+/**
+* Called when the slider stops moving.  We end the previously created transaction
+*/
+void SDynamicNumericEntry::OnEndSliderMovement( float NewValue )
+{
+	bIsUsingSlider = false;
+
+	// When the slider end, we may have not called SetValue(NewValue) without the InteractiveChange|NotTransactable flags.
+	//That prevents some transaction and callback to be triggered like the NotifyHook.
+	if (LastSliderCommittedValue != NewValue)
+	{
+		EPropertyValueSetFlags::Type Flags = EPropertyValueSetFlags::DefaultFlags;
+		Source->SetValue(NewValue, Flags);
+	}
+	else
+	{
+		GEditor->EndTransaction();
+	}
+
+}
+
+
+
+
+
+#undef LOCTEXT_NAMESPACE
\ No newline at end of file
diff --git a/Engine/Plugins/Experimental/MeshModelingToolsetExp/Source/ModelingEditorUI/Private/ModelingWidgets/SToolInputAssetComboPanel.cpp b/Engine/Plugins/Experimental/MeshModelingToolsetExp/Source/ModelingEditorUI/Private/ModelingWidgets/SToolInputAssetComboPanel.cpp
new file mode 100644
--- /dev/null
+++ b/Engine/Plugins/Experimental/MeshModelingToolsetExp/Source/ModelingEditorUI/Private/ModelingWidgets/SToolInputAssetComboPanel.cpp
@@ -0,0 +1,380 @@
+// Copyright Epic Games, Inc. All Rights Reserved.
+
+#include "ModelingWidgets/SToolInputAssetComboPanel.h"
+
+#include "Widgets/SOverlay.h"
+#include "Widgets/Input/SButton.h"
+#include "Widgets/Input/SCheckBox.h"
+#include "Widgets/Text/STextBlock.h"
+#include "Widgets/Input/SComboButton.h"
+#include "Widgets/Images/SImage.h"
+#include "Widgets/Views/SListView.h"
+#include "SAssetView.h"
+#include "SAssetDropTarget.h"
+
+#include "Editor.h"
+
+#include "ModelingWidgets/SToolInputAssetPicker.h"
+#include "ModelingWidgets/ModelingCustomizationUtil.h"
+
+#include "Styling/SlateStyleMacros.h"
+
+#define LOCTEXT_NAMESPACE "SToolInputAssetComboPanel"
+
+
+void SToolInputAssetComboPanel::Construct(const FArguments& InArgs)
+{
+	this->ComboButtonTileSize = InArgs._ComboButtonTileSize;
+	this->FlyoutTileSize = InArgs._FlyoutTileSize;
+	this->FlyoutSize = InArgs._FlyoutSize;
+	this->Property = InArgs._Property;
+	this->CollectionSets = InArgs._CollectionSets;
+	this->OnSelectionChanged = InArgs._OnSelectionChanged;
+
+	// validate arguments
+	if ( InArgs._AssetClassType == nullptr )
+	{
+		UE::ModelingUI::SetCustomWidgetErrorString(LOCTEXT("MissingAssetType", "Please specify an AssetClassType."), this->ChildSlot);
+		return;
+	}
+	this->AssetClassType = InArgs._AssetClassType;
+
+	FText UseTooltipText = Property.IsValid() ? Property->GetToolTipText() : InArgs._ToolTipText;
+
+	// create our own thumbnail pool?
+	ThumbnailPool = MakeShared<FAssetThumbnailPool>(20, 1.0 /** must be large or thumbnails will not render?? */ );
+	ensure(ThumbnailPool.IsValid());
+
+	// make an empty asset thumbnail
+	FAssetThumbnailConfig ThumbnailConfig;
+	ThumbnailConfig.bAllowRealTimeOnHovered = false;
+	this->AssetThumbnail = MakeShared<FAssetThumbnail>(FAssetData(), ComboButtonTileSize.X, ComboButtonTileSize.Y, ThumbnailPool);
+
+	TSharedRef<SBox> IconBox = SNew(SBox)
+		.WidthOverride(ComboButtonTileSize.X)
+		.HeightOverride(ComboButtonTileSize.Y)
+		[
+			AssetThumbnail->MakeThumbnailWidget(ThumbnailConfig)
+		];
+
+	RecentAssetsProvider = InArgs._RecentAssetsProvider;
+
+
+	ComboButton = SNew(SComboButton)
+		.ToolTipText(UseTooltipText)
+		.HasDownArrow(false)
+		.ButtonContent()
+		[
+			SNew(SAssetDropTarget)
+			.OnAreAssetsAcceptableForDropWithReason_Lambda( [this](TArrayView<FAssetData> InAssets, FText& OutReason) {
+				UObject* AssetObject = InAssets[0].GetAsset();
+				if ( (AssetObject != nullptr) && AssetObject->IsA(this->AssetClassType))
+				{
+					return true;
+				}
+				return false;
+			})
+			.OnAssetsDropped_Lambda( [this](const FDragDropEvent&, TArrayView<FAssetData> InAssets)
+			{
+				this->NewAssetSelected(InAssets[0].GetAsset());
+			})
+			[
+				SNew(SBorder)
+					.Visibility(EVisibility::SelfHitTestInvisible)
+					.Padding(FMargin(0.0f))
+					.BorderImage(FAppStyle::Get().GetBrush("ProjectBrowser.ProjectTile.DropShadow"))
+					[
+						SNew(SOverlay)
+						+SOverlay::Slot()
+						//.Padding(1)
+						[
+							SNew(SVerticalBox)
+							// Thumbnail
+							+ SVerticalBox::Slot()
+							.AutoHeight()
+							.HAlign(HAlign_Center)
+							.VAlign(VAlign_Center)
+							[
+								SNew(SBox)
+								.WidthOverride(ComboButtonTileSize.X)
+								.HeightOverride(ComboButtonTileSize.Y)
+								[
+									SAssignNew(ThumbnailBorder, SBorder)
+									.Padding(0)
+									.BorderImage(FStyleDefaults::GetNoBrush())
+									.HAlign(HAlign_Center)
+									.VAlign(VAlign_Center)
+									.OnMouseDoubleClick(this, &SToolInputAssetComboPanel::OnAssetThumbnailDoubleClick)
+									[
+										IconBox
+									]
+								]
+							]
+						]
+
+					]
+
+			]
+		];
+
+
+	ComboButton->SetOnGetMenuContent(FOnGetContent::CreateLambda([this]() 
+	{
+		// Configure filter for asset picker
+		FAssetPickerConfig Config;
+		Config.SelectionMode = ESelectionMode::Single;
+		Config.Filter.bRecursiveClasses = true;
+		Config.Filter.ClassNames.Add( this->AssetClassType->GetFName() );
+		Config.Filter.bRecursivePaths = true;
+		//Config.Filter.bIncludeOnlyOnDiskAssets = true;
+		Config.InitialAssetViewType = EAssetViewType::Tile;
+		Config.bFocusSearchBoxWhenOpened = true;
+		Config.bAllowNullSelection = true;
+		Config.bAllowDragging = false;
+		Config.ThumbnailLabel = EThumbnailLabel::Type::NoLabel;
+		Config.OnAssetSelected = FOnAssetSelected::CreateSP(this, &SToolInputAssetComboPanel::NewAssetSelected);
+
+		// build asset picker UI
+		TSharedRef<SToolInputAssetPicker> AssetPickerWidget = SNew( SToolInputAssetPicker )
+			.IsEnabled( true )
+			.AssetPickerConfig( Config );
+
+		TSharedRef<SVerticalBox> PopupContent = SNew(SVerticalBox);
+
+		int32 FilterButtonBarVertPadding = 2;
+		UpdateRecentAssets();
+		if ( RecentAssetsProvider.IsValid() && RecentAssetData.Num() > 0 )
+		{
+			TSharedRef<SListView<TSharedPtr<FRecentAssetInfo>>> RecentsListView = SNew(SListView<TSharedPtr<FRecentAssetInfo>>)
+				.ItemHeight(FlyoutTileSize.X)
+				.Orientation(Orient_Horizontal)
+				.ListItemsSource(&RecentAssetData)
+				.OnGenerateRow(this, &SToolInputAssetComboPanel::OnGenerateWidgetForRecentList)
+				.OnSelectionChanged_Lambda([this](TSharedPtr<FRecentAssetInfo> SelectedItem, ESelectInfo::Type SelectInfo)
+				{
+					NewAssetSelected(SelectedItem->AssetData);
+				})
+				.ClearSelectionOnClick(false)
+				.SelectionMode(ESelectionMode::Single);
+
+			PopupContent->AddSlot()
+			.AutoHeight()
+			[
+				SNew(SBox)
+				.Padding(6.0f)
+				.HeightOverride(FlyoutTileSize.Y + 30.0f)
+				[
+					SNew(SVerticalBox)
+					+ SVerticalBox::Slot()
+					.AutoHeight()
+					[
+						SNew(STextBlock)
+						.Text(LOCTEXT("RecentsHeaderText", "Recently Used"))
+					]
+					+ SVerticalBox::Slot()
+					.AutoHeight()
+					[
+						RecentsListView->AsShared()
+					]
+				]
+			];
+		}
+		else
+		{
+			FilterButtonBarVertPadding = 10;
+		}
+
+		if (CollectionSets.Num() > 0)
+		{
+			PopupContent->AddSlot()
+			.AutoHeight()
+			.Padding(10, FilterButtonBarVertPadding, 4, 2)
+			[
+				MakeCollectionSetsButtonPanel(AssetPickerWidget)
+			];
+		}
+
+		PopupContent->AddSlot()
+		[
+			SNew(SBox)
+			.Padding(6.0f)
+			.HeightOverride(FlyoutSize.Y)
+			.WidthOverride(FlyoutSize.X)
+			[
+				SNew( SBorder )
+				.BorderImage( FEditorStyle::GetBrush("Menu.Background") )
+				[
+					AssetPickerWidget
+				]
+			]
+		];
+
+		return PopupContent;
+	}));
+
+
+	// set initial thumbnail
+	if (Property.IsValid())
+	{
+		FAssetData AssetData;
+		if (Property->GetValue(AssetData) == FPropertyAccess::Result::Success)
+		{
+			AssetThumbnail->SetAsset(AssetData);
+		}
+	}
+
+
+	ChildSlot
+		[
+			ComboButton->AsShared()
+		];
+
+}
+
+
+
+
+TSharedRef<SWidget> SToolInputAssetComboPanel::MakeCollectionSetsButtonPanel(TSharedRef<SToolInputAssetPicker> AssetPickerView)
+{
+	TArray<FNamedCollectionList> ShowCollectionSets = CollectionSets;
+	ShowCollectionSets.Insert(FNamedCollectionList{ FString(), TArray<FCollectionNameType>() }, 0);
+
+	ActiveCollectionSetIndex = 0;
+
+	TSharedRef<SHorizontalBox> HorizontalBox = SNew(SHorizontalBox);
+		
+	for ( int32 Index = 0; Index < ShowCollectionSets.Num(); ++Index)
+	{
+		const FNamedCollectionList& Collection = ShowCollectionSets[Index];
+
+		HorizontalBox->AddSlot()
+		.AutoWidth()
+		.Padding(4,0)
+		[
+			SNew(SCheckBox)
+				.Style(FAppStyle::Get(), "DetailsView.SectionButton")
+				.HAlign(HAlign_Center)
+				.OnCheckStateChanged_Lambda([this, Index, AssetPickerView](ECheckBoxState NewState)
+				{
+					if ( NewState == ECheckBoxState::Checked )
+					{
+						AssetPickerView->UpdateAssetSourceCollections( (Index == 0) ? 
+							TArray<FCollectionNameType>() : this->CollectionSets[Index-1].CollectionNames );
+						ActiveCollectionSetIndex = Index;
+					}
+				})
+				.IsChecked_Lambda([this, Index]() -> ECheckBoxState
+				{
+					return (this->ActiveCollectionSetIndex == Index) ? ECheckBoxState::Checked : ECheckBoxState::Unchecked;
+				})
+				[
+					SNew(SHorizontalBox)
+					+ SHorizontalBox::Slot()
+					.VAlign(VAlign_Center)
+					.Padding(FMargin(4, 2))
+					.AutoWidth()
+					[
+						SNew(STextBlock)
+						.Justification(ETextJustify::Center)
+						.TextStyle(FAppStyle::Get(), "DetailsView.CategoryTextStyle")
+						.Text( (Index == 0) ? LOCTEXT("AllFilterLabel", "Show All") :  FText::FromString(CollectionSets[Index-1].Name) )
+					]
+				]
+		];
+	}
+
+	return HorizontalBox;
+}
+
+
+void SToolInputAssetComboPanel::UpdateRecentAssets()
+{
+	if (RecentAssetsProvider.IsValid() == false)
+	{
+		return;
+	}
+	TArray<FAssetData> RecentAssets = RecentAssetsProvider->GetRecentAssetsList();
+	int32 NumRecent = RecentAssets.Num();
+
+	while (RecentAssetData.Num() < NumRecent)
+	{
+		FAssetThumbnailConfig ThumbnailConfig;
+		ThumbnailConfig.bAllowRealTimeOnHovered = false;
+
+		TSharedPtr<FRecentAssetInfo> NewInfo = MakeShared<FRecentAssetInfo>();
+		NewInfo->Index = RecentAssetData.Num();
+		RecentAssetData.Add(NewInfo);
+
+		RecentThumbnails.Add(MakeShared<FAssetThumbnail>(FAssetData(), FlyoutTileSize.X, FlyoutTileSize.Y, ThumbnailPool));
+		RecentThumbnailWidgets.Add(SNew(SBox)
+			.WidthOverride(FlyoutTileSize.X)
+			.HeightOverride(FlyoutTileSize.Y)
+			[
+				RecentThumbnails[NewInfo->Index]->MakeThumbnailWidget(ThumbnailConfig)
+			]);
+	}
+
+	for (int32 k = 0; k < NumRecent; ++k)
+	{
+		RecentAssetData[k]->AssetData = RecentAssets[k];
+		RecentThumbnails[k]->SetAsset(RecentAssets[k]);
+	}
+
+}
+
+
+TSharedRef<ITableRow> SToolInputAssetComboPanel::OnGenerateWidgetForRecentList(TSharedPtr<FRecentAssetInfo> InItem, const TSharedRef<STableViewBase>& OwnerTable)
+{
+	return  SNew(STableRow<TSharedPtr<FRecentAssetInfo>>, OwnerTable)
+		//.Style(FAppStyle::Get(), "ProjectBrowser.TableRow")		// do not like this style...
+		//.ToolTipText( InItem->AssetData )		// todo - cannot be assigned to STableRow, needs to go on nested widget
+		.Padding(2.0f)
+		[
+			// todo: maybe add a border around this?
+			RecentThumbnailWidgets[InItem->Index]->AsShared()
+		];
+}
+
+
+void SToolInputAssetComboPanel::NewAssetSelected(const FAssetData& AssetData)
+{
+	AssetThumbnail->SetAsset(AssetData);
+
+	if (Property.IsValid())
+	{
+		Property->SetValue(AssetData);
+	}
+
+	if (RecentAssetsProvider.IsValid())
+	{
+		RecentAssetsProvider->NotifyNewAsset(AssetData);
+	}
+
+	OnSelectionChanged.ExecuteIfBound(AssetData);
+
+	ComboButton->SetIsOpen(false);
+}
+
+
+FReply SToolInputAssetComboPanel::OnAssetThumbnailDoubleClick(const FGeometry& InMyGeometry, const FPointerEvent& InMouseEvent)
+{
+	const FAssetData& CurData = AssetThumbnail->GetAssetData();
+	if (CurData.IsValid())
+	{
+		UObject* ObjectToEdit = CurData.GetAsset();
+		if( ObjectToEdit )
+		{
+			GEditor->EditObject( ObjectToEdit );
+		}
+	}
+
+	// might be open from first click?
+	ComboButton->SetIsOpen(false);
+
+	return FReply::Handled();
+}
+
+
+
+
+#undef LOCTEXT_NAMESPACE
\ No newline at end of file
diff --git a/Engine/Plugins/Experimental/MeshModelingToolsetExp/Source/ModelingEditorUI/Private/ModelingWidgets/SToolInputAssetPicker.cpp b/Engine/Plugins/Experimental/MeshModelingToolsetExp/Source/ModelingEditorUI/Private/ModelingWidgets/SToolInputAssetPicker.cpp
new file mode 100644
--- /dev/null
+++ b/Engine/Plugins/Experimental/MeshModelingToolsetExp/Source/ModelingEditorUI/Private/ModelingWidgets/SToolInputAssetPicker.cpp
@@ -0,0 +1,408 @@
+// Copyright Epic Games, Inc. All Rights Reserved.
+
+#include "ModelingWidgets/SToolInputAssetPicker.h"
+#include "Styling/SlateTypes.h"
+#include "Widgets/SBoxPanel.h"
+#include "Layout/WidgetPath.h"
+#include "Framework/Application/SlateApplication.h"
+#include "Widgets/Layout/SSeparator.h"
+#include "Widgets/Layout/SSpacer.h"
+#include "Widgets/Images/SImage.h"
+#include "Widgets/Text/STextBlock.h"
+#include "Widgets/Input/SButton.h"
+#include "Widgets/Input/SComboButton.h"
+#include "Widgets/Input/SCheckBox.h"
+#include "EditorStyleSet.h"
+#include "FrontendFilters.h"
+#include "SAssetSearchBox.h"
+#include "SAssetView.h"
+#include "Framework/Commands/GenericCommands.h"
+#include "Editor.h"
+#include "PropertyHandle.h"
+
+#define LOCTEXT_NAMESPACE "SToolInputAssetPicker"
+
+SToolInputAssetPicker::~SToolInputAssetPicker()
+{
+}
+
+void SToolInputAssetPicker::Construct( const FArguments& InArgs )
+{
+	OnAssetsActivated = InArgs._AssetPickerConfig.OnAssetsActivated;
+	OnAssetSelected = InArgs._AssetPickerConfig.OnAssetSelected;
+	OnAssetEnterPressed = InArgs._AssetPickerConfig.OnAssetEnterPressed;
+	bPendingFocusNextFrame = InArgs._AssetPickerConfig.bFocusSearchBoxWhenOpened;
+	DefaultFilterMenuExpansion = InArgs._AssetPickerConfig.DefaultFilterMenuExpansion;
+
+	if ( InArgs._AssetPickerConfig.bFocusSearchBoxWhenOpened )
+	{
+		RegisterActiveTimer( 0.f, FWidgetActiveTimerDelegate::CreateSP( this, &SToolInputAssetPicker::SetFocusPostConstruct ) );
+	}
+
+	TSharedRef<SVerticalBox> VerticalBox = SNew(SVerticalBox);
+
+	ChildSlot
+	[
+		VerticalBox
+	];
+
+	TAttribute< FText > HighlightText;
+	EThumbnailLabel::Type ThumbnailLabel = InArgs._AssetPickerConfig.ThumbnailLabel;
+
+	FrontendFilters = MakeShareable(new FAssetFilterCollectionType());
+
+	TSharedRef<SHorizontalBox> HorizontalBox = SNew(SHorizontalBox);
+
+	if (!InArgs._AssetPickerConfig.bAutohideSearchBar)
+	{
+		// Search box
+		HighlightText = TAttribute< FText >(this, &SToolInputAssetPicker::GetHighlightedText);
+		HorizontalBox->AddSlot()
+		.FillWidth(1.0f)
+		[
+			SAssignNew( SearchBoxPtr, SAssetSearchBox )
+			.HintText(NSLOCTEXT( "ContentBrowser", "SearchBoxHint", "Search Assets" ))
+			.OnTextChanged( this, &SToolInputAssetPicker::OnSearchBoxChanged )
+			.OnTextCommitted( this, &SToolInputAssetPicker::OnSearchBoxCommitted )
+			.DelayChangeNotificationsWhileTyping( true )
+			.OnKeyDownHandler( this, &SToolInputAssetPicker::HandleKeyDownFromSearchBox )
+		];
+
+	}
+	else
+	{
+		HorizontalBox->AddSlot()
+		.FillWidth(1.0)
+		[
+			SNew(SSpacer)
+		];
+	}
+
+		
+	VerticalBox->AddSlot()
+	.AutoHeight()
+	.Padding(8.f, 0.f, 8.f, 8.f)
+	[
+		HorizontalBox
+	];
+
+	// "None" button
+	if (InArgs._AssetPickerConfig.bAllowNullSelection)
+	{
+		VerticalBox->AddSlot()
+		.AutoHeight()
+		[
+			SNew(SVerticalBox)
+
+			+ SVerticalBox::Slot()
+				.AutoHeight()
+				[
+					SNew(SButton)
+						.ButtonStyle( FEditorStyle::Get(), "ContentBrowser.NoneButton" )
+						.TextStyle( FEditorStyle::Get(), "ContentBrowser.NoneButtonText" )
+						.Text( LOCTEXT("NoneButtonText", "( None )") )
+						.ToolTipText( LOCTEXT("NoneButtonTooltip", "Clears the asset selection.") )
+						.HAlign(HAlign_Center)
+						.VAlign(VAlign_Center)
+						.OnClicked(this, &SToolInputAssetPicker::OnNoneButtonClicked)
+				]
+
+			// Trailing separator
+			+ SVerticalBox::Slot()
+				.AutoHeight()
+				.Padding(0, 0, 0, 4)
+				[
+					SNew(SSeparator)
+						.Orientation(Orient_Horizontal)
+				]
+		];
+	}
+
+	// Asset view
+	
+	// Break up the incoming filter into a sources data and backend filter.
+	CurrentSourcesData = FSourcesData(InArgs._AssetPickerConfig.Filter.PackagePaths, InArgs._AssetPickerConfig.Collections);
+	CurrentBackendFilter = InArgs._AssetPickerConfig.Filter;
+	CurrentBackendFilter.PackagePaths.Reset();
+
+	// Make game-specific filter
+	FOnShouldFilterAsset ShouldFilterAssetDelegate;
+	{
+		FAssetReferenceFilterContext AssetReferenceFilterContext;
+		AssetReferenceFilterContext.ReferencingAssets = InArgs._AssetPickerConfig.AdditionalReferencingAssets;
+		if (InArgs._AssetPickerConfig.PropertyHandle.IsValid())
+		{
+			TArray<UObject*> ReferencingObjects;
+			InArgs._AssetPickerConfig.PropertyHandle->GetOuterObjects(ReferencingObjects);
+			for (UObject* ReferencingObject : ReferencingObjects)
+			{
+				AssetReferenceFilterContext.ReferencingAssets.Add(FAssetData(ReferencingObject));
+			}
+		}
+		TSharedPtr<IAssetReferenceFilter> AssetReferenceFilter = GEditor ? GEditor->MakeAssetReferenceFilter(AssetReferenceFilterContext) : nullptr;
+		if (AssetReferenceFilter.IsValid())
+		{
+			FOnShouldFilterAsset ConfigFilter = InArgs._AssetPickerConfig.OnShouldFilterAsset;
+			ShouldFilterAssetDelegate = FOnShouldFilterAsset::CreateLambda([ConfigFilter, AssetReferenceFilter](const FAssetData& AssetData) -> bool {
+				if (!AssetReferenceFilter->PassesFilter(AssetData))
+				{
+					return true;
+				}
+				if (ConfigFilter.IsBound())
+				{
+					return ConfigFilter.Execute(AssetData);
+				}
+				return false;
+			});
+		}
+		else
+		{
+			ShouldFilterAssetDelegate = InArgs._AssetPickerConfig.OnShouldFilterAsset;
+		}
+	}
+
+	VerticalBox->AddSlot()
+	.FillHeight(1.f)
+	[
+		SAssignNew(AssetViewPtr, SAssetView)
+		.InitialCategoryFilter(EContentBrowserItemCategoryFilter::IncludeAssets)
+		.SelectionMode( InArgs._AssetPickerConfig.SelectionMode )
+		.OnShouldFilterAsset(ShouldFilterAssetDelegate)
+		.OnItemSelectionChanged(this, &SToolInputAssetPicker::HandleItemSelectionChanged)
+		.OnItemsActivated(this, &SToolInputAssetPicker::HandleItemsActivated)
+		.OnIsAssetValidForCustomToolTip(InArgs._AssetPickerConfig.OnIsAssetValidForCustomToolTip)
+		.OnGetCustomAssetToolTip(InArgs._AssetPickerConfig.OnGetCustomAssetToolTip)
+		.OnVisualizeAssetToolTip(InArgs._AssetPickerConfig.OnVisualizeAssetToolTip)
+		.OnAssetToolTipClosing(InArgs._AssetPickerConfig.OnAssetToolTipClosing)
+		.FrontendFilters(FrontendFilters)
+		.InitialSourcesData(CurrentSourcesData)
+		.InitialBackendFilter(CurrentBackendFilter)
+		.InitialViewType(InArgs._AssetPickerConfig.InitialAssetViewType)
+		.InitialAssetSelection(InArgs._AssetPickerConfig.InitialAssetSelection)
+		.ShowBottomToolbar(InArgs._AssetPickerConfig.bShowBottomToolbar)
+		.OnAssetTagWantsToBeDisplayed(InArgs._AssetPickerConfig.OnAssetTagWantsToBeDisplayed)
+		.OnGetCustomSourceAssets(InArgs._AssetPickerConfig.OnGetCustomSourceAssets)
+		.AllowDragging( InArgs._AssetPickerConfig.bAllowDragging )
+		.CanShowClasses( InArgs._AssetPickerConfig.bCanShowClasses )
+		.CanShowFolders( InArgs._AssetPickerConfig.bCanShowFolders )
+		.CanShowReadOnlyFolders( InArgs._AssetPickerConfig.bCanShowReadOnlyFolders )
+		.ShowPathInColumnView( InArgs._AssetPickerConfig.bShowPathInColumnView)
+		.ShowTypeInColumnView( InArgs._AssetPickerConfig.bShowTypeInColumnView)
+		.ShowViewOptions(false)  // We control this in the asset picker
+		.SortByPathInColumnView( InArgs._AssetPickerConfig.bSortByPathInColumnView)
+		.FilterRecursivelyWithBackendFilter( false )
+		.CanShowRealTimeThumbnails( InArgs._AssetPickerConfig.bCanShowRealTimeThumbnails )
+		.CanShowDevelopersFolder( InArgs._AssetPickerConfig.bCanShowDevelopersFolder )
+		.ForceShowEngineContent( InArgs._AssetPickerConfig.bForceShowEngineContent )
+		.ForceShowPluginContent( InArgs._AssetPickerConfig.bForceShowPluginContent )
+		.PreloadAssetsForContextMenu( InArgs._AssetPickerConfig.bPreloadAssetsForContextMenu )
+		.HighlightedText( HighlightText )
+		.ThumbnailLabel( ThumbnailLabel )
+		.AssetShowWarningText( InArgs._AssetPickerConfig.AssetShowWarningText)
+		.AllowFocusOnSync(false)	// Stop the asset view from stealing focus (we're in control of that)
+		.HiddenColumnNames(InArgs._AssetPickerConfig.HiddenColumnNames)
+		.CustomColumns(InArgs._AssetPickerConfig.CustomColumns)
+		// custom stuff, should expose as parameters
+		.InitialThumbnailSize(EThumbnailSize::Small)
+		.ShowTypeInTileView(false)
+	];
+
+
+	if (AssetViewPtr.IsValid() && !InArgs._AssetPickerConfig.bAutohideSearchBar)
+	{
+		TextFilter = MakeShareable(new FFrontendFilter_Text());
+		bool bClassNamesProvided = (InArgs._AssetPickerConfig.Filter.ClassNames.Num() != 1);
+		TextFilter->SetIncludeClassName(bClassNamesProvided || AssetViewPtr->IsIncludingClassNames());
+		TextFilter->SetIncludeAssetPath(AssetViewPtr->IsIncludingAssetPaths());
+		TextFilter->SetIncludeCollectionNames(AssetViewPtr->IsIncludingCollectionNames());
+	}
+
+	AssetViewPtr->RequestSlowFullListRefresh();
+}
+
+
+void SToolInputAssetPicker::UpdateAssetSourceCollections(TArray<FCollectionNameType> Collections)
+{
+	CurrentSourcesData = FSourcesData(CurrentSourcesData.VirtualPaths, Collections);
+	if (AssetViewPtr.IsValid())
+	{
+		AssetViewPtr->SetSourcesData(CurrentSourcesData);
+	}
+}
+
+EActiveTimerReturnType SToolInputAssetPicker::SetFocusPostConstruct( double InCurrentTime, float InDeltaTime )
+{
+	if ( SearchBoxPtr.IsValid() )
+	{
+		FWidgetPath WidgetToFocusPath;
+		FSlateApplication::Get().GeneratePathToWidgetUnchecked( SearchBoxPtr.ToSharedRef(), WidgetToFocusPath );
+		FSlateApplication::Get().SetKeyboardFocus( WidgetToFocusPath, EFocusCause::SetDirectly );
+		WidgetToFocusPath.GetWindow()->SetWidgetToFocusOnActivate(SearchBoxPtr);
+
+		return EActiveTimerReturnType::Stop;
+	}
+
+	return EActiveTimerReturnType::Continue;
+}
+
+FReply SToolInputAssetPicker::HandleKeyDownFromSearchBox(const FGeometry& MyGeometry, const FKeyEvent& InKeyEvent)
+{
+	// Up/Right and Down/Left move thru the filtered list
+	if (InKeyEvent.GetKey() == EKeys::Up || InKeyEvent.GetKey() == EKeys::Left)
+	{
+		AssetViewPtr->AdjustActiveSelection(-1);
+		return FReply::Handled();
+	}
+	else if (InKeyEvent.GetKey() == EKeys::Down || InKeyEvent.GetKey() == EKeys::Right)
+	{
+		AssetViewPtr->AdjustActiveSelection(+1);
+	}
+
+	return FReply::Unhandled();
+}
+
+FReply SToolInputAssetPicker::OnKeyDown(const FGeometry& MyGeometry, const FKeyEvent& InKeyEvent)
+{
+	if (InKeyEvent.GetKey() == EKeys::Enter)
+	{
+		TArray<FContentBrowserItem> SelectionSet = AssetViewPtr->GetSelectedFileItems();
+		if (SelectionSet.Num() == 1)
+		{
+			FAssetData ItemAssetData;
+			SelectionSet[0].Legacy_TryGetAssetData(ItemAssetData);
+			OnAssetSelected.ExecuteIfBound(ItemAssetData);
+		}
+		else
+		{
+			OnNoneButtonClicked();
+		}
+
+		return FReply::Handled();
+	}
+
+	return FReply::Unhandled();
+}
+
+FText SToolInputAssetPicker::GetHighlightedText() const
+{
+	return TextFilter->GetRawFilterText();
+}
+
+void SToolInputAssetPicker::SetSearchBoxText(const FText& InSearchText)
+{
+	// Has anything changed? (need to test case as the operators are case-sensitive)
+	if (!InSearchText.ToString().Equals(TextFilter->GetRawFilterText().ToString(), ESearchCase::CaseSensitive))
+	{
+		TextFilter->SetRawFilterText(InSearchText);
+		if (InSearchText.IsEmpty())
+		{
+			FrontendFilters->Remove(TextFilter);
+			AssetViewPtr->SetUserSearching(false);
+		}
+		else
+		{
+			FrontendFilters->Add(TextFilter);
+			AssetViewPtr->SetUserSearching(true);
+		}
+	}
+}
+
+void SToolInputAssetPicker::OnSearchBoxChanged(const FText& InSearchText)
+{
+	SetSearchBoxText( InSearchText );
+}
+
+void SToolInputAssetPicker::OnSearchBoxCommitted(const FText& InSearchText, ETextCommit::Type CommitInfo)
+{
+	SetSearchBoxText( InSearchText );
+
+	if (CommitInfo == ETextCommit::OnEnter)
+	{
+		TArray<FContentBrowserItem> SelectionSet = AssetViewPtr->GetSelectedFileItems();
+		if ( SelectionSet.Num() == 0 )
+		{
+			AssetViewPtr->AdjustActiveSelection(1);
+			SelectionSet = AssetViewPtr->GetSelectedFileItems();
+		}
+		HandleItemsActivated(SelectionSet, EAssetTypeActivationMethod::Opened);
+	}
+}
+
+
+
+FReply SToolInputAssetPicker::OnNoneButtonClicked()
+{
+	OnAssetSelected.ExecuteIfBound(FAssetData());
+	if (AssetViewPtr.IsValid())
+	{
+		AssetViewPtr->ClearSelection(true);
+	}
+	return FReply::Handled();
+}
+
+
+void SToolInputAssetPicker::HandleItemSelectionChanged(const FContentBrowserItem& InSelectedItem, ESelectInfo::Type InSelectInfo)
+{
+	if (InSelectInfo != ESelectInfo::Direct)
+	{
+		FAssetData ItemAssetData;
+		InSelectedItem.Legacy_TryGetAssetData(ItemAssetData);
+		OnAssetSelected.ExecuteIfBound(ItemAssetData);
+		
+	}
+}
+
+void SToolInputAssetPicker::HandleItemsActivated(TArrayView<const FContentBrowserItem> ActivatedItems, EAssetTypeActivationMethod::Type ActivationMethod)
+{
+	FContentBrowserItem FirstActivatedFolder;
+
+	TArray<FAssetData> ActivatedAssets;
+	for (const FContentBrowserItem& ActivatedItem : ActivatedItems)
+	{
+		if (ActivatedItem.IsFile())
+		{
+			FAssetData ItemAssetData;
+			if (ActivatedItem.Legacy_TryGetAssetData(ItemAssetData))
+			{
+				ActivatedAssets.Add(MoveTemp(ItemAssetData));
+			}
+		}
+
+		if (ActivatedItem.IsFolder() && !FirstActivatedFolder.IsValid())
+		{
+			FirstActivatedFolder = ActivatedItem;
+		}
+	}
+
+	if (ActivatedAssets.Num() == 0)
+	{
+		return;
+	}
+
+	if (ActivationMethod == EAssetTypeActivationMethod::Opened)
+	{
+		OnAssetEnterPressed.ExecuteIfBound(ActivatedAssets);
+	}
+
+	OnAssetsActivated.ExecuteIfBound( ActivatedAssets, ActivationMethod );
+}
+
+
+
+void SToolInputAssetPicker::RefreshAssetView(bool bRefreshSources)
+{
+	if (bRefreshSources)
+	{
+		AssetViewPtr->RequestSlowFullListRefresh();
+	}
+	else
+	{
+		AssetViewPtr->RequestQuickFrontendListRefresh();
+	}
+}
+
+
+
+
+
+#undef LOCTEXT_NAMESPACE
diff --git a/Engine/Plugins/Experimental/MeshModelingToolsetExp/Source/ModelingEditorUI/Public/ModelingEditorUIModule.h b/Engine/Plugins/Experimental/MeshModelingToolsetExp/Source/ModelingEditorUI/Public/ModelingEditorUIModule.h
new file mode 100644
--- /dev/null
+++ b/Engine/Plugins/Experimental/MeshModelingToolsetExp/Source/ModelingEditorUI/Public/ModelingEditorUIModule.h
@@ -0,0 +1,15 @@
+// Copyright Epic Games, Inc. All Rights Reserved.
+
+#pragma once
+
+#include "CoreMinimal.h"
+#include "Modules/ModuleManager.h"
+
+class FModelingEditorUIModule : public IModuleInterface
+{
+public:
+
+	/** IModuleInterface implementation */
+	virtual void StartupModule() override;
+	virtual void ShutdownModule() override;
+};
diff --git a/Engine/Plugins/Experimental/MeshModelingToolsetExp/Source/ModelingEditorUI/Public/ModelingWidgets/ModelingCustomizationUtil.h b/Engine/Plugins/Experimental/MeshModelingToolsetExp/Source/ModelingEditorUI/Public/ModelingWidgets/ModelingCustomizationUtil.h
new file mode 100644
--- /dev/null
+++ b/Engine/Plugins/Experimental/MeshModelingToolsetExp/Source/ModelingEditorUI/Public/ModelingWidgets/ModelingCustomizationUtil.h
@@ -0,0 +1,108 @@
+// Copyright Epic Games, Inc. All Rights Reserved.
+
+#pragma once
+
+#include "CoreMinimal.h"
+#include "PropertyHandle.h"
+#include "Widgets/SCompoundWidget.h"
+#include "Widgets/Text/STextBlock.h"
+#include "ModelingWidgets/SDynamicNumericEntry.h"
+
+class SWidget;
+class SBox;
+class SCheckBox;
+
+namespace UE
+{
+namespace ModelingUI
+{
+
+
+/**
+ * Constants used for various standard Modeling UI gizmos/etc
+ */
+struct ModelingUIConstants
+{
+	// top/bottom padding for detail row widgets, to be consistent w/ standard details panel
+	MODELINGEDITORUI_API static constexpr int DetailRowVertPadding = 4;
+
+	// minimum padding between label and widget, for tight packing scenarios
+	MODELINGEDITORUI_API static constexpr int LabelWidgetMinPadding = 4;
+
+	// padding between multiple widgets in a horizontal row (eg for two widgets, this is the padding between them) (("widget" here is often label+widget))
+	MODELINGEDITORUI_API static constexpr int MultiWidgetRowHorzPadding = 10;
+};
+
+
+
+
+
+/**
+ * Wrap a widget in a fixed-width (ie non-resizable) box 
+ */
+MODELINGEDITORUI_API TSharedRef<SBox> WrapInFixedWidthBox(TSharedRef<SWidget> SubWidget, int32 Width);
+
+
+/**
+* Make an HBox for a label/slider pair, with the label in a fixed-width box (ie no resizing)
+* The Label is taken from the existing label of a UProperty Handle
+*/
+MODELINGEDITORUI_API TSharedRef<SHorizontalBox> MakeFixedWidthLabelSliderHBox(
+	TSharedPtr<IPropertyHandle> LabelHandle,
+	TSharedPtr<SDynamicNumericEntry::FDataSource> SliderDataSource,
+	int32 LabelFixedWidth);
+
+/**
+* Make an HBox for a label/slider pair, with the label in a fixed-width box (ie no resizing)
+* The Label is taken from the existing label of a UProperty Handle
+*/
+MODELINGEDITORUI_API TSharedRef<SHorizontalBox> MakeToggleSliderHBox(
+	TSharedPtr<IPropertyHandle> BoolToggleHandle,
+	FText ToggleLabelText,
+	TSharedPtr<SDynamicNumericEntry::FDataSource> SliderDataSource,
+	int32 ToggleFixedWidth);
+
+ 
+/**
+* Make an HBox for two widgets in a details row, with standard padding between them
+*/
+MODELINGEDITORUI_API TSharedRef<SHorizontalBox> MakeTwoWidgetDetailRowHBox(
+	TSharedRef<SWidget> Widget1,
+	TSharedRef<SWidget> Widget2,
+	float FillWidth1 = 1.0f, float FillWidth2 = 1.0f);
+
+
+/**
+ * Make a Button that can be clicked to toggle the value of a bool UProperty.
+ * When true, the button is highlighted
+ */
+MODELINGEDITORUI_API TSharedRef<SCheckBox> MakeBoolToggleButton(
+	TSharedPtr<IPropertyHandle> BoolToggleHandle,
+	FText ButtonLabelText,
+	TFunction<void(bool)> ToggledCallback = [](bool) {},
+	int HorzPadding = 4);
+
+
+/**
+ * Create an error-string widget and add it to the Slot argument, to give the (developer) feedback about slate programming errors
+ * @param ErrorString error message
+ * @param Slot Object that supports bracket[SWidget] operator, such as any Box Slot object (eg can be called w/ SCompoundWidget::ChildSlot as the Slot argument)
+ */
+template<typename SlotType>
+void SetCustomWidgetErrorString(
+	FText ErrorString,
+	SlotType& Slot)
+{
+	Slot
+		.HAlign(HAlign_Center)
+		.VAlign(VAlign_Center)
+		[
+			SNew(STextBlock)
+			.Text(ErrorString)
+		];
+}
+
+
+
+}  // end namepace ModelingUI
+}  // end namepace UE
\ No newline at end of file
diff --git a/Engine/Plugins/Experimental/MeshModelingToolsetExp/Source/ModelingEditorUI/Public/ModelingWidgets/SComboPanel.h b/Engine/Plugins/Experimental/MeshModelingToolsetExp/Source/ModelingEditorUI/Public/ModelingWidgets/SComboPanel.h
new file mode 100644
--- /dev/null
+++ b/Engine/Plugins/Experimental/MeshModelingToolsetExp/Source/ModelingEditorUI/Public/ModelingWidgets/SComboPanel.h
@@ -0,0 +1,130 @@
+// Copyright Epic Games, Inc. All Rights Reserved.
+
+#pragma once
+
+#include "CoreMinimal.h"
+#include "Misc/Attribute.h"
+#include "Widgets/SWidget.h"
+#include "Widgets/SCompoundWidget.h"
+#include "Widgets/Views/STileView.h"
+#include "Styling/CoreStyle.h"
+
+class SComboButton;
+
+/**
+* SComboPanel provides a ComboBox-type interface but instead of a drop-down
+* list, a STileView panel ("Flyout") is shown with an icon for each selectable item.
+* Various options for the ComboButton are available, eg a single large icon, icon with adjacent label, etc.
+* 
+* Client provides a TArray<TSharedPtr<FComboPanelItem>> via the .ListItems() argument, 
+* where SComboPanel::FComboPanelItem contains the label, icon, and integer identifier for each selectable option.
+* 
+* 
+*/
+class MODELINGEDITORUI_API SComboPanel : public SCompoundWidget
+{
+public:
+	/** Item in the Combo Flyout (ie selectable option). */
+	struct FComboPanelItem
+	{
+		/** Text label shown as overlay on the icon */
+		FText Name;
+		/** Icon */
+		const FSlateBrush* Icon = nullptr;
+		/** Integer identifier. Must be unique for each Item. */
+		int Identifier = 0;
+	};
+
+	/** EComboDisplayType is used to specify the widgets in the ComboButton (ie non-expanded widget) */
+	enum class EComboDisplayType
+	{
+		/** Display a single large icon with overlay label, the same as is shown in the flyout */
+		LargeIcon,
+		IconAndLabel
+	};
+
+	/** Event that fires when selection changes */
+	DECLARE_DELEGATE_OneParam(FOnActiveItemChanged, TSharedPtr<FComboPanelItem>);
+	typedef FOnActiveItemChanged FOnSelectionChanged;
+
+public:
+
+	SLATE_BEGIN_ARGS( SComboPanel )
+		: _ComboButtonTileSize(50, 50)
+		, _FlyoutTileSize(85, 85)
+		, _FlyoutSize(600, 400)
+		, _OnSelectionChanged()
+		, _ComboDisplayType(EComboDisplayType::LargeIcon)
+		, _InitialSelectionIndex(0)
+		{
+		}
+
+		/** The size of the combo button icon tile */
+		SLATE_ARGUMENT(FVector2D, ComboButtonTileSize)
+
+		/** The size of the icon tiles in the flyout */
+		SLATE_ARGUMENT(FVector2D, FlyoutTileSize)
+
+		/** Size of the Flyout Panel */
+		SLATE_ARGUMENT(FVector2D, FlyoutSize)
+
+		/** List of selectable options */
+		SLATE_ARGUMENT( TArray<TSharedPtr<FComboPanelItem>> , ListItems )
+
+		/** Callback called when selection changes */
+		SLATE_EVENT( FOnSelectionChanged, OnSelectionChanged )
+
+		/** The text to display above the button list */
+		SLATE_ATTRIBUTE( FText, FlyoutHeaderText )
+
+		/** Type of widget layout for the ComboButton (default to LargeIcon) */
+		SLATE_ARGUMENT(EComboDisplayType, ComboDisplayType)
+		
+		/** Index into ListItems of the initially-selected item, defaults to index 0 */
+		SLATE_ARGUMENT(int, InitialSelectionIndex)
+
+	SLATE_END_ARGS()
+
+
+	/**
+	 * Construct this widget
+	 *
+	 * @param	InArgs	The declaration data for this widget
+	 */
+	void Construct( const FArguments& InArgs );
+
+
+
+protected:
+	FVector2D ComboButtonTileSize;
+	FVector2D FlyoutTileSize;
+	FVector2D FlyoutSize;
+
+	/** Pointer to the array of data items that we are observing */
+	TArray<TSharedPtr<FComboPanelItem>> Items;
+
+	/** Delegate to invoke when selection changes. */
+	FOnSelectionChanged OnSelectionChanged;
+
+	// STileView functions
+	virtual TSharedRef<ITableRow> CreateFlyoutIconTile(
+		TSharedPtr<FComboPanelItem> Item,
+		const TSharedRef<STableViewBase>& OwnerTable);
+	virtual void FlyoutSelectionChanged(
+		TSharedPtr<FComboPanelItem> SelectedItemIn,
+		ESelectInfo::Type SelectInfo);
+
+
+	// SComboButton functions
+	TSharedRef<SWidget> OnGetMenuContent();
+	void OnMenuOpenChanged(bool bOpen);
+
+	/** Icon to use when Item does not have a valid icon */
+	TSharedPtr<FSlateBrush> MissingIcon;
+
+	TSharedPtr<FComboPanelItem> SelectedItem;
+
+	TSharedPtr<SComboButton> ComboButton;
+	TSharedPtr<STileView<TSharedPtr<FComboPanelItem>>> TileView;
+	TSharedPtr<SBox> TileViewContainer;
+};
\ No newline at end of file
diff --git a/Engine/Plugins/Experimental/MeshModelingToolsetExp/Source/ModelingEditorUI/Public/ModelingWidgets/SDynamicNumericEntry.h b/Engine/Plugins/Experimental/MeshModelingToolsetExp/Source/ModelingEditorUI/Public/ModelingWidgets/SDynamicNumericEntry.h
new file mode 100644
--- /dev/null
+++ b/Engine/Plugins/Experimental/MeshModelingToolsetExp/Source/ModelingEditorUI/Public/ModelingWidgets/SDynamicNumericEntry.h
@@ -0,0 +1,66 @@
+// Copyright Epic Games, Inc. All Rights Reserved.
+
+#pragma once
+
+#include "CoreMinimal.h"
+#include "PropertyHandle.h"
+#include "Widgets/Input/SNumericEntryBox.h"
+#include "Widgets/Input/NumericTypeInterface.h"
+#include "Math/Interval.h"
+
+
+/**
+ * SNumericEntryBox wrapper that supports custom Get/SetValue and dynamic value/UI range
+ */
+class MODELINGEDITORUI_API SDynamicNumericEntry : public SCompoundWidget
+{
+public:
+
+	/**
+	 * 
+	 */
+	struct FDataSource
+	{
+		TFunction<void(float, EPropertyValueSetFlags::Type)> SetValue;
+		TFunction<float()> GetValue;
+		TFunction<TInterval<float>()> GetValueRange;
+		TFunction<TInterval<float>()> GetUIRange;
+	};
+
+public:
+
+	SLATE_BEGIN_ARGS( SDynamicNumericEntry )
+		: _MaxNumFloatDigits(2)
+	{
+	}
+
+	SLATE_ARGUMENT( TSharedPtr<FDataSource>, Source )
+
+	/** Maximum number of digits of precision that will be displayed for float slider */
+	SLATE_ARGUMENT( int, MaxNumFloatDigits )
+
+	SLATE_END_ARGS()
+
+
+	/**
+	* Construct this widget
+	*
+	* @param	InArgs	The declaration data for this widget
+	*/
+	void Construct( const FArguments& InArgs );
+
+	static TSharedPtr<FDataSource> MakeSimpleDataSource(TSharedPtr<IPropertyHandle> PropertyHandle, TInterval<float> ValueRange, TInterval<float> UIRange);
+
+protected:
+	TSharedPtr<FDataSource> Source;
+	TSharedPtr<SNumericEntryBox<float>> NumericEntry;
+	TSharedPtr<TDefaultNumericTypeInterface<float>> TypeInterface;
+
+	float LastSliderCommittedValue;
+	void OnValueChanged(float NewValue);
+	void OnValueCommitted(float NewValue, ETextCommit::Type CommitInfo);
+
+	bool bIsUsingSlider = false;
+	void OnBeginSliderMovement();
+	void OnEndSliderMovement(float NewValue);
+};
diff --git a/Engine/Plugins/Experimental/MeshModelingToolsetExp/Source/ModelingEditorUI/Public/ModelingWidgets/SToolInputAssetComboPanel.h b/Engine/Plugins/Experimental/MeshModelingToolsetExp/Source/ModelingEditorUI/Public/ModelingWidgets/SToolInputAssetComboPanel.h
new file mode 100644
--- /dev/null
+++ b/Engine/Plugins/Experimental/MeshModelingToolsetExp/Source/ModelingEditorUI/Public/ModelingWidgets/SToolInputAssetComboPanel.h
@@ -0,0 +1,147 @@
+// Copyright Epic Games, Inc. All Rights Reserved.
+
+#pragma once
+
+#include "CoreMinimal.h"
+#include "Misc/Attribute.h"
+#include "Widgets/SWidget.h"
+#include "Widgets/SCompoundWidget.h"
+#include "PropertyHandle.h"
+#include "AssetRegistry/AssetData.h"
+#include "CollectionManagerTypes.h"
+
+class SComboButton;
+class FAssetThumbnailPool;
+class FAssetThumbnail;
+class SBorder;
+class SBox;
+class ITableRow;
+class STableViewBase;
+class SToolInputAssetPicker;
+
+
+/**
+* SToolInputAssetComboPanel provides a similar UI to SComboPanel but
+* specifically for picking Assets. The standard widget is a SComboButton
+* that displays a thumbnail of the selected Asset, and on click a flyout
+* panel is shown that has an Asset Picker tile view, as well as (optionally)
+* a list of recently-used Assets, and also Collection-based filters.
+* 
+* Drag-and-drop onto the SComboButton is also supported, and the "selected Asset"
+* can be mapped to/from a PropertyHandle. However note that a PropertyHandle is *not* required,
+* each time the selection is modified the OnSelectionChanged(AssetData) delegate will also fire.
+* 
+* Note that "No Selection" is valid option by default
+*/
+class MODELINGEDITORUI_API SToolInputAssetComboPanel : public SCompoundWidget
+{
+public:
+
+	DECLARE_DELEGATE_OneParam(FOnSelectedAssetChanged, const FAssetData& AssetData);
+
+	/** List of Collections with associated Name, used to provide pickable Collection filters */
+	struct FNamedCollectionList
+	{
+		FString Name;
+		TArray<FCollectionNameType> CollectionNames;
+	};
+
+	/** IRecentAssetsProvider allows the Client to specify a set of "recently-used" Assets which the SToolInputAssetComboPanel will try to update as the selected Asset changes */
+	class IRecentAssetsProvider
+	{
+	public:
+		virtual ~IRecentAssetsProvider() {}
+		// SToolInputAssetComboPanel calls this to get the recent-assets list each time the flyout is opened
+		virtual TArray<FAssetData> GetRecentAssetsList() = 0;
+		// SToolInputAssetComboPanel calls this whenever the selected asset changes
+		virtual void NotifyNewAsset(const FAssetData& NewAsset) = 0;
+	};
+
+public:
+
+	SLATE_BEGIN_ARGS( SToolInputAssetComboPanel )
+		: _ComboButtonTileSize(50, 50)
+		, _FlyoutTileSize(85, 85)
+		, _FlyoutSize(600, 400)
+		, _AssetClassType(0)
+		, _OnSelectionChanged()
+	{
+		}
+
+		/** The size of the combo button icon tile */
+		SLATE_ARGUMENT(FVector2D, ComboButtonTileSize)
+
+		/** The size of the icon tiles in the flyout */
+		SLATE_ARGUMENT(FVector2D, FlyoutTileSize)
+
+		/** Size of the Flyout Panel */
+		SLATE_ARGUMENT(FVector2D, FlyoutSize)
+
+		/** Target PropertyHandle, selected value will be written here (Note: not required, selected  */
+		SLATE_ARGUMENT( TSharedPtr<IPropertyHandle>, Property )
+
+		/** Tooltip for the ComboButton. If Property is defined, this will be ignored. */
+		SLATE_ARGUMENT( FText, ToolTipText )
+
+		/** UClass of Asset to pick. Required, and only one class is supported */
+		SLATE_ARGUMENT( UClass*, AssetClassType )
+
+		/** (Optional) external provider/tracker of Recently-Used Assets. If not provided, Recent Assets area will not be shown. */
+		SLATE_ARGUMENT(TSharedPtr<IRecentAssetsProvider>, RecentAssetsProvider)
+
+		/** (Optional) set of collection-lists, if provided, button bar will be shown with each CollectionSet as an option */
+		SLATE_ARGUMENT(TArray<FNamedCollectionList>, CollectionSets)
+			
+		/** This delegate is executed each time the Selected Asset is modified */
+		SLATE_EVENT( FOnSelectedAssetChanged, OnSelectionChanged )
+
+	SLATE_END_ARGS()
+
+
+	/**
+	 * Construct this widget
+	 *
+	 * @param	InArgs	The declaration data for this widget
+	 */
+	void Construct( const FArguments& InArgs );
+
+protected:
+
+	FVector2D ComboButtonTileSize;
+	FVector2D FlyoutTileSize;
+	FVector2D FlyoutSize;
+	TSharedPtr<IPropertyHandle> Property;
+	UClass* AssetClassType = nullptr;
+
+
+	/** Delegate to invoke when selection changes. */
+	FOnSelectedAssetChanged OnSelectionChanged;
+
+	TSharedPtr<SComboButton> ComboButton;
+	TSharedPtr<FAssetThumbnailPool> ThumbnailPool;
+	TSharedPtr<FAssetThumbnail> AssetThumbnail;
+	TSharedPtr<SBorder> ThumbnailBorder;
+
+	TSharedPtr<IRecentAssetsProvider> RecentAssetsProvider;
+
+	struct FRecentAssetInfo
+	{
+		int Index;
+		FAssetData AssetData;
+	};
+	TArray<TSharedPtr<FRecentAssetInfo>> RecentAssetData;
+
+	TArray<TSharedPtr<FAssetThumbnail>> RecentThumbnails;
+	TArray<TSharedPtr<SBox>> RecentThumbnailWidgets;
+
+	void UpdateRecentAssets();
+
+	virtual void NewAssetSelected(const FAssetData& AssetData);
+	virtual FReply OnAssetThumbnailDoubleClick(const FGeometry& InMyGeometry, const FPointerEvent& InMouseEvent);
+
+	TSharedRef<ITableRow> OnGenerateWidgetForRecentList(TSharedPtr<FRecentAssetInfo> InItem, const TSharedRef<STableViewBase>& OwnerTable);
+
+	TArray<FNamedCollectionList> CollectionSets;
+	int32 ActiveCollectionSetIndex = -1;
+	TSharedRef<SWidget> MakeCollectionSetsButtonPanel(TSharedRef<SToolInputAssetPicker> AssetPickerView);
+};
\ No newline at end of file
diff --git a/Engine/Plugins/Experimental/MeshModelingToolsetExp/Source/ModelingEditorUI/Public/ModelingWidgets/SToolInputAssetPicker.h b/Engine/Plugins/Experimental/MeshModelingToolsetExp/Source/ModelingEditorUI/Public/ModelingWidgets/SToolInputAssetPicker.h
new file mode 100644
--- /dev/null
+++ b/Engine/Plugins/Experimental/MeshModelingToolsetExp/Source/ModelingEditorUI/Public/ModelingWidgets/SToolInputAssetPicker.h
@@ -0,0 +1,130 @@
+// Copyright Epic Games, Inc. All Rights Reserved.
+
+#pragma once
+
+#include "CoreMinimal.h"
+#include "Input/Reply.h"
+#include "Widgets/DeclarativeSyntaxSupport.h"
+#include "Widgets/SWidget.h"
+#include "Widgets/SCompoundWidget.h"
+#include "AssetData.h"
+#include "ARFilter.h"
+#include "IContentBrowserSingleton.h"
+#include "SourcesData.h"
+
+class FFrontendFilter_Text;
+class FUICommandList;
+class SAssetSearchBox;
+class SAssetView;
+class SComboButton;
+enum class ECheckBoxState : uint8;
+
+/**
+ * SToolInputAssetPicker is designed to support Asset picking in Modeling Tools, 
+ * where the Assets in question are input parameters/options to Tools, eg such as
+ * a brush alpha texture for use in a Painting/Sculpting Tool. 
+ * 
+ * Implementation is derived from SAssetPicker (private class in the ContentBrowser module).
+ * However many optional features have been stripped out as they are not relevant
+ * in the Modeling-Tool Parameters context. 
+ * 
+ * Most settings are provided via the FAssetPickerConfig, which is passed to an SAssetView internally.
+ * 
+ * Unless you are really certain you need to use this class directly, it's likely that 
+ * you are looking for SToolInputAssetComboPanel, which provides a combobox/flyout-style
+ * widget suitable for user interface panels.
+ */
+class MODELINGEDITORUI_API SToolInputAssetPicker : public SCompoundWidget
+{
+public:
+	SLATE_BEGIN_ARGS( SToolInputAssetPicker ){}
+
+		/** A struct containing details about how the asset picker should behave */
+		SLATE_ARGUMENT(FAssetPickerConfig, AssetPickerConfig)
+
+	SLATE_END_ARGS()
+
+	virtual ~SToolInputAssetPicker();
+
+	/** Constructs this widget with InArgs */
+	void Construct( const FArguments& InArgs );
+
+	// SWidget implementation
+	virtual FReply OnKeyDown( const FGeometry& MyGeometry, const FKeyEvent& InKeyEvent ) override;
+	// End of SWidget implementation
+
+	/** Return the associated AssetView */
+	const TSharedPtr<SAssetView>& GetAssetView() const { return AssetViewPtr; }
+
+	/**
+	 * Update the set of input Assets to be only based on the given set of Collections
+	 * (or all Assets, if the Collections list is empty)
+	 */
+	virtual void UpdateAssetSourceCollections(TArray<FCollectionNameType> Collections);
+
+private:
+	/** Focuses the search box post-construct */
+	EActiveTimerReturnType SetFocusPostConstruct( double InCurrentTime, float InDeltaTime );
+
+	/** Special case handling for SAssetSearchBox key commands */
+	FReply HandleKeyDownFromSearchBox(const FGeometry& MyGeometry, const FKeyEvent& InKeyEvent);
+
+	/** Called when the editable text needs to be set or cleared */
+	void SetSearchBoxText(const FText& InSearchText);
+
+	/** Called by the editable text control when the search text is changed by the user */
+	void OnSearchBoxChanged(const FText& InSearchText);
+
+	/** Called by the editable text control when the user commits a text change */
+	void OnSearchBoxCommitted(const FText& InSearchText, ETextCommit::Type CommitInfo);
+
+	/** Handler for when the "None" button is clicked */
+	FReply OnNoneButtonClicked();
+
+	/** Handle forwarding picking events. We wrap OnAssetSelected here to prevent 'Direct' selections being identified as user actions */
+	void HandleItemSelectionChanged(const FContentBrowserItem& InSelectedItem, ESelectInfo::Type InSelectInfo);
+
+	/** Handler for when the user double clicks, presses enter, or presses space on an asset */
+	void HandleItemsActivated(TArrayView<const FContentBrowserItem> ActivatedItems, EAssetTypeActivationMethod::Type ActivationMethod);
+
+	/** Forces a refresh */
+	void RefreshAssetView(bool bRefreshSources);
+
+	/** @return The text to highlight on the assets  */
+	FText GetHighlightedText() const;
+
+private:
+
+	/** The list of FrontendFilters currently applied to the asset view */
+	TSharedPtr<FAssetFilterCollectionType> FrontendFilters;
+
+	/** The asset view widget */
+	TSharedPtr<SAssetView> AssetViewPtr;
+
+	/** The search box */
+	TSharedPtr<SAssetSearchBox> SearchBoxPtr;
+
+	/** Called to when an asset is selected or the none button is pressed */
+	FOnAssetSelected OnAssetSelected;
+
+	/** Called when enter is pressed while an asset is selected */
+	FOnAssetEnterPressed OnAssetEnterPressed;
+
+	/** Called when any number of assets are activated */
+	FOnAssetsActivated OnAssetsActivated;
+
+	/** True if the search box will take keyboard focus next frame */
+	bool bPendingFocusNextFrame;
+
+	/** Filters needed for filtering the assets */
+	TSharedPtr< FAssetFilterCollectionType > FilterCollection;
+	TSharedPtr< FFrontendFilter_Text > TextFilter;
+
+	EAssetTypeCategories::Type DefaultFilterMenuExpansion;
+
+	/** The sources data currently used by the picker */
+	FSourcesData CurrentSourcesData;
+
+	/** Current filter we are using, needed reset asset view after we have custom filtered */
+	FARFilter CurrentBackendFilter;
+};
