ModelingTools: Improve performance of marquee selection for large meshes and large selections by deferring computations until the mouse is released

#rnx

diff --git a/Engine/Plugins/Experimental/MeshModelingToolsetExp/Source/ModelingComponents/Private/Mechanics/RectangleMarqueeMechanic.cpp b/Engine/Plugins/Experimental/MeshModelingToolsetExp/Source/ModelingComponents/Private/Mechanics/RectangleMarqueeMechanic.cpp
--- a/Engine/Plugins/Experimental/MeshModelingToolsetExp/Source/ModelingComponents/Private/Mechanics/RectangleMarqueeMechanic.cpp
+++ b/Engine/Plugins/Experimental/MeshModelingToolsetExp/Source/ModelingComponents/Private/Mechanics/RectangleMarqueeMechanic.cpp
@@ -5,6 +5,7 @@
 #include "CanvasTypes.h"
 #include "CanvasItem.h"
 #include "ToolSceneQueriesUtil.h"
+#include "ProfilingDebugging/ScopedTimers.h"
 
 #include "Intersection/IntersectionQueries2.h"
 #include "ExplicitUseGeometryMathTypes.h"		// using UE::Geometry::(math types)
@@ -179,6 +180,7 @@ void URectangleMarqueeMechanic::OnClickPress(const FInputDeviceRay& PressPos)
 	CameraRectangle.RectangleEndRay = PressPos;
 	CameraRectangle.Initialize();
 
+	bIsOnDragRectangleChangedDeferred = false;
 	OnDragRectangleStarted.Broadcast();
 }
 
@@ -193,11 +195,30 @@ void URectangleMarqueeMechanic::OnClickDrag(const FInputDeviceRay& DragPos)
 	CameraRectangle.RectangleEndRay = DragPos;
 	CameraRectangle.Initialize();
 	
+	if (bIsOnDragRectangleChangedDeferred)
+	{
+		return;
+	}
+	
+	double Time = 0.0;
+	FDurationTimer Timer(Time);
 	OnDragRectangleChanged.Broadcast(CameraRectangle);
+	Timer.Stop();
+	
+	if (Time > OnDragRectangleChangedDeferredThreshold)
+	{
+		bIsOnDragRectangleChangedDeferred = true;
+	}
 }
 
 void URectangleMarqueeMechanic::OnClickRelease(const FInputDeviceRay& ReleasePos)
 {
+	if (bIsOnDragRectangleChangedDeferred)
+	{
+		bIsOnDragRectangleChangedDeferred = false;
+		OnClickDrag(ReleasePos);
+	}
+	
 	bIsDragging = false;
 	OnDragRectangleFinished.Broadcast(CameraRectangle, false);
 }
diff --git a/Engine/Plugins/Experimental/MeshModelingToolsetExp/Source/ModelingComponents/Private/Selection/MeshSelectionMechanic.cpp b/Engine/Plugins/Experimental/MeshModelingToolsetExp/Source/ModelingComponents/Private/Selection/MeshSelectionMechanic.cpp
--- a/Engine/Plugins/Experimental/MeshModelingToolsetExp/Source/ModelingComponents/Private/Selection/MeshSelectionMechanic.cpp
+++ b/Engine/Plugins/Experimental/MeshModelingToolsetExp/Source/ModelingComponents/Private/Selection/MeshSelectionMechanic.cpp
@@ -280,6 +280,7 @@ void UMeshSelectionMechanic::Setup(UInteractiveTool* ParentToolIn)
 	// This will be the target for the click drag behavior below
 	MarqueeMechanic = NewObject<URectangleMarqueeMechanic>();
 	MarqueeMechanic->bUseExternalClickDragBehavior = true;
+	MarqueeMechanic->OnDragRectangleChangedDeferredThreshold = 1./60;
 	MarqueeMechanic->Setup(ParentToolIn);
 	MarqueeMechanic->OnDragRectangleStarted.AddUObject(this, &UMeshSelectionMechanic::OnDragRectangleStarted);
 	MarqueeMechanic->OnDragRectangleChanged.AddUObject(this, &UMeshSelectionMechanic::OnDragRectangleChanged);
@@ -733,6 +734,8 @@ void UMeshSelectionMechanic::OnClicked(const FInputDeviceRay& ClickPos)
 
 void UMeshSelectionMechanic::OnDragRectangleStarted()
 {
+	TRACE_CPUPROFILER_EVENT_SCOPE(MeshSelectionMechanic_OnDragRectangleStarted); // Mark start of drag sequence
+	
 	if (CurrentSelection.Type != MeshSelectionMechanicLocals::ToCompatibleDynamicMeshSelectionType(SelectionMode))
 	{
 		ClearCurrentSelection();	
@@ -880,6 +883,8 @@ void UMeshSelectionMechanic::OnDragRectangleChanged(const FCameraRectangle& Curr
 
 void UMeshSelectionMechanic::OnDragRectangleFinished(const FCameraRectangle&, bool bCancelled)
 {
+	TRACE_CPUPROFILER_EVENT_SCOPE(MeshSelectionMechanic_OnDragRectangleFinished); // Mark end of drag sequence
+
 	if (!bCancelled && (PreDragSelection != CurrentSelection))
 	{
 		UpdateCentroid();
diff --git a/Engine/Plugins/Experimental/MeshModelingToolsetExp/Source/ModelingComponents/Public/Mechanics/RectangleMarqueeMechanic.h b/Engine/Plugins/Experimental/MeshModelingToolsetExp/Source/ModelingComponents/Public/Mechanics/RectangleMarqueeMechanic.h
--- a/Engine/Plugins/Experimental/MeshModelingToolsetExp/Source/ModelingComponents/Public/Mechanics/RectangleMarqueeMechanic.h
+++ b/Engine/Plugins/Experimental/MeshModelingToolsetExp/Source/ModelingComponents/Public/Mechanics/RectangleMarqueeMechanic.h
@@ -130,6 +130,14 @@ public:
 	UPROPERTY()
 	bool bUseExternalClickDragBehavior = false;
 
+	/**
+	 * If the computation time for a single call to OnDragRectangleChanged ever exceeds this threshold then future calls
+	 * to this function (in the current drag sequence) will be deferred until the mouse button is released. This will
+	 * improve the responsiveness of the UI. The default value is set so this optimization is never triggered, if you
+	 * want it you can set this to a small value e.g., 1./60 (time elapsed by 1 frame at 60 fps)
+	 */
+	UPROPERTY()
+	double OnDragRectangleChangedDeferredThreshold = TNumericLimits<double>::Max(); // In seconds
 
 public:
 
@@ -198,4 +206,5 @@ private:
 
 	bool bIsEnabled;
 	bool bIsDragging;
+	bool bIsOnDragRectangleChangedDeferred;
 };
