MeshModelingToolset: Cleanup PolyDef tool

diff --git a/Engine/Plugins/Experimental/MeshModelingToolsetExp/Source/MeshModelingToolsExp/Private/DeformMeshPolygonsTool.cpp b/Engine/Plugins/Experimental/MeshModelingToolsetExp/Source/MeshModelingToolsExp/Private/DeformMeshPolygonsTool.cpp
--- a/Engine/Plugins/Experimental/MeshModelingToolsetExp/Source/MeshModelingToolsExp/Private/DeformMeshPolygonsTool.cpp
+++ b/Engine/Plugins/Experimental/MeshModelingToolsetExp/Source/MeshModelingToolsExp/Private/DeformMeshPolygonsTool.cpp
@@ -2,14 +2,9 @@
 
 #include "DeformMeshPolygonsTool.h"
 
-#include "Async/ParallelFor.h"
-#include "Containers/BitArray.h"
-#include "DynamicMesh/DynamicMeshAttributeSet.h"
 #include "InteractiveToolManager.h"
-#include "Intersection/IntersectionUtil.h"
 #include "DynamicMesh/MeshNormals.h"
 #include "ModelingOperators/Public/ModelingTaskTypes.h"
-#include "SegmentTypes.h"
 #include "Solvers/ConstrainedMeshDeformer.h"
 #include "ToolBuilderUtil.h"
 #include "ToolSetupUtil.h"
@@ -84,8 +79,8 @@ public:
 	};
 
 	FConstrainedMeshDeformerTask(const ELaplacianWeightScheme SelectedWeightScheme)
+		: LaplacianWeightScheme(SelectedWeightScheme)
 	{
-		LaplacianWeightScheme = SelectedWeightScheme;
 	}
 
 	virtual ~FConstrainedMeshDeformerTask(){};
@@ -145,7 +140,7 @@ private:
 	bool* bAbortSource = nullptr;
 
 	/** Used to initialize the array mapping, updated during the UpdateDeformer() function */
-	int SrcMeshMaxVertexID;
+	int SrcMeshMaxVertexID = 0;
 
 	/** A subset of the original mesh */
 	FDynamicMesh3 SubsetMesh;
@@ -155,9 +150,6 @@ private:
 
 	/** Laplacian deformer object gets rebuilt each new transaction */
 	TUniquePtr<UE::Solvers::IConstrainedMeshSolver> ConstrainedDeformer;
-
-private:
-	FConstrainedMeshDeformerTask();
 };
 
 class FGroupTopologyLaplacianDeformer : public FGroupTopologyDeformer
@@ -478,7 +470,6 @@ void FConstrainedMeshDeformerTask::ExportResults(FDynamicMesh3& TargetMesh) cons
 
 void FConstrainedMeshDeformerTask::ApplyAttenuation()
 {
-	size_t Size = SrcMeshMaxVertexID;
 	TSet<int> Handles;
 
 	auto InPlaceMinMaxElements = [](FVector3d& Min, FVector3d& Max, const FVector3d Test) {
@@ -838,7 +829,7 @@ inline void FGroupTopologyLaplacianDeformer::Shutdown()
 
 UDeformMeshPolygonsTool::UDeformMeshPolygonsTool()
 {
-	SetToolDisplayName(LOCTEXT("DeformPolygroupsToolName", "PolyGroup Deform"));
+	UInteractiveTool::SetToolDisplayName(LOCTEXT("DeformPolygroupsToolName", "PolyGroup Deform"));
 }
 
 void UDeformMeshPolygonsTool::Setup()
@@ -929,12 +920,12 @@ void UDeformMeshPolygonsTool::Setup()
 	{
 		GetToolManager()->DisplayMessage(
 			LOCTEXT("NoGroupsWarning",
-			        "This object has only a single Polygroup. Use the GenGrps, GrpPnt or TriSel (Create Polygroup) tools to modify Polygroups."),
+			        "This object has only a single PolyGroup. Use the GrpGen, GrpPnt or TriSel (Create PolyGroup) tools to modify PolyGroups."),
 			EToolMessageLevel::UserWarning);
 	}
 
 	GetToolManager()->DisplayMessage(
-		LOCTEXT("DeformMeshPolygonsToolDescription", "Deform the mesh by directly manipulating (i.e. click-and-drag) the Polygroup Edges/Faces/Vertices."),
+		LOCTEXT("DeformMeshPolygonsToolDescription", "Deform the mesh by directly manipulating (i.e. click-and-drag) the PolyGroup edges, faces, and vertices."),
 		EToolMessageLevel::UserNotification);
 }
 
@@ -952,7 +943,7 @@ void UDeformMeshPolygonsTool::Shutdown(EToolShutdownType ShutdownType)
 		if (ShutdownType == EToolShutdownType::Accept)
 		{
 			// this block bakes the modified DynamicMeshComponent back into the StaticMeshComponent inside an undo transaction
-			GetToolManager()->BeginUndoTransaction(LOCTEXT("DeformMeshPolygonsToolTransactionName", "Deform Mesh"));
+			GetToolManager()->BeginUndoTransaction(LOCTEXT("DeformMeshPolygonsToolTransactionName", "PolyGroup Deform"));
 			DynamicMeshComponent->ProcessMesh([&](const FDynamicMesh3& ReadMesh)
 			{
 				FConversionToMeshDescriptionOptions ConversionOptions;
@@ -1389,7 +1380,7 @@ void UDeformMeshPolygonsTool::ComputeUpdate_Rotate()
 	FGroupTopologyDeformer& SelectedDeformer = (bIsLaplacian) ? *LaplacianDeformer : LinearDeformer;
 
 	FDynamicMesh3* Mesh    = DynamicMeshComponent->GetMesh();
-	FVector NewHitPosWorld = LastHitPosWorld;
+	FVector NewHitPosWorld;
 
 	FVector3d SnappedPoint;
 	if (QuickAxisRotator.UpdateSnap(FRay3d(UpdateRay), SnappedPoint))
@@ -1482,7 +1473,7 @@ void UDeformMeshPolygonsTool::ComputeUpdate_Translate()
 		};
 	}
 
-	FVector NewHitPosWorld = LastHitPosWorld;
+	FVector NewHitPosWorld;
 	FVector3d SnappedPoint;
 	if (QuickAxisTranslater.UpdateSnap(FRay3d(UpdateRay), SnappedPoint, PointConstraintFunc))
 	{
diff --git a/Engine/Plugins/Experimental/MeshModelingToolsetExp/Source/MeshModelingToolsExp/Public/DeformMeshPolygonsTool.h b/Engine/Plugins/Experimental/MeshModelingToolsetExp/Source/MeshModelingToolsExp/Public/DeformMeshPolygonsTool.h
--- a/Engine/Plugins/Experimental/MeshModelingToolsetExp/Source/MeshModelingToolsExp/Public/DeformMeshPolygonsTool.h
+++ b/Engine/Plugins/Experimental/MeshModelingToolsetExp/Source/MeshModelingToolsExp/Public/DeformMeshPolygonsTool.h
@@ -17,7 +17,6 @@
 #include "Selection/GroupTopologySelector.h"
 #include "Operations/GroupTopologyDeformer.h"
 #include "Solvers/MeshLaplacian.h"
-#include "Curves/CurveFloat.h"
 #include "DeformMeshPolygonsTool.generated.h"
 
 class FMeshVertexChangeBuilder;
@@ -33,9 +32,6 @@ class MESHMODELINGTOOLSEXP_API UDeformMeshPolygonsToolBuilder : public UMeshSurf
 	GENERATED_BODY()
 
 public:
-	UDeformMeshPolygonsToolBuilder()
-	{
-	}
 
 	virtual UMeshSurfacePointTool* CreateNewTool(const FToolBuilderState& SceneState) const override;
 };
@@ -44,10 +40,10 @@ public:
 UENUM()
 enum class EGroupTopologyDeformationStrategy : uint8
 {
-	/** Deforms the mesh using linear translations*/
+	/** Deforms the mesh using linear deformation. */
 	Linear UMETA(DisplayName = "Linear"),
 
-	/** Deforms the mesh using laplacian deformation*/
+	/** Deforms the mesh using Laplacian deformation. */
 	Laplacian UMETA(DisplayName = "Smooth")
 };
 
@@ -75,10 +71,10 @@ static ELaplacianWeightScheme ConvertToLaplacianWeightScheme(const EWeightScheme
 UENUM()
 enum class EQuickTransformerMode : uint8
 {
-	/** Translation along frame axes */
+	/** Translation along axes */
 	AxisTranslation = 0 UMETA(DisplayName = "Translate"),
 
-	/** Rotation around frame axes*/
+	/** Rotation around axes */
 	AxisRotation = 1 UMETA(DisplayName = "Rotate"),
 };
 
@@ -95,21 +91,27 @@ public:
 
 	//Options
 
-	UPROPERTY(EditAnywhere, Category = Options, meta = (DisplayName = "Deformation Type", ToolTip = "Select the type of deformation you wish to employ on a polygroup."))
+	/** Type of deformation used. */
+	UPROPERTY(EditAnywhere, Category = Options, meta = (DisplayName = "Deformation"))
 	EGroupTopologyDeformationStrategy DeformationStrategy;
 
-	UPROPERTY(EditAnywhere, Category = Options)
+	/** Type of transformation used. */
+	UPROPERTY(EditAnywhere, Category = Options, meta = (DisplayName = "Transformation"))
 	EQuickTransformerMode TransformMode;
 
-	UPROPERTY(EditAnywhere, Category = SelectionFilter)
+	/** Allow for faces (PolyGroups) to be selected. */
+	UPROPERTY(EditAnywhere, Category = Selection)
 	bool bSelectFaces;
 
-	UPROPERTY(EditAnywhere, Category = SelectionFilter)
+	/** Allow for edges to be selected. */
+	UPROPERTY(EditAnywhere, Category = Selection)
 	bool bSelectEdges;
 
-	UPROPERTY(EditAnywhere, Category = SelectionFilter)
+	/** Allow for vertices to be selected. */
+	UPROPERTY(EditAnywhere, Category = Selection)
 	bool bSelectVertices;
 
+	/** If true, overlays preview with wireframe. */
 	UPROPERTY(EditAnywhere, Category = Display)
 	bool bShowWireframe;
 
