ModelingTools: minor refactors to Sculpt Tools to support UI customization
- Added new FBrushToolRadius struct that provides a "radius" abstraction that supports adaptive and world parameters
- USculptBrushProperties no longer subclasses UBrushBaseProperties, most parameters were not being used, now uses new FBrushToolRadius
- added UMeshSculptToolBase::FBrushTypeInfo FFalloffTypeInfo, and track one for each registered brush type and falloff type. Allows external UI code to know about available brush and falloff types.
- UMeshVertexSculptTool now has functions to set active brush, falloff, and region filter type that can be used at UI level
- minor updates to various brush ops, dynamic and group paint tool for above refactoring and/or UI stuff

#rnx

diff --git a/Engine/Plugins/Experimental/MeshModelingToolsetExp/Source/MeshModelingToolsExp/Private/DynamicMeshSculptTool.cpp b/Engine/Plugins/Experimental/MeshModelingToolsetExp/Source/MeshModelingToolsExp/Private/DynamicMeshSculptTool.cpp
--- a/Engine/Plugins/Experimental/MeshModelingToolsetExp/Source/MeshModelingToolsExp/Private/DynamicMeshSculptTool.cpp
+++ b/Engine/Plugins/Experimental/MeshModelingToolsetExp/Source/MeshModelingToolsExp/Private/DynamicMeshSculptTool.cpp
@@ -147,10 +147,11 @@ void UDynamicMeshSculptTool::Setup()
 	double MaxDimension = DynamicMeshComponent->GetMesh()->GetBounds(true).MaxDim();
 	BrushRelativeSizeRange = FInterval1d(MaxDimension*0.01, MaxDimension);
 	BrushProperties = NewObject<USculptBrushProperties>(this);
-	BrushProperties->bShowStrength = false;
 	BrushProperties->bShowFlowRate = false;
 	BrushProperties->bShowSpacing = false;
 	BrushProperties->bShowLazyness = false;
+	BrushProperties->BrushSize.InitializeWorldSizeRange(
+		TInterval<float>((float)BrushRelativeSizeRange.Min, (float)BrushRelativeSizeRange.Max));
 	CalculateBrushRadius();
 
 	// initialize other properties
@@ -484,15 +485,7 @@ void UDynamicMeshSculptTool::OnUpdateDrag(const FRay& WorldRay)
 
 void UDynamicMeshSculptTool::CalculateBrushRadius()
 {
-	CurrentBrushRadius = 0.5 * BrushRelativeSizeRange.Interpolate(BrushProperties->BrushSize);
-	if (BrushProperties->bSpecifyRadius)
-	{
-		CurrentBrushRadius = BrushProperties->BrushRadius;
-	}
-	else
-	{
-		BrushProperties->BrushRadius = CurrentBrushRadius;
-	}
+	CurrentBrushRadius = BrushProperties->BrushSize.GetWorldRadius();
 }
 
 void UDynamicMeshSculptTool::ApplyStamp(const FRay& WorldRay)
@@ -1131,7 +1124,7 @@ bool UDynamicMeshSculptTool::ApplyPullKelvinBrush(const FRay& WorldRay)
 
 	const EKelvinletBrushMode KelvinMode = EKelvinletBrushMode::PullKelvinlet;
 
-	FKelvinletBrushOp::FKelvinletBrushOpProperties  KelvinletBrushOpProperties(KelvinMode, *KelvinBrushProperties, *BrushProperties);
+	FKelvinletBrushOp::FKelvinletBrushOpProperties  KelvinletBrushOpProperties(KelvinMode, *KelvinBrushProperties, CurrentBrushRadius, BrushProperties->BrushFalloffAmount);
 	KelvinletBrushOpProperties.Direction = FVector(MoveVec.X, MoveVec.Y, MoveVec.Z);  //FVector(BrushNormalLocal.X, BrushNormalLocal.Y, BrushNormalLocal.Z);
 	KelvinletBrushOpProperties.Size *= 0.6;
 
@@ -1163,7 +1156,7 @@ bool UDynamicMeshSculptTool::ApplyPullSharpKelvinBrush(const FRay& WorldRay)
 
 	const EKelvinletBrushMode KelvinMode = EKelvinletBrushMode::SharpPullKelvinlet;
 
-	FKelvinletBrushOp::FKelvinletBrushOpProperties  KelvinletBrushOpProperties(KelvinMode, *KelvinBrushProperties, *BrushProperties);
+	FKelvinletBrushOp::FKelvinletBrushOpProperties  KelvinletBrushOpProperties(KelvinMode, *KelvinBrushProperties, CurrentBrushRadius, BrushProperties->BrushFalloffAmount);
 	KelvinletBrushOpProperties.Direction = FVector(MoveVec.X, MoveVec.Y, MoveVec.Z);  //FVector(BrushNormalLocal.X, BrushNormalLocal.Y, BrushNormalLocal.Z);
 	KelvinletBrushOpProperties.Size *= 0.6;
 
@@ -1189,7 +1182,7 @@ bool UDynamicMeshSculptTool::ApplyTwistKelvinBrush(const FRay& WorldRay)
 	FDynamicMesh3* Mesh = DynamicMeshComponent->GetMesh();
 	FKelvinletBrushOp KelvinBrushOp(*Mesh);
 
-	FKelvinletBrushOp::FKelvinletBrushOpProperties  KelvinletBrushOpProperties(EKelvinletBrushMode::TwistKelvinlet, *KelvinBrushProperties, *BrushProperties);
+	FKelvinletBrushOp::FKelvinletBrushOpProperties  KelvinletBrushOpProperties(EKelvinletBrushMode::TwistKelvinlet, *KelvinBrushProperties, CurrentBrushRadius, BrushProperties->BrushFalloffAmount);
 	KelvinletBrushOpProperties.Direction = UseSpeed * FVector(BrushNormalLocal.X, BrushNormalLocal.Y, BrushNormalLocal.Z); // twist about local normal
 	KelvinletBrushOpProperties.Size *= 0.35; // reduce the core size of this brush.
 	
@@ -1216,7 +1209,7 @@ bool UDynamicMeshSculptTool::ApplyScaleKelvinBrush(const FRay& WorldRay)
 	FDynamicMesh3* Mesh = DynamicMeshComponent->GetMesh();
 	FKelvinletBrushOp KelvinBrushOp(*Mesh);
 
-	FKelvinletBrushOp::FKelvinletBrushOpProperties  KelvinletBrushOpProperties(EKelvinletBrushMode::ScaleKelvinlet, *KelvinBrushProperties, *BrushProperties);
+	FKelvinletBrushOp::FKelvinletBrushOpProperties  KelvinletBrushOpProperties(EKelvinletBrushMode::ScaleKelvinlet, *KelvinBrushProperties, CurrentBrushRadius, BrushProperties->BrushFalloffAmount);
 	KelvinletBrushOpProperties.Direction = FVector(UseSpeed, 0., 0.); // it is a bit iffy, but we only use the first component for the scale
 	KelvinletBrushOpProperties.Size *= 0.35;
 
@@ -1468,7 +1461,8 @@ void UDynamicMeshSculptTool::Render(IToolsContextRenderAPI* RenderAPI)
 
 	FViewCameraState RenderCameraState = RenderAPI->GetCameraState();
 
-	BrushIndicator->Update( (float)this->CurrentBrushRadius, (FVector)this->LastBrushPosWorld, (FVector)this->LastBrushPosNormalWorld, 1.0f-BrushProperties->BrushFalloffAmount);
+	//BrushIndicator->Update( (float)this->CurrentBrushRadius, (FVector)this->LastBrushPosWorld, (FVector)this->LastBrushPosNormalWorld, 1.0f-BrushProperties->BrushFalloffAmount);
+	BrushIndicator->Update( (float)this->CurrentBrushRadius, (FVector)this->LastBrushPosWorld, (FVector)this->LastBrushPosNormalWorld, 0.5f);
 	if (BrushIndicatorMaterial)
 	{
 		double FixedDimScale = ToolSceneQueriesUtil::CalculateDimensionFromVisualAngleD(RenderCameraState, LastBrushPosWorld, 1.5f);
@@ -2202,27 +2196,23 @@ UPreviewMesh* UDynamicMeshSculptTool::MakeDefaultSphereMesh(UObject* Parent, UWo
 }
 
 void UDynamicMeshSculptTool::IncreaseBrushRadiusAction()
-{
-	BrushProperties->BrushSize = FMath::Clamp(BrushProperties->BrushSize + 0.025f, 0.0f, 1.0f);
-	CalculateBrushRadius();
+{	
+	BrushProperties->BrushSize.IncreaseRadius(false);
 }
 
 void UDynamicMeshSculptTool::DecreaseBrushRadiusAction()
 {
-	BrushProperties->BrushSize = FMath::Clamp(BrushProperties->BrushSize - 0.025f, 0.0f, 1.0f);
-	CalculateBrushRadius();
+	BrushProperties->BrushSize.DecreaseRadius(false);
 }
 
 void UDynamicMeshSculptTool::IncreaseBrushRadiusSmallStepAction()
 {
-	BrushProperties->BrushSize = FMath::Clamp(BrushProperties->BrushSize + 0.005f, 0.0f, 1.0f);
-	CalculateBrushRadius();
+	BrushProperties->BrushSize.IncreaseRadius(true);
 }
 
 void UDynamicMeshSculptTool::DecreaseBrushRadiusSmallStepAction()
 {
-	BrushProperties->BrushSize = FMath::Clamp(BrushProperties->BrushSize - 0.005f, 0.0f, 1.0f);
-	CalculateBrushRadius();
+	BrushProperties->BrushSize.DecreaseRadius(true);
 }
 
 void UDynamicMeshSculptTool::IncreaseBrushSpeedAction()
@@ -2235,26 +2225,6 @@ void UDynamicMeshSculptTool::DecreaseBrushSpeedAction()
 	SculptProperties->PrimaryBrushSpeed = FMath::Clamp(SculptProperties->PrimaryBrushSpeed - 0.05f, 0.0f, 1.0f);
 }
 
-void UDynamicMeshSculptTool::IncreaseBrushFalloffAction()
-{
-	const float ChangeAmount = 0.1f;
-	const float OldValue = BrushProperties->BrushFalloffAmount;
-
-	float NewValue = OldValue + ChangeAmount;
-	BrushProperties->BrushFalloffAmount = FMath::Clamp(NewValue, 0.f, 1.f);
-}
-
-void UDynamicMeshSculptTool::DecreaseBrushFalloffAction()
-{
-	const float ChangeAmount = 0.1f;
-	const float OldValue = BrushProperties->BrushFalloffAmount;
-
-	float NewValue = OldValue - ChangeAmount;
-	BrushProperties->BrushFalloffAmount = FMath::Clamp(NewValue, 0.f, 1.f);
-}
-
-
-
 
 void UDynamicMeshSculptTool::NextHistoryBrushModeAction()
 {
@@ -2293,20 +2263,6 @@ void UDynamicMeshSculptTool::RegisterActions(FInteractiveToolActionSet& ActionSe
 		EModifierKey::None, EKeys::LeftBracket,
 		[this]() { DecreaseBrushRadiusAction(); } );
 
-	ActionSet.RegisterAction(this, (int32)EStandardToolActions::BaseClientDefinedActionID + 12,
-		TEXT("BrushIncreaseFalloff"),
-		LOCTEXT("BrushIncreaseFalloff", "Increase Brush Falloff"),
-		LOCTEXT("BrushIncreaseFalloffTooltip", "Press this key to increase brush falloff by a fixed increment."),
-		EModifierKey::Shift | EModifierKey::Control, EKeys::RightBracket,
-		[this]() { IncreaseBrushFalloffAction(); });
-
-	ActionSet.RegisterAction(this, (int32)EStandardToolActions::BaseClientDefinedActionID + 13,
-		TEXT("BrushDecreaseFalloff"),
-		LOCTEXT("BrushDecreaseFalloff", "Decrease Brush Falloff"),
-		LOCTEXT("BrushDecreaseFalloffTooltip", "Press this key to decrease brush falloff by a fixed increment."),
-		EModifierKey::Shift | EModifierKey::Control, EKeys::LeftBracket,
-		[this]() { DecreaseBrushFalloffAction(); });
-
 
 	//ActionSet.RegisterAction(this, (int32)EStandardToolActions::BaseClientDefinedActionID + 10,
 	//	TEXT("NextBrushHistoryState"),
diff --git a/Engine/Plugins/Experimental/MeshModelingToolsetExp/Source/MeshModelingToolsExp/Private/MeshGroupPaintTool.cpp b/Engine/Plugins/Experimental/MeshModelingToolsetExp/Source/MeshModelingToolsExp/Private/MeshGroupPaintTool.cpp
--- a/Engine/Plugins/Experimental/MeshModelingToolsetExp/Source/MeshModelingToolsExp/Private/MeshGroupPaintTool.cpp
+++ b/Engine/Plugins/Experimental/MeshModelingToolsetExp/Source/MeshModelingToolsExp/Private/MeshGroupPaintTool.cpp
@@ -157,11 +157,11 @@ void UMeshGroupPaintTool::Setup()
 	FilterProperties->WatchProperty(FilterProperties->SubToolType,
 		[this](EMeshGroupPaintInteractionType NewType) { UpdateSubToolType(NewType); });
 	FilterProperties->WatchProperty(FilterProperties->BrushSize,
-		[this](float NewSize) { UMeshSculptToolBase::BrushProperties->BrushSize = NewSize; });
+		[this](float NewSize) { UMeshSculptToolBase::BrushProperties->BrushSize.AdaptiveSize = NewSize; });
 	FilterProperties->WatchProperty(FilterProperties->bHitBackFaces,
 		[this](bool bNewValue) { UMeshSculptToolBase::BrushProperties->bHitBackFaces = bNewValue; });
 	FilterProperties->RestoreProperties(this);
-	FilterProperties->BrushSize = UMeshSculptToolBase::BrushProperties->BrushSize;
+	FilterProperties->BrushSize = UMeshSculptToolBase::BrushProperties->BrushSize.AdaptiveSize;
 	FilterProperties->bHitBackFaces = UMeshSculptToolBase::BrushProperties->bHitBackFaces;
 	AddToolPropertySource(FilterProperties);
 
@@ -175,7 +175,7 @@ void UMeshGroupPaintTool::Setup()
 	CalculateBrushRadius();
 
 	PaintBrushOpProperties = NewObject<UGroupPaintBrushOpProps>(this);
-	RegisterBrushType((int32)EMeshGroupPaintBrushType::Paint,
+	RegisterBrushType((int32)EMeshGroupPaintBrushType::Paint, LOCTEXT("Paint", "Paint"),
 		MakeUnique<FLambdaMeshSculptBrushOpFactory>([this]() { return MakeUnique<FGroupPaintBrushOp>(); }),
 		PaintBrushOpProperties);
 
@@ -183,7 +183,7 @@ void UMeshGroupPaintTool::Setup()
 	EraseBrushOpProperties = NewObject<UGroupEraseBrushOpProps>(this);
 	EraseBrushOpProperties->GetCurrentGroupLambda = [this]() { return PaintBrushOpProperties->GetGroup(); };
 
-	RegisterSecondaryBrushType((int32)EMeshGroupPaintBrushType::Erase,
+	RegisterSecondaryBrushType((int32)EMeshGroupPaintBrushType::Erase, LOCTEXT("Erase", "Erase"),
 		MakeUnique<TBasicMeshSculptBrushOpFactory<FGroupEraseBrushOp>>(),
 		EraseBrushOpProperties);
 
diff --git a/Engine/Plugins/Experimental/MeshModelingToolsetExp/Source/MeshModelingToolsExp/Private/MeshVertexSculptTool.cpp b/Engine/Plugins/Experimental/MeshModelingToolsetExp/Source/MeshModelingToolsExp/Private/MeshVertexSculptTool.cpp
--- a/Engine/Plugins/Experimental/MeshModelingToolsetExp/Source/MeshModelingToolsExp/Private/MeshVertexSculptTool.cpp
+++ b/Engine/Plugins/Experimental/MeshModelingToolsetExp/Source/MeshModelingToolsExp/Private/MeshVertexSculptTool.cpp
@@ -153,6 +153,7 @@ void UMeshVertexSculptTool::Setup()
 
 	// initialize other properties
 	SculptProperties = NewObject<UVertexBrushSculptProperties>(this);
+	SculptProperties->Tool = this;
 	
 	// init state flags flags
 	ActiveVertexChange = nullptr;
@@ -176,75 +177,77 @@ void UMeshVertexSculptTool::Setup()
 		return GetBaseMeshNearest(VertexID, Position, MaxDist, PosOut, NormalOut);
 	};
 
-	RegisterBrushType((int32)EMeshVertexSculptBrushType::Smooth,
+	RegisterBrushType((int32)EMeshVertexSculptBrushType::Smooth, LOCTEXT("SmoothBrush", "Smooth"),
 		MakeUnique<TBasicMeshSculptBrushOpFactory<FSmoothBrushOp>>(),
 		NewObject<USmoothBrushOpProps>(this));
 
-	RegisterBrushType((int32)EMeshVertexSculptBrushType::SmoothFill,
+	RegisterBrushType((int32)EMeshVertexSculptBrushType::SmoothFill, LOCTEXT("SmoothFill", "SmoothFill"),
 		MakeUnique<TBasicMeshSculptBrushOpFactory<FSmoothFillBrushOp>>(),
 		NewObject<USmoothFillBrushOpProps>(this));
 
-	RegisterBrushType((int32)EMeshVertexSculptBrushType::Move,
+	RegisterBrushType((int32)EMeshVertexSculptBrushType::Move, LOCTEXT("Move", "Move"),
 		MakeUnique<TBasicMeshSculptBrushOpFactory<FMoveBrushOp>>(),
 		NewObject<UMoveBrushOpProps>(this));
 
-	RegisterBrushType((int32)EMeshVertexSculptBrushType::Offset,
+	RegisterBrushType((int32)EMeshVertexSculptBrushType::Offset, LOCTEXT("Offset", "SculptN"),
 		MakeUnique<FLambdaMeshSculptBrushOpFactory>([this]() { return MakeUnique<FSurfaceSculptBrushOp>(BaseMeshQueryFunc); }),
 		NewObject<UStandardSculptBrushOpProps>(this));
 
-	RegisterBrushType((int32)EMeshVertexSculptBrushType::SculptView,
+	RegisterBrushType((int32)EMeshVertexSculptBrushType::SculptView, LOCTEXT("SculptView", "SculptV"),
 		MakeUnique<FLambdaMeshSculptBrushOpFactory>( [this]() { return MakeUnique<FViewAlignedSculptBrushOp>(BaseMeshQueryFunc); } ),
 		NewObject<UViewAlignedSculptBrushOpProps>(this));
 
-	RegisterBrushType((int32)EMeshVertexSculptBrushType::SculptMax,
+	RegisterBrushType((int32)EMeshVertexSculptBrushType::SculptMax, LOCTEXT("SculptMax", "SculptMx"),
 		MakeUnique<FLambdaMeshSculptBrushOpFactory>([this]() { return MakeUnique<FSurfaceMaxSculptBrushOp>(BaseMeshQueryFunc); }),
 		NewObject<USculptMaxBrushOpProps>(this));
 
-	RegisterBrushType((int32)EMeshVertexSculptBrushType::Inflate,
+	RegisterBrushType((int32)EMeshVertexSculptBrushType::Inflate, LOCTEXT("Inflate", "Inflate"),
 		MakeUnique<TBasicMeshSculptBrushOpFactory<FInflateBrushOp>>(),
 		NewObject<UInflateBrushOpProps>(this));
 
-	RegisterBrushType((int32)EMeshVertexSculptBrushType::Pinch,
+	RegisterBrushType((int32)EMeshVertexSculptBrushType::Pinch, LOCTEXT("Pinch", "Pinch"),
 		MakeUnique<TBasicMeshSculptBrushOpFactory<FPinchBrushOp>>(),
 		NewObject<UPinchBrushOpProps>(this));
 
-	RegisterBrushType((int32)EMeshVertexSculptBrushType::Flatten,
+	RegisterBrushType((int32)EMeshVertexSculptBrushType::Flatten, LOCTEXT("Flatten", "Flatten"),
 		MakeUnique<TBasicMeshSculptBrushOpFactory<FFlattenBrushOp>>(),
 		NewObject<UFlattenBrushOpProps>(this));
 
-	RegisterBrushType((int32)EMeshVertexSculptBrushType::Plane,
+	RegisterBrushType((int32)EMeshVertexSculptBrushType::Plane, LOCTEXT("Plane", "PlaneN"),
 		MakeUnique<TBasicMeshSculptBrushOpFactory<FPlaneBrushOp>>(),
 		NewObject<UPlaneBrushOpProps>(this));
 
-	RegisterBrushType((int32)EMeshVertexSculptBrushType::PlaneViewAligned,
+	RegisterBrushType((int32)EMeshVertexSculptBrushType::PlaneViewAligned, LOCTEXT("PlaneViewAligned", "PlaneV"),
 		MakeUnique<TBasicMeshSculptBrushOpFactory<FPlaneBrushOp>>(),
 		NewObject<UViewAlignedPlaneBrushOpProps>(this));
 
-	RegisterBrushType((int32)EMeshVertexSculptBrushType::FixedPlane,
+	RegisterBrushType((int32)EMeshVertexSculptBrushType::FixedPlane, LOCTEXT("FixedPlane", "PlaneW"),
 		MakeUnique<TBasicMeshSculptBrushOpFactory<FPlaneBrushOp>>(),
 		NewObject<UFixedPlaneBrushOpProps>(this));
 
-	RegisterBrushType((int32)EMeshVertexSculptBrushType::ScaleKelvin ,
+	RegisterBrushType((int32)EMeshVertexSculptBrushType::ScaleKelvin , LOCTEXT("ScaleKelvin", "Scale"),
 		MakeUnique<TBasicMeshSculptBrushOpFactory<FScaleKelvinletBrushOp>>(),
 		NewObject<UScaleKelvinletBrushOpProps>(this));
 
-	RegisterBrushType((int32)EMeshVertexSculptBrushType::PullKelvin,
+	RegisterBrushType((int32)EMeshVertexSculptBrushType::PullKelvin, LOCTEXT("PullKelvin", "Grab"),
 		MakeUnique<TBasicMeshSculptBrushOpFactory<FPullKelvinletBrushOp>>(),
 		NewObject<UPullKelvinletBrushOpProps>(this));
 
-	RegisterBrushType((int32)EMeshVertexSculptBrushType::PullSharpKelvin,
+	RegisterBrushType((int32)EMeshVertexSculptBrushType::PullSharpKelvin, LOCTEXT("PullSharpKelvin", "GrabSharp"),
 		MakeUnique<TBasicMeshSculptBrushOpFactory<FSharpPullKelvinletBrushOp>>(),
 		NewObject<USharpPullKelvinletBrushOpProps>(this));
 
-	RegisterBrushType((int32)EMeshVertexSculptBrushType::TwistKelvin,
+	RegisterBrushType((int32)EMeshVertexSculptBrushType::TwistKelvin, LOCTEXT("TwistKelvin", "Twist"),
 		MakeUnique<TBasicMeshSculptBrushOpFactory<FTwistKelvinletBrushOp>>(),
 		NewObject<UTwistKelvinletBrushOpProps>(this));
 
 	// secondary brushes
-	RegisterSecondaryBrushType((int32)EMeshVertexSculptBrushType::Smooth,
+	RegisterSecondaryBrushType((int32)EMeshVertexSculptBrushType::Smooth, LOCTEXT("Smooth", "Smooth"),
 		MakeUnique<TBasicMeshSculptBrushOpFactory<FSmoothBrushOp>>(),
 		NewObject<USecondarySmoothBrushOpProps>(this));
 
+	// falloffs
+	RegisterStandardFalloffTypes();
 
 	AddToolPropertySource(UMeshSculptToolBase::GizmoProperties);
 	SetToolPropertySourceEnabled(UMeshSculptToolBase::GizmoProperties, false);
@@ -292,6 +295,38 @@ void UMeshVertexSculptTool::OnPropertyModified(UObject* PropertySet, FProperty*
 	CalculateBrushRadius();
 }
 
+void UMeshVertexSculptTool::SetActiveBrushType(int32 Identifier)
+{
+	EMeshVertexSculptBrushType NewBrushType =  static_cast<EMeshVertexSculptBrushType>(Identifier);
+	if (SculptProperties->PrimaryBrushType != NewBrushType)
+	{
+		SculptProperties->PrimaryBrushType = NewBrushType;
+		UpdateBrushType(SculptProperties->PrimaryBrushType);
+		SculptProperties->SilentUpdateWatched();
+	}
+
+	// this forces full rebuild of properties panel (!!)
+	//this->NotifyOfPropertyChangeByTool(SculptProperties);
+}
+
+void UMeshVertexSculptTool::SetActiveFalloffType(int32 Identifier)
+{
+	EMeshSculptFalloffType NewFalloffType = static_cast<EMeshSculptFalloffType>(Identifier);
+	if (SculptProperties->PrimaryFalloffType != NewFalloffType)
+	{
+		SculptProperties->PrimaryFalloffType = NewFalloffType;
+		SetPrimaryFalloffType(SculptProperties->PrimaryFalloffType);
+		SculptProperties->SilentUpdateWatched();
+	}
+
+	// this forces full rebuild of properties panel (!!)
+	//this->NotifyOfPropertyChangeByTool(SculptProperties);
+}
+
+void UMeshVertexSculptTool::SetRegionFilterType(int32 Identifier)
+{
+	SculptProperties->BrushFilter = static_cast<EMeshVertexSculptBrushFilterType>(Identifier);
+}
 
 
 void UMeshVertexSculptTool::OnBeginStroke(const FRay& WorldRay)
diff --git a/Engine/Plugins/Experimental/MeshModelingToolsetExp/Source/MeshModelingToolsExp/Private/Sculpting/KelvinletBrushOp.h b/Engine/Plugins/Experimental/MeshModelingToolsetExp/Source/MeshModelingToolsExp/Private/Sculpting/KelvinletBrushOp.h
--- a/Engine/Plugins/Experimental/MeshModelingToolsetExp/Source/MeshModelingToolsExp/Private/Sculpting/KelvinletBrushOp.h
+++ b/Engine/Plugins/Experimental/MeshModelingToolsetExp/Source/MeshModelingToolsExp/Private/Sculpting/KelvinletBrushOp.h
@@ -5,8 +5,8 @@
 #include "MatrixTypes.h"
 #include "DynamicMesh/DynamicMesh3.h"
 #include "Deformers/Kelvinlets.h"
-#include "Sculpting/MeshSculptToolBase.h"
 #include "Sculpting/MeshBrushOpBase.h"
+#include "Sculpting/MeshSculptToolBase.h"
 #include "Async/ParallelFor.h"
 #include "KelvinletBrushOp.generated.h"
 
@@ -487,15 +487,15 @@ public:
 
 	struct FKelvinletBrushOpProperties
 	{
-		FKelvinletBrushOpProperties(const EKelvinletBrushMode& BrushMode, const UKelvinBrushProperties& Properties, const UBrushBaseProperties& Brush) 
+		FKelvinletBrushOpProperties(const EKelvinletBrushMode& BrushMode, const UKelvinBrushProperties& Properties, double BrushRadius, double BrushFalloff) 
 			: Mode(BrushMode)
 			, Direction(1.0, 0.0, 0.0)
 		{
 			Speed = 0.;
-			FallOff = Brush.BrushFalloffAmount;
+			FallOff = BrushFalloff;
 			Mu = FMath::Max(Properties.Stiffness, 0.f);
 			Nu = FMath::Clamp(0.5f * (1.f - 2.f * Properties.Incompressiblity), 0.f, 0.5f);
-			Size = FMath::Max(Brush.BrushRadius * Properties.FallOffDistance, 0.f);
+			Size = FMath::Max(BrushRadius * Properties.FallOffDistance, 0.f);
 			NumSteps = Properties.BrushSteps;
 		}
 
diff --git a/Engine/Plugins/Experimental/MeshModelingToolsetExp/Source/MeshModelingToolsExp/Private/Sculpting/MeshSculptToolBase.cpp b/Engine/Plugins/Experimental/MeshModelingToolsetExp/Source/MeshModelingToolsExp/Private/Sculpting/MeshSculptToolBase.cpp
--- a/Engine/Plugins/Experimental/MeshModelingToolsetExp/Source/MeshModelingToolsExp/Private/Sculpting/MeshSculptToolBase.cpp
+++ b/Engine/Plugins/Experimental/MeshModelingToolsetExp/Source/MeshModelingToolsExp/Private/Sculpting/MeshSculptToolBase.cpp
@@ -29,6 +29,62 @@ namespace
 }
 
 
+void FBrushToolRadius::InitializeWorldSizeRange(TInterval<float> Range, bool bValidateWorldRadius)
+{
+	WorldSizeRange = Range;
+	if (WorldRadius < WorldSizeRange.Min)
+	{
+		WorldRadius = WorldSizeRange.Interpolate(0.2);
+	}
+	else if (WorldRadius > WorldSizeRange.Max)
+	{
+		WorldRadius = WorldSizeRange.Interpolate(0.8);
+	}
+}
+
+float FBrushToolRadius::GetWorldRadius() const
+{
+	if (SizeType == EBrushToolSizeType::Adaptive)
+	{
+		return 0.5 * WorldSizeRange.Interpolate( FMath::Max(0, AdaptiveSize) );
+	}
+	else
+	{
+		return WorldRadius;
+	}
+}
+
+void FBrushToolRadius::IncreaseRadius(bool bSmallStep)
+{
+	float StepSize = (bSmallStep) ? 0.025f : 0.005f;
+	if (SizeType == EBrushToolSizeType::Adaptive)
+	{
+		AdaptiveSize = FMath::Clamp(AdaptiveSize + StepSize, 0.0f, 1.0f);
+	}
+	else
+	{
+		float dt = StepSize * 0.5 * WorldSizeRange.Size();
+		WorldRadius = FMath::Clamp(WorldRadius + dt, WorldSizeRange.Min, WorldSizeRange.Max);
+	}
+}
+
+void FBrushToolRadius::DecreaseRadius(bool bSmallStep)
+{
+	float StepSize = (bSmallStep) ? 0.025f : 0.005f;
+	if (SizeType == EBrushToolSizeType::Adaptive)
+	{
+		AdaptiveSize = FMath::Clamp(AdaptiveSize - StepSize, 0.0f, 1.0f);
+	}
+	else
+	{
+		float dt = StepSize * 0.5 * WorldSizeRange.Size();
+		WorldRadius = FMath::Clamp(WorldRadius - dt, WorldSizeRange.Min, WorldSizeRange.Max);
+	}
+}
+
+
+
+
 void UMeshSculptToolBase::SetWorld(UWorld* World)
 {
 	this->TargetWorld = World;
@@ -41,7 +97,6 @@ void UMeshSculptToolBase::Setup()
 
 	BrushProperties = NewObject<USculptBrushProperties>(this);
 	BrushProperties->RestoreProperties(this);
-	BrushProperties->bShowStrength = false;
 	// Note that brush properties includes BrushRadius, which, when not used as a constant,
 	// serves as an output property based on target size and brush size, and so it would need
 	// updating after the RestoreProperties() call. But deriving classes will call 
@@ -266,9 +321,13 @@ void UMeshSculptToolBase::InitializeSculptMeshComponent(UBaseDynamicMeshComponen
 
 
 
-
-void UMeshSculptToolBase::RegisterBrushType(int32 Identifier, TUniquePtr<FMeshSculptBrushOpFactory> Factory, UMeshSculptBrushOpProps* PropSet)
+void UMeshSculptToolBase::RegisterBrushType(int32 Identifier, FText Name, TUniquePtr<FMeshSculptBrushOpFactory> Factory, UMeshSculptBrushOpProps* PropSet)
 {
+	FBrushTypeInfo TypeInfo;
+	TypeInfo.Name = Name;
+	TypeInfo.Identifier = Identifier;
+	RegisteredPrimaryBrushTypes.Add(TypeInfo);
+
 	check(BrushOpPropSets.Contains(Identifier) == false && BrushOpFactories.Contains(Identifier) == false);
 	BrushOpPropSets.Add(Identifier, PropSet);
 	BrushOpFactories.Add(Identifier, MoveTemp(Factory));
@@ -277,8 +336,13 @@ void UMeshSculptToolBase::RegisterBrushType(int32 Identifier, TUniquePtr<FMeshSc
 	SetToolPropertySourceEnabled(PropSet, false);
 }
 
-void UMeshSculptToolBase::RegisterSecondaryBrushType(int32 Identifier, TUniquePtr<FMeshSculptBrushOpFactory> Factory, UMeshSculptBrushOpProps* PropSet)
+void UMeshSculptToolBase::RegisterSecondaryBrushType(int32 Identifier, FText Name, TUniquePtr<FMeshSculptBrushOpFactory> Factory, UMeshSculptBrushOpProps* PropSet)
 {
+	FBrushTypeInfo TypeInfo;
+	TypeInfo.Name = Name;
+	TypeInfo.Identifier = Identifier;
+	RegisteredSecondaryBrushTypes.Add(TypeInfo);
+
 	check(SecondaryBrushOpPropSets.Contains(Identifier) == false && SecondaryBrushOpFactories.Contains(Identifier) == false);
 	SecondaryBrushOpPropSets.Add(Identifier, PropSet);
 	SecondaryBrushOpFactories.Add(Identifier, MoveTemp(Factory));
@@ -288,6 +352,7 @@ void UMeshSculptToolBase::RegisterSecondaryBrushType(int32 Identifier, TUniquePt
 }
 
 
+
 void UMeshSculptToolBase::SaveAllBrushTypeProperties(UInteractiveTool* SaveFromTool)
 {
 	for (auto Pair : BrushOpPropSets)
@@ -400,6 +465,18 @@ void UMeshSculptToolBase::SetBrushOpPropsVisibility(bool bVisible)
 }
 
 
+void UMeshSculptToolBase::RegisterStandardFalloffTypes()
+{
+	RegisteredPrimaryFalloffTypes.Add( FFalloffTypeInfo{ LOCTEXT("Smooth", "Smooth"), TEXT("Smooth"), (int32)EMeshSculptFalloffType::Smooth});
+	RegisteredPrimaryFalloffTypes.Add( FFalloffTypeInfo{ LOCTEXT("Linear", "Linear"), TEXT("Linear"), (int32)EMeshSculptFalloffType::Linear } );
+	RegisteredPrimaryFalloffTypes.Add( FFalloffTypeInfo{ LOCTEXT("Inverse", "Inverse"), TEXT("Inverse"), (int32)EMeshSculptFalloffType::Inverse } );
+	RegisteredPrimaryFalloffTypes.Add( FFalloffTypeInfo{ LOCTEXT("Round", "Round"), TEXT("Round"), (int32)EMeshSculptFalloffType::Round } );
+	RegisteredPrimaryFalloffTypes.Add( FFalloffTypeInfo{ LOCTEXT("BoxSmooth", "BoxSmooth"), TEXT("BoxSmooth"), (int32)EMeshSculptFalloffType::BoxSmooth } );
+	RegisteredPrimaryFalloffTypes.Add( FFalloffTypeInfo{ LOCTEXT("BoxLinear", "BoxLinear"), TEXT("BoxLinear"), (int32)EMeshSculptFalloffType::BoxLinear } );
+	RegisteredPrimaryFalloffTypes.Add( FFalloffTypeInfo{ LOCTEXT("BoxInverse", "BoxInverse"), TEXT("BoxInverse"), (int32)EMeshSculptFalloffType::BoxInverse } );
+	RegisteredPrimaryFalloffTypes.Add( FFalloffTypeInfo{ LOCTEXT("BoxRound", "BoxRound"), TEXT("BoxRound"), (int32)EMeshSculptFalloffType::BoxRound } );
+}
+
 void UMeshSculptToolBase::SetPrimaryFalloffType(EMeshSculptFalloffType FalloffType)
 {
 	PrimaryFalloff = MakeShared<FMeshSculptFallofFunc>();
@@ -541,7 +618,7 @@ retry_frame_update:
 
 	if (InStroke() && BrushProperties->Lazyness > 0)
 	{
-		double t = FMathd::Lerp(1.0, 0.1, BrushProperties->Lazyness);
+		double t = FMathd::Lerp(1.0, 0.1, (double)BrushProperties->Lazyness);
 		LastBrushFrameWorld.Origin = UE::Geometry::Lerp(LastBrushFrameWorld.Origin, NewFrame.Origin, t);
 		LastBrushFrameWorld.Rotation = FQuaterniond(LastBrushFrameWorld.Rotation, NewFrame.Rotation, t);
 	}
@@ -810,21 +887,15 @@ void UMeshSculptToolBase::InitializeBrushSizeRange(const FAxisAlignedBox3d& Targ
 {
 	double MaxDimension = TargetBounds.MaxDim();
 	BrushRelativeSizeRange = FInterval1d(MaxDimension * 0.01, MaxDimension);
+	BrushProperties->BrushSize.InitializeWorldSizeRange(
+		TInterval<float>((float)BrushRelativeSizeRange.Min, (float)BrushRelativeSizeRange.Max));
 	CalculateBrushRadius();
 }
 
 
 void UMeshSculptToolBase::CalculateBrushRadius()
 {
-	CurrentBrushRadius = 0.5 * BrushRelativeSizeRange.Interpolate(BrushProperties->BrushSize);
-	if (BrushProperties->bSpecifyRadius)
-	{
-		CurrentBrushRadius = BrushProperties->BrushRadius;
-	}
-	else
-	{
-		BrushProperties->BrushRadius = CurrentBrushRadius;
-	}
+	CurrentBrushRadius = BrushProperties->BrushSize.GetWorldRadius();
 }
 
 
@@ -855,26 +926,22 @@ double UMeshSculptToolBase::GetCurrentBrushDepth()
 
 void UMeshSculptToolBase::IncreaseBrushRadiusAction()
 {
-	BrushProperties->BrushSize = FMath::Clamp(BrushProperties->BrushSize + 0.025f, 0.0f, 1.0f);
-	CalculateBrushRadius();
+	BrushProperties->BrushSize.IncreaseRadius(false);
 }
 
 void UMeshSculptToolBase::DecreaseBrushRadiusAction()
 {
-	BrushProperties->BrushSize = FMath::Clamp(BrushProperties->BrushSize - 0.025f, 0.0f, 1.0f);
-	CalculateBrushRadius();
+	BrushProperties->BrushSize.DecreaseRadius(false);
 }
 
 void UMeshSculptToolBase::IncreaseBrushRadiusSmallStepAction()
 {
-	BrushProperties->BrushSize = FMath::Clamp(BrushProperties->BrushSize + 0.005f, 0.0f, 1.0f);
-	CalculateBrushRadius();
+	BrushProperties->BrushSize.IncreaseRadius(true);
 }
 
 void UMeshSculptToolBase::DecreaseBrushRadiusSmallStepAction()
 {
-	BrushProperties->BrushSize = FMath::Clamp(BrushProperties->BrushSize - 0.005f, 0.0f, 1.0f);
-	CalculateBrushRadius();
+	BrushProperties->BrushSize.DecreaseRadius(true);
 }
 
 
diff --git a/Engine/Plugins/Experimental/MeshModelingToolsetExp/Source/MeshModelingToolsExp/Public/DynamicMeshSculptTool.h b/Engine/Plugins/Experimental/MeshModelingToolsetExp/Source/MeshModelingToolsExp/Public/DynamicMeshSculptTool.h
--- a/Engine/Plugins/Experimental/MeshModelingToolsetExp/Source/MeshModelingToolsExp/Public/DynamicMeshSculptTool.h
+++ b/Engine/Plugins/Experimental/MeshModelingToolsetExp/Source/MeshModelingToolsExp/Public/DynamicMeshSculptTool.h
@@ -288,8 +288,6 @@ public:
 	virtual void DecreaseBrushRadiusAction();
 	virtual void IncreaseBrushRadiusSmallStepAction();
 	virtual void DecreaseBrushRadiusSmallStepAction();
-	virtual void IncreaseBrushFalloffAction();
-	virtual void DecreaseBrushFalloffAction();
 
 	virtual void IncreaseBrushSpeedAction();
 	virtual void DecreaseBrushSpeedAction();
diff --git a/Engine/Plugins/Experimental/MeshModelingToolsetExp/Source/MeshModelingToolsExp/Public/MeshVertexSculptTool.h b/Engine/Plugins/Experimental/MeshModelingToolsetExp/Source/MeshModelingToolsExp/Public/MeshVertexSculptTool.h
--- a/Engine/Plugins/Experimental/MeshModelingToolsetExp/Source/MeshModelingToolsExp/Public/MeshVertexSculptTool.h
+++ b/Engine/Plugins/Experimental/MeshModelingToolsetExp/Source/MeshModelingToolsExp/Public/MeshVertexSculptTool.h
@@ -104,11 +104,11 @@ UENUM()
 enum class EMeshVertexSculptBrushFilterType : uint8
 {
 	/** Do not filter brush area */
-	None,
+	None = 0,
 	/** Only apply brush to triangles in the same connected mesh component/island */
-	Component,
+	Component = 1,
 	/** Only apply brush to triangles with the same PolyGroup */
-	PolyGroup
+	PolyGroup = 2
 };
 
 
@@ -120,25 +120,24 @@ class MESHMODELINGTOOLSEXP_API UVertexBrushSculptProperties : public UInteractiv
 
 public:
 	/** Primary Brush Mode */
-	UPROPERTY(EditAnywhere, Category = Sculpting, meta = (DisplayName = "Brush Type"))
+	UPROPERTY(EditAnywhere, Category = Sculpting, meta = (DisplayName = "Brush"))
 	EMeshVertexSculptBrushType PrimaryBrushType = EMeshVertexSculptBrushType::Offset;
 
 	/** Primary Brush Falloff Type, multiplied by Alpha Mask where applicable */
-	UPROPERTY(EditAnywhere, Category = Sculpting, meta = (DisplayName = "Falloff Shape"))
+	UPROPERTY(EditAnywhere, Category = Sculpting, meta = (DisplayName = "Falloff"))
 	EMeshSculptFalloffType PrimaryFalloffType = EMeshSculptFalloffType::Smooth;
 
 	/** Filter applied to Stamp Region Triangles, based on first Stroke Stamp */
-	UPROPERTY(EditAnywhere, Category = Sculpting, meta = (DisplayName = "Region Filter"))
+	UPROPERTY(EditAnywhere, Category = Sculpting, meta = (DisplayName = "Region"))
 	EMeshVertexSculptBrushFilterType BrushFilter = EMeshVertexSculptBrushFilterType::None;
 
 	/** When Freeze Target is toggled on, the Brush Target Surface will be Frozen in its current state, until toggled off. Brush strokes will be applied relative to the Target Surface, for applicable Brushes */
 	UPROPERTY(EditAnywhere, Category = Sculpting, meta = (EditCondition = "PrimaryBrushType == EMeshVertexSculptBrushType::Offset || PrimaryBrushType == EMeshVertexSculptBrushType::SculptMax || PrimaryBrushType == EMeshVertexSculptBrushType::SculptView || PrimaryBrushType == EMeshVertexSculptBrushType::Pinch || PrimaryBrushType == EMeshVertexSculptBrushType::Resample" ))
 	bool bFreezeTarget = false;
 
-	/** When enabled, instead of Mesh Smoothing, the Shift-Smooth modifier will "erase" the displacement relative to the Brush Target Surface */
-	//UPROPERTY(EditAnywhere, Category = Sculpting, meta = (DisplayName = "Shift-Smooth Erases", EditCondition = "bFreezeTarget == true"))
-	UPROPERTY()
-	bool bSmoothErases = false;
+	// parent ref required for details customization
+	UPROPERTY(meta = (TransientToolProperty))
+	TWeakObjectPtr<UMeshVertexSculptTool> Tool;
 };
 
 
@@ -211,6 +210,10 @@ public:
 
 	virtual void UpdateBrushAlpha(UTexture2D* NewAlpha);
 
+	virtual void SetActiveBrushType(int32 Identifier);
+	virtual void SetActiveFalloffType(int32 Identifier);
+	virtual void SetRegionFilterType(int32 Identifier);
+
 protected:
 	// UMeshSculptToolBase API
 	virtual UBaseDynamicMeshComponent* GetSculptMeshComponent() { return DynamicMeshComponent; }
diff --git a/Engine/Plugins/Experimental/MeshModelingToolsetExp/Source/MeshModelingToolsExp/Public/Sculpting/MeshMoveBrushOps.h b/Engine/Plugins/Experimental/MeshModelingToolsetExp/Source/MeshModelingToolsExp/Public/Sculpting/MeshMoveBrushOps.h
--- a/Engine/Plugins/Experimental/MeshModelingToolsetExp/Source/MeshModelingToolsExp/Public/Sculpting/MeshMoveBrushOps.h
+++ b/Engine/Plugins/Experimental/MeshModelingToolsetExp/Source/MeshModelingToolsExp/Public/Sculpting/MeshMoveBrushOps.h
@@ -35,7 +35,7 @@ public:
 	virtual float GetDepth() override { return Depth; }
 
 	/** Axis filters restrict mesh movement to World X/Y/Z axes */
-	UPROPERTY(EditAnywhere, Category = MoveBrush, AdvancedDisplay)
+	UPROPERTY(EditAnywhere, Category = MoveBrush)
 	FModelingToolsAxisFilter AxisFilters;
 };
 
diff --git a/Engine/Plugins/Experimental/MeshModelingToolsetExp/Source/MeshModelingToolsExp/Public/Sculpting/MeshSculptToolBase.h b/Engine/Plugins/Experimental/MeshModelingToolsetExp/Source/MeshModelingToolsExp/Public/Sculpting/MeshSculptToolBase.h
--- a/Engine/Plugins/Experimental/MeshModelingToolsetExp/Source/MeshModelingToolsExp/Public/Sculpting/MeshSculptToolBase.h
+++ b/Engine/Plugins/Experimental/MeshModelingToolsetExp/Source/MeshModelingToolsExp/Public/Sculpting/MeshSculptToolBase.h
@@ -18,19 +18,73 @@ class UCombinedTransformGizmo;
 class UTransformProxy;
 class UPreviewMesh;
 
+/**
+ * Type of Brush Size currently active in FBrushToolRadius
+ */
+UENUM()
+enum class EBrushToolSizeType : uint8
+{
+	/** Brush size is a dimensionless scale relative to the target object size */
+	Adaptive = 0,
+	/** Brush size is defined in world dimensions */
+	World = 1
+};
+
+/**
+ * FBrushToolRadius is used to define the size of 3D "brushes" used in (eg) sculpting tools.
+ * The brush size can be defined in various ways.
+ */
+USTRUCT()
+struct FBrushToolRadius
+{
+	GENERATED_BODY()
+
+	/** Specify the type of brush size currently in use */
+	UPROPERTY(EditAnywhere, Category = Brush)
+	EBrushToolSizeType SizeType = EBrushToolSizeType::Adaptive;
+
+	/** Adaptive brush size is used to interpolate between an object-specific minimum and maximum brush size */
+	UPROPERTY(EditAnywhere, Category = Brush, AdvancedDisplay, meta = (EditCondition = "SizeType == EBrushToolSizeType::Adaptive",
+		DisplayName = "Size", UIMin = "0.0", UIMax = "1.0", ClampMin = "0.0", ClampMax = "10.0"))
+	float AdaptiveSize = 0.25;
+
+	/** World brush size is a dimension in world coordinates */
+	UPROPERTY(EditAnywhere, Category = Brush, AdvancedDisplay, meta = (EditCondition = "SizeType == EBrushToolSizeType::World",
+		DisplayName = "World Radius", UIMin = "1.0", UIMax = "1000.0", ClampMin = "0.1", ClampMax = "50000.0"))
+	float WorldRadius = 100.0;
+
+	/**
+	 * WorldSizeRange defines the min/max dimensions for Adaptive brush size
+	 */
+	TInterval<float> WorldSizeRange = TInterval<float>(1.0f, 1000.0f);
+
+	//
+	// util functions
+	//
 
+	/** Set the WorldSizeRange value and optionally clamp the WorldRadius based on this new range */
+	void InitializeWorldSizeRange(TInterval<float> Range, bool bValidateWorldRadius = true);
+	/** Return the set/calculated world-space radius for the current settings */
+	float GetWorldRadius() const;
+	/** Increase the current radius dimension by a fixed step (or a smaller fixed step) */
+	void IncreaseRadius(bool bSmallStep);
+	/** Decrease the current radius dimension by a fixed step (or a smaller fixed step) */
+	void DecreaseRadius(bool bSmallStep);
+};
+
+ 
 /** Mesh Sculpting Brush Types */
 UENUM()
 enum class EMeshSculptFalloffType : uint8
 {
-	Smooth,
-	Linear,
-	Inverse,
-	Round,
-	BoxSmooth,
-	BoxLinear,
-	BoxInverse,
-	BoxRound,
+	Smooth = 0,
+	Linear = 1,
+	Inverse = 2,
+	Round = 3,
+	BoxSmooth = 4,
+	BoxLinear = 5,
+	BoxInverse = 6,
+	BoxRound = 7,
 
 	LastValue UMETA(Hidden)
 };
@@ -40,14 +94,21 @@ enum class EMeshSculptFalloffType : uint8
 
 
 UCLASS()
-class MESHMODELINGTOOLSEXP_API USculptBrushProperties : public UBrushBaseProperties
+class MESHMODELINGTOOLSEXP_API USculptBrushProperties : public UInteractiveToolPropertySet
 {
 	GENERATED_BODY()
 public:
-	USculptBrushProperties()
-	{
-		this->BrushFalloffAmount = 0.5f;
-	}
+
+	UPROPERTY(EditAnywhere, Category = Brush, meta = (DisplayPriority = 1))
+	FBrushToolRadius BrushSize;
+
+	/** Amount of falloff to apply (0.0 - 1.0) */
+	UPROPERTY(EditAnywhere, Category = Brush, meta = (DisplayName = "Falloff", UIMin = "0.0", UIMax = "1.0", ClampMin = "0.0", ClampMax = "1.0", HideEditConditionToggle, EditConditionHides, EditCondition="bShowFalloff", DisplayPriority = 3))
+	float BrushFalloffAmount;
+
+	/** If false, then BrushFalloffAmount will not be shown in DetailsView panels (otherwise no effect) */
+	UPROPERTY(meta = (TransientToolProperty))
+	bool bShowFalloff = true;
 
 	/** Depth of Brush into surface along view ray or surface normal, depending on the Active Brush Type */
 	UPROPERTY(EditAnywhere, Category = Brush, meta = (UIMin = "-0.5", UIMax = "0.5", ClampMin = "-1.0", ClampMax = "1.0", DisplayPriority = 5, HideEditConditionToggle, EditConditionHides, EditCondition = "bShowPerBrushProps"))
@@ -58,15 +119,15 @@ public:
 	bool bHitBackFaces = true;
 
 	/** Brush stamps are applied at this time interval. 0 for a single stamp, 1 for continuous stamps, 0.5 is a stamp every half-second */
-	UPROPERTY(EditAnywhere, Category = Brush, meta = (UIMin = "0.0", UIMax = "1.0", ClampMin = "0.0", ClampMax = "1.0", DisplayPriority = 7, HideEditConditionToggle, EditConditionHides, EditCondition = "bShowFlowRate"))
+	UPROPERTY(EditAnywhere, Category = Brush, meta = (DisplayName = "Flow", UIMin = "0.0", UIMax = "1.0", ClampMin = "0.0", ClampMax = "1.0", DisplayPriority = 7, HideEditConditionToggle, EditConditionHides, EditCondition = "bShowFlowRate"))
 	float FlowRate = 1.0f;
 
-	/** Minimum world-space spacing between stamps, defined along the brush path. Zero spacing means continuous stamps. */
+	/** Space out stamp centers at distances Spacing*BrushDiameter along the stroke (so Spacing of 1 means that stamps will be adjacent but non-overlapping). Zero spacing means continuous stamps. */
 	UPROPERTY(EditAnywhere, Category = Brush, meta = (UIMin = "0.0", UIMax = "4.0", ClampMin = "0.0", ClampMax = "1000.0", DisplayPriority = 8, HideEditConditionToggle, EditConditionHides, EditCondition = "bShowSpacing"))
 	float Spacing = 0.0f;
 
 	/** Lazy brush smooths out the brush path by averaging the cursor positions */
-	UPROPERTY(EditAnywhere, Category = Brush, meta = (UIMin = "0.0", UIMax = "1.0", ClampMin = "0.0", ClampMax = "1.0", DisplayPriority = 9, HideEditConditionToggle, EditConditionHides, EditCondition = "bShowLazyness"))
+	UPROPERTY(EditAnywhere, Category = Brush, meta = (DisplayName = "Lazy", UIMin = "0", UIMax = "1.0", ClampMin = "0", ClampMax = "1.0", DisplayPriority = 9, HideEditConditionToggle, EditConditionHides, EditCondition = "bShowLazyness"))
 	float Lazyness = 0;
 
 	/**  */
@@ -241,7 +302,19 @@ protected:
 	//
 	// Brush Types
 	//
+public:
+	struct FBrushTypeInfo
+	{
+		FText Name;
+		int32 Identifier;
+	};
+	const TArray<FBrushTypeInfo>& GetRegisteredPrimaryBrushTypes() const { return RegisteredPrimaryBrushTypes; }
+	const TArray<FBrushTypeInfo>& GetRegisteredSecondaryBrushTypes() const { return RegisteredSecondaryBrushTypes; }
+
 protected:
+	TArray<FBrushTypeInfo> RegisteredPrimaryBrushTypes;
+	TArray<FBrushTypeInfo> RegisteredSecondaryBrushTypes;
+
 	UPROPERTY()
 	TMap<int32, TObjectPtr<UMeshSculptBrushOpProps>> BrushOpPropSets;
 
@@ -252,8 +325,8 @@ protected:
 
 	TMap<int32, TUniquePtr<FMeshSculptBrushOpFactory>> SecondaryBrushOpFactories;
 
-	void RegisterBrushType(int32 Identifier, TUniquePtr<FMeshSculptBrushOpFactory> Factory, UMeshSculptBrushOpProps* PropSet);
-	void RegisterSecondaryBrushType(int32 Identifier, TUniquePtr<FMeshSculptBrushOpFactory> Factory, UMeshSculptBrushOpProps* PropSet);
+	void RegisterBrushType(int32 Identifier, FText Name, TUniquePtr<FMeshSculptBrushOpFactory> Factory, UMeshSculptBrushOpProps* PropSet);
+	void RegisterSecondaryBrushType(int32 Identifier, FText Name, TUniquePtr<FMeshSculptBrushOpFactory> Factory, UMeshSculptBrushOpProps* PropSet);
 
 	virtual void SaveAllBrushTypeProperties(UInteractiveTool* SaveFromTool);
 	virtual void RestoreAllBrushTypeProperties(UInteractiveTool* RestoreToTool);
@@ -275,12 +348,23 @@ protected:
 	//
 	// Falloff types
 	//
-protected:
-	TSharedPtr<FMeshSculptFallofFunc> PrimaryFalloff;
+public:
+	struct FFalloffTypeInfo
+	{
+		FText Name;
+		FString StringIdentifier;
+		int32 Identifier;
+	};
+	const TArray<FFalloffTypeInfo>& GetRegisteredPrimaryFalloffTypes() const { return RegisteredPrimaryFalloffTypes; }
 
-	void SetPrimaryFalloffType(EMeshSculptFalloffType Falloff);
+	/** Set the active falloff type for the primary brush */
+	virtual void SetPrimaryFalloffType(EMeshSculptFalloffType Falloff);
 
+protected:
+	TSharedPtr<FMeshSculptFallofFunc> PrimaryFalloff;
 
+	TArray<FFalloffTypeInfo> RegisteredPrimaryFalloffTypes;
+	void RegisterStandardFalloffTypes();
 
 
 	//
