ModelingTools: Add analytics to BakeTx and BakeAll.

#rnx

diff --git a/Engine/Plugins/Experimental/MeshModelingToolsetExp/Source/MeshModelingToolsExp/Private/BakeMeshAttributeMapsTool.cpp b/Engine/Plugins/Experimental/MeshModelingToolsetExp/Source/MeshModelingToolsExp/Private/BakeMeshAttributeMapsTool.cpp
--- a/Engine/Plugins/Experimental/MeshModelingToolsetExp/Source/MeshModelingToolsExp/Private/BakeMeshAttributeMapsTool.cpp
+++ b/Engine/Plugins/Experimental/MeshModelingToolsetExp/Source/MeshModelingToolsExp/Private/BakeMeshAttributeMapsTool.cpp
@@ -19,6 +19,8 @@
 #include "AssetUtils/Texture2DUtil.h"
 #include "ModelingObjectsCreationAPI.h"
 
+#include "EngineAnalytics.h"
+
 #include "TargetInterfaces/MaterialProvider.h"
 #include "TargetInterfaces/MeshDescriptionProvider.h"
 #include "TargetInterfaces/PrimitiveComponentBackedTarget.h"
@@ -354,15 +356,14 @@ void UBakeMeshAttributeMapsTool::Setup()
 
 	UpdateOnModeChange();
 
-	bInputsDirty = true;
-	bDetailMeshValid = false;
+	UpdateDetailMesh();
 
 	SetToolDisplayName(LOCTEXT("ToolName", "Bake Textures"));
 	GetToolManager()->DisplayMessage(
 		LOCTEXT("OnStartTool", "Bake Maps. Select Bake Mesh (LowPoly) first, then (optionally) Detail Mesh second. Texture Assets will be created on Accept. "),
 		EToolMessageLevel::UserNotification);
 
-	SetupBaseToolProperties();
+	PostSetup();
 }
 
 
@@ -471,16 +472,7 @@ void UBakeMeshAttributeMapsTool::Shutdown(EToolShutdownType ShutdownType)
 			SourceAsset = SkeletalMeshTarget ? SkeletalMeshTarget->GetSkeletalMesh() : nullptr;
 		}
 		const UPrimitiveComponent* SourceComponent = UE::ToolTarget::GetTargetComponent(Targets[0]);
-		const FString BaseName = UE::ToolTarget::GetTargetActor(Targets[0])->GetActorNameOrLabel();
-
-		bool bCreatedAssetOK = true;
-		for (const TTuple<EBakeMapType, TObjectPtr<UTexture2D>>& Result : Settings->Result)
-		{
-			FString TexName;
-			GetTextureName(Result.Get<0>(), BaseName, TexName);
-			bCreatedAssetOK = bCreatedAssetOK && UE::Modeling::CreateTextureObject(GetToolManager(), FCreateTextureObjectParams{ 0, SourceComponent->GetWorld(), SourceAsset, TexName, Result.Get<1>() }).IsOK();
-		}
-		ensure(bCreatedAssetOK);
+		CreateTextureAssets(Settings->Result, SourceComponent->GetWorld(), SourceAsset);
 	}
 }
 
@@ -500,7 +492,7 @@ void UBakeMeshAttributeMapsTool::UpdateDetailMesh()
 		const FTransform3d WorldToBase = UE::ToolTarget::GetLocalToWorldTransform(Targets[0]);
 		MeshTransforms::ApplyTransform(*DetailMesh, WorldToBase.Inverse());
 	}
-	
+
 	DetailSpatial = MakeShared<FDynamicMeshAABBTree3, ESPMode::ThreadSafe>();
 	DetailSpatial->SetMesh(DetailMesh.Get(), true);
 
@@ -544,18 +536,18 @@ void UBakeMeshAttributeMapsTool::UpdateDetailMesh()
 		}
 	});
 
+	bDetailMeshValid = true;
 	bInputsDirty = true;
+	CachedBakeCacheSettings = FBakeCacheSettings();
 	DetailMeshTimestamp++;
 }
 
 
 void UBakeMeshAttributeMapsTool::UpdateResult()
 {
-	if (bDetailMeshValid == false)
+	if (!bDetailMeshValid)
 	{
 		UpdateDetailMesh();
-		bDetailMeshValid = true;
-		CachedBakeCacheSettings = FBakeCacheSettings();
 	}
 
 	// bInputsDirty ensures that we only validate parameters once per param
@@ -578,8 +570,11 @@ void UBakeMeshAttributeMapsTool::UpdateResult()
 	BakeCacheSettings.DetailTimestamp = this->DetailMeshTimestamp;
 	BakeCacheSettings.Thickness = Settings->Thickness;
 	BakeCacheSettings.Multisampling = (int32)Settings->Multisampling;
+	BakeCacheSettings.bUseWorldSpace = Settings->bUseWorldSpace;
 
-	// process the raw bitfield before caching which may add additional targets.
+	// Record the original map types and process the raw bitfield which may add
+	// additional targets.
+	BakeCacheSettings.SourceBakeMapTypes = static_cast<EBakeMapType>(Settings->MapTypes);
 	BakeCacheSettings.BakeMapTypes = GetMapTypes(Settings->MapTypes);
 
 	// update bake cache settings
@@ -639,19 +634,7 @@ void UBakeMeshAttributeMapsTool::UpdateResult()
 
 	// This should be the only point of compute invalidation to
 	// minimize synchronization issues.
-	bool bInvalidate = bInputsDirty || (bool)(OpState & EBakeOpState::Evaluate);
-	if (!Compute)
-	{
-		Compute = MakeUnique<TGenericDataBackgroundCompute<FMeshMapBaker>>();
-		Compute->Setup(this);
-		Compute->OnResultUpdated.AddLambda([this](const TUniquePtr<FMeshMapBaker>& NewResult) { OnMapsUpdated(NewResult, CachedBakeCacheSettings.SourceFormat); });
-		Compute->InvalidateResult();
-	}
-	else if (bInvalidate)
-	{
-		Compute->InvalidateResult();
-	}
-	bInputsDirty = false;
+	InvalidateCompute();
 }
 
 
@@ -993,5 +976,17 @@ void UBakeMeshAttributeMapsTool::InvalidateResults()
 }
 
 
+void UBakeMeshAttributeMapsTool::GatherAnalytics(FBakeAnalytics::FMeshSettings& Data)
+{
+	if (FEngineAnalytics::IsAvailable())
+	{
+		Data.NumTargetMeshTris = BaseMesh.TriangleCount();
+		Data.NumDetailMesh = 1;
+		Data.NumDetailMeshTris = DetailMesh->TriangleCount();
+	}
+}
+
+
+
 
 #undef LOCTEXT_NAMESPACE
diff --git a/Engine/Plugins/Experimental/MeshModelingToolsetExp/Source/MeshModelingToolsExp/Private/BakeMeshAttributeMapsToolBase.cpp b/Engine/Plugins/Experimental/MeshModelingToolsetExp/Source/MeshModelingToolsExp/Private/BakeMeshAttributeMapsToolBase.cpp
--- a/Engine/Plugins/Experimental/MeshModelingToolsetExp/Source/MeshModelingToolsExp/Private/BakeMeshAttributeMapsToolBase.cpp
+++ b/Engine/Plugins/Experimental/MeshModelingToolsetExp/Source/MeshModelingToolsExp/Private/BakeMeshAttributeMapsToolBase.cpp
@@ -19,8 +19,13 @@
 
 #include "TargetInterfaces/MaterialProvider.h"
 #include "ModelingToolTargetUtil.h"
+#include "ModelingObjectsCreationAPI.h"
+
+#include "EngineAnalytics.h"
 
 #include "ExplicitUseGeometryMathTypes.h"		// using UE::Geometry::(math types)
+#include "Sampling/MeshCurvatureMapEvaluator.h"
+#include "Sampling/MeshResampleImageEvaluator.h"
 using namespace UE::Geometry;
 
 #define LOCTEXT_NAMESPACE "UBakeMeshAttributeMapsToolBase"
@@ -113,11 +118,13 @@ void UBakeMeshAttributeMapsToolBase::Setup()
 	PreviewMesh->SetVisible(true);
 }
 
-void UBakeMeshAttributeMapsToolBase::SetupBaseToolProperties()
+void UBakeMeshAttributeMapsToolBase::PostSetup()
 {
 	VisualizationProps = NewObject<UBakedOcclusionMapVisualizationProperties>(this);
 	VisualizationProps->RestoreProperties(this);
 	AddToolPropertySource(VisualizationProps);
+
+	GatherAnalytics(BakeAnalytics.MeshSettings);
 }
 
 void UBakeMeshAttributeMapsToolBase::Shutdown(EToolShutdownType ShutdownType)
@@ -172,15 +179,55 @@ void UBakeMeshAttributeMapsToolBase::SetWorld(UWorld* World)
 	TargetWorld = World;
 }
 
+
 void UBakeMeshAttributeMapsToolBase::UpdateResult()
 {
 }
 
+
 void UBakeMeshAttributeMapsToolBase::UpdateVisualization()
 {
 }
 
 
+void UBakeMeshAttributeMapsToolBase::InvalidateCompute()
+{
+	const bool bInvalidate = bInputsDirty || static_cast<bool>(OpState & EBakeOpState::Evaluate);
+	if (!Compute)
+	{
+		Compute = MakeUnique<TGenericDataBackgroundCompute<FMeshMapBaker>>();
+		Compute->Setup(this);
+		Compute->OnResultUpdated.AddLambda([this](const TUniquePtr<FMeshMapBaker>& NewResult) { OnMapsUpdated(NewResult); });
+		Compute->InvalidateResult();
+	}
+	else if (bInvalidate)
+	{
+		Compute->InvalidateResult();
+	}
+	bInputsDirty = false;
+}
+
+
+void UBakeMeshAttributeMapsToolBase::CreateTextureAssets(const TMap<EBakeMapType, TObjectPtr<UTexture2D>>& Textures, UWorld* SourceWorld, UObject* SourceAsset)
+{
+	bool bCreatedAssetOK = true;
+	const FString BaseName = UE::ToolTarget::GetTargetActor(Targets[0])->GetActorNameOrLabel();
+	for (const TTuple<EBakeMapType, TObjectPtr<UTexture2D>>& Result : Textures)
+	{
+		FString TexName;
+		GetTextureName(Result.Get<0>(), BaseName, TexName);
+		bCreatedAssetOK = bCreatedAssetOK &&
+			UE::Modeling::CreateTextureObject(
+				GetToolManager(),
+				FCreateTextureObjectParams{ 0, SourceWorld, SourceAsset, TexName, Result.Get<1>() }).IsOK();
+	}
+	ensure(bCreatedAssetOK);
+
+	RecordAnalytics(BakeAnalytics, GetAnalyticsEventName());
+}
+
+
+
 void UBakeMeshAttributeMapsToolBase::UpdatePreview(const EBakeMapType PreviewMapType)
 {
 	if (PreviewMapType == EBakeMapType::None)
@@ -237,9 +284,10 @@ void UBakeMeshAttributeMapsToolBase::UpdatePreview(const EBakeMapType PreviewMap
 }
 
 
-void UBakeMeshAttributeMapsToolBase::OnMapsUpdated(const TUniquePtr<UE::Geometry::FMeshMapBaker>& NewResult, const EBakeTextureFormat Format)
+void UBakeMeshAttributeMapsToolBase::OnMapsUpdated(const TUniquePtr<UE::Geometry::FMeshMapBaker>& NewResult)
 {
-	FImageDimensions BakeDimensions = NewResult->GetDimensions();
+	const FImageDimensions BakeDimensions = NewResult->GetDimensions();
+	const EBakeTextureFormat Format = CachedBakeCacheSettings.SourceFormat;
 	const int32 NumEval = NewResult->NumEvaluators();
 	for (int32 EvalIdx = 0; EvalIdx < NumEval; ++EvalIdx)
 	{
@@ -350,6 +398,7 @@ void UBakeMeshAttributeMapsToolBase::OnMapsUpdated(const TUniquePtr<UE::Geometry
 		}
 	}
 
+	GatherAnalytics(*NewResult, CachedBakeCacheSettings, BakeAnalytics);
 	UpdateVisualization();
 	GetToolManager()->PostInvalidation();
 }
@@ -516,4 +565,154 @@ void UBakeMeshAttributeMapsToolBase::InitializeEmptyMaps()
 	EmptyColorMapWhite = ColorBuilderWhite.GetTexture2D();
 }
 
+
+void UBakeMeshAttributeMapsToolBase::GatherAnalytics(FBakeAnalytics::FMeshSettings& Data)
+{
+	
+}
+
+
+void UBakeMeshAttributeMapsToolBase::GatherAnalytics(const FMeshMapBaker& Result, const FBakeCacheSettings& Settings, FBakeAnalytics& Data)
+{
+	if (!FEngineAnalytics::IsAvailable())
+	{
+		return;
+	}
+	
+	Data.TotalBakeDuration = Result.BakeAnalytics.TotalBakeDuration;
+	Data.WriteToImageDuration = Result.BakeAnalytics.WriteToImageDuration;
+	Data.WriteToGutterDuration = Result.BakeAnalytics.WriteToGutterDuration;
+	Data.NumBakedPixels = Result.BakeAnalytics.NumBakedPixels;
+	Data.NumGutterPixels = Result.BakeAnalytics.NumGutterPixels;
+	Data.BakeSettings = Settings;
+
+	const int NumEvaluators = Result.NumEvaluators();
+	for (int EvalId = 0; EvalId < NumEvaluators; ++EvalId)
+	{
+		FMeshMapEvaluator* Eval = Result.GetEvaluator(EvalId);
+		switch(Eval->Type())
+		{
+		case EMeshMapEvaluatorType::Occlusion:
+		{
+			const FMeshOcclusionMapEvaluator* OcclusionEval = static_cast<FMeshOcclusionMapEvaluator*>(Eval);
+			Data.OcclusionSettings.OcclusionRays = OcclusionEval->NumOcclusionRays;
+			Data.OcclusionSettings.MaxDistance = OcclusionEval->MaxDistance;
+			Data.OcclusionSettings.SpreadAngle = OcclusionEval->SpreadAngle;
+			Data.OcclusionSettings.BiasAngle = OcclusionEval->BiasAngleDeg;
+			break;
+		}
+		case EMeshMapEvaluatorType::Curvature:
+		{
+			const FMeshCurvatureMapEvaluator* CurvatureEval = static_cast<FMeshCurvatureMapEvaluator*>(Eval);
+			Data.CurvatureSettings.CurvatureType = static_cast<int>(CurvatureEval->UseCurvatureType);
+			Data.CurvatureSettings.RangeMultiplier = CurvatureEval->RangeScale;
+			Data.CurvatureSettings.MinRangeMultiplier = CurvatureEval->MinRangeScale;
+			Data.CurvatureSettings.ColorMode = static_cast<int>(CurvatureEval->UseColorMode);
+			Data.CurvatureSettings.ClampMode = static_cast<int>(CurvatureEval->UseClampMode);
+			break;
+		}
+		default:
+			break;
+		};
+	}
+}
+
+
+void UBakeMeshAttributeMapsToolBase::RecordAnalytics(const FBakeAnalytics& Data, const FString& EventName)
+{
+	if (!FEngineAnalytics::IsAvailable())
+	{
+		return;
+	}
+	
+	TArray<FAnalyticsEventAttribute> Attributes;
+
+	// General
+	Attributes.Add(FAnalyticsEventAttribute(TEXT("Bake.Duration.Total.Seconds"), Data.TotalBakeDuration));
+	Attributes.Add(FAnalyticsEventAttribute(TEXT("Bake.Duration.WriteToImage.Seconds"), Data.WriteToImageDuration));
+	Attributes.Add(FAnalyticsEventAttribute(TEXT("Bake.Duration.WriteToGutter.Seconds"), Data.WriteToGutterDuration));
+	Attributes.Add(FAnalyticsEventAttribute(TEXT("Bake.Stats.NumBakedPixels"), Data.NumBakedPixels));
+	Attributes.Add(FAnalyticsEventAttribute(TEXT("Bake.Stats.NumGutterPixels"), Data.NumGutterPixels));
+
+	// Mesh data
+	Attributes.Add(FAnalyticsEventAttribute(TEXT("Input.TargetMesh.NumTriangles"), Data.MeshSettings.NumTargetMeshTris));
+	Attributes.Add(FAnalyticsEventAttribute(TEXT("Input.DetailMesh.NumMeshes"), Data.MeshSettings.NumDetailMesh));
+	Attributes.Add(FAnalyticsEventAttribute(TEXT("Input.DetailMesh.NumTriangles"), Data.MeshSettings.NumDetailMeshTris));
+
+	// Bake settings
+	Attributes.Add(FAnalyticsEventAttribute(TEXT("Settings.Image.Width"), Data.BakeSettings.Dimensions.GetWidth()));
+	Attributes.Add(FAnalyticsEventAttribute(TEXT("Settings.Image.Height"), Data.BakeSettings.Dimensions.GetHeight()));
+	Attributes.Add(FAnalyticsEventAttribute(TEXT("Settings.Multisampling"), Data.BakeSettings.Multisampling));
+	Attributes.Add(FAnalyticsEventAttribute(TEXT("Settings.Thickness"), Data.BakeSettings.Thickness));
+	Attributes.Add(FAnalyticsEventAttribute(TEXT("Settings.UVLayer"), Data.BakeSettings.UVLayer));
+	Attributes.Add(FAnalyticsEventAttribute(TEXT("Settings.UseWorldSpace"), Data.BakeSettings.bUseWorldSpace));
+
+	// Map types
+	const bool bTangentSpaceNormal = static_cast<bool>(Data.BakeSettings.SourceBakeMapTypes & EBakeMapType::TangentSpaceNormal);
+	Attributes.Add(FAnalyticsEventAttribute(TEXT("Settings.TangentSpaceNormal.Enabled"), bTangentSpaceNormal));
+
+	const bool bAmbientOcclusion = static_cast<bool>(Data.BakeSettings.SourceBakeMapTypes & EBakeMapType::AmbientOcclusion);
+	Attributes.Add(FAnalyticsEventAttribute(TEXT("Settings.AmbientOcclusion.Enabled"), bAmbientOcclusion));
+	if (bAmbientOcclusion)
+	{
+		Attributes.Add(FAnalyticsEventAttribute(TEXT("Settings.AmbientOcclusion.OcclusionRays"), Data.OcclusionSettings.OcclusionRays));
+		Attributes.Add(FAnalyticsEventAttribute(TEXT("Settings.AmbientOcclusion.MaxDistance"), Data.OcclusionSettings.MaxDistance));
+		Attributes.Add(FAnalyticsEventAttribute(TEXT("Settings.AmbientOcclusion.SpreadAngle"), Data.OcclusionSettings.SpreadAngle));
+		Attributes.Add(FAnalyticsEventAttribute(TEXT("Settings.AmbientOcclusion.BiasAngle"), Data.OcclusionSettings.BiasAngle));
+	}
+	const bool bBentNormal = static_cast<bool>(Data.BakeSettings.SourceBakeMapTypes & EBakeMapType::BentNormal);
+	Attributes.Add(FAnalyticsEventAttribute(TEXT("Settings.BentNormal.Enabled"), bBentNormal));
+	if (bBentNormal)
+	{
+		Attributes.Add(FAnalyticsEventAttribute(TEXT("Settings.BentNormal.OcclusionRays"), Data.OcclusionSettings.OcclusionRays));
+		Attributes.Add(FAnalyticsEventAttribute(TEXT("Settings.BentNormal.MaxDistance"), Data.OcclusionSettings.MaxDistance));
+		Attributes.Add(FAnalyticsEventAttribute(TEXT("Settings.BentNormal.SpreadAngle"), Data.OcclusionSettings.SpreadAngle));
+	}
+
+	const bool bCurvature = static_cast<bool>(Data.BakeSettings.SourceBakeMapTypes & EBakeMapType::Curvature);
+	Attributes.Add(FAnalyticsEventAttribute(TEXT("Settings.Curvature.Enabled"), bCurvature));
+	if (bCurvature)
+	{
+		Attributes.Add(FAnalyticsEventAttribute(TEXT("Settings.Curvature.CurvatureType"), Data.CurvatureSettings.CurvatureType));
+		Attributes.Add(FAnalyticsEventAttribute(TEXT("Settings.Curvature.RangeMultiplier"), Data.CurvatureSettings.RangeMultiplier));
+		Attributes.Add(FAnalyticsEventAttribute(TEXT("Settings.Curvature.MinRangeMultiplier"), Data.CurvatureSettings.MinRangeMultiplier));
+		Attributes.Add(FAnalyticsEventAttribute(TEXT("Settings.Curvature.ClampMode"), Data.CurvatureSettings.ClampMode));
+		Attributes.Add(FAnalyticsEventAttribute(TEXT("Settings.Curvature.ColorMode"), Data.CurvatureSettings.ColorMode));
+	}
+
+	const bool bObjectSpaceNormal = static_cast<bool>(Data.BakeSettings.SourceBakeMapTypes & EBakeMapType::ObjectSpaceNormal);
+	Attributes.Add(FAnalyticsEventAttribute(TEXT("Settings.ObjectSpaceNormal.Enabled"), bObjectSpaceNormal));
+
+	const bool bFaceNormal = static_cast<bool>(Data.BakeSettings.SourceBakeMapTypes & EBakeMapType::FaceNormal);
+	Attributes.Add(FAnalyticsEventAttribute(TEXT("Settings.FaceNormal.Enabled"), bFaceNormal));
+
+	const bool bPosition = static_cast<bool>(Data.BakeSettings.SourceBakeMapTypes & EBakeMapType::Position);
+	Attributes.Add(FAnalyticsEventAttribute(TEXT("Settings.Position.Enabled"), bPosition));
+
+	const bool bMaterialID = static_cast<bool>(Data.BakeSettings.SourceBakeMapTypes & EBakeMapType::MaterialID);
+	Attributes.Add(FAnalyticsEventAttribute(TEXT("Settings.MaterialID.Enabled"), bMaterialID));
+	
+	const bool bTexture = static_cast<bool>(Data.BakeSettings.SourceBakeMapTypes & EBakeMapType::Texture);
+	Attributes.Add(FAnalyticsEventAttribute(TEXT("Settings.Texture.Enabled"), bTexture));
+
+	const bool bMultiTexture = static_cast<bool>(Data.BakeSettings.SourceBakeMapTypes & EBakeMapType::MultiTexture); 
+	Attributes.Add(FAnalyticsEventAttribute(TEXT("Settings.MultiTexture.Enabled"), bMultiTexture));
+
+	const bool bVertexColor = static_cast<bool>(Data.BakeSettings.SourceBakeMapTypes & EBakeMapType::VertexColor);
+	Attributes.Add(FAnalyticsEventAttribute(TEXT("Settings.VertexColor.Enabled"), bVertexColor));
+
+	FEngineAnalytics::GetProvider().RecordEvent(FString(TEXT("Editor.Usage.ModelingTools.")) + EventName, Attributes);
+
+	constexpr bool bLogAnalytics = false; 
+	if constexpr (bLogAnalytics)
+	{
+		for (const FAnalyticsEventAttribute& Attr : Attributes)
+		{
+			UE_LOG(LogGeometry, Log, TEXT("[%s] %s = %s"), *EventName, *Attr.GetName(), *Attr.GetValue());
+		}
+	}
+}
+
+
+
 #undef LOCTEXT_NAMESPACE
diff --git a/Engine/Plugins/Experimental/MeshModelingToolsetExp/Source/MeshModelingToolsExp/Private/BakeMultiMeshAttributeMapsTool.cpp b/Engine/Plugins/Experimental/MeshModelingToolsetExp/Source/MeshModelingToolsExp/Private/BakeMultiMeshAttributeMapsTool.cpp
--- a/Engine/Plugins/Experimental/MeshModelingToolsetExp/Source/MeshModelingToolsExp/Private/BakeMultiMeshAttributeMapsTool.cpp
+++ b/Engine/Plugins/Experimental/MeshModelingToolsetExp/Source/MeshModelingToolsExp/Private/BakeMultiMeshAttributeMapsTool.cpp
@@ -20,6 +20,8 @@
 #include "AssetUtils/Texture2DUtil.h"
 #include "ModelingObjectsCreationAPI.h"
 
+#include "EngineAnalytics.h"
+
 #include "TargetInterfaces/MaterialProvider.h"
 #include "TargetInterfaces/MeshDescriptionProvider.h"
 #include "TargetInterfaces/PrimitiveComponentBackedTarget.h"
@@ -111,6 +113,12 @@ public:
 		MeshScene->ProcessActorChildMeshes(ProcessChildMesh);
 	}
 
+	virtual int32 GetTriangleCount(const void* Mesh) const override
+	{
+		const IMeshSpatialWrapper* Spatial = static_cast<const IMeshSpatialWrapper*>(Mesh);
+		return Spatial->GetTriangleCount();
+	}
+
 	virtual void SetColorMap(const void* Mesh, const FBakeDetailTexture& Map) override
 	{
 		DetailColorMaps[Mesh] = Map;
@@ -480,7 +488,7 @@ void UBakeMultiMeshAttributeMapsTool::Setup()
 		LOCTEXT("OnStartTool", "Bake Maps. Select Bake Mesh (LowPoly) first, then select Detail Meshes (HiPoly) to bake. Texture Assets will be created on Accept. "),
 		EToolMessageLevel::UserNotification);
 
-	SetupBaseToolProperties();
+	PostSetup();
 }
 
 
@@ -544,19 +552,10 @@ void UBakeMultiMeshAttributeMapsTool::Shutdown(EToolShutdownType ShutdownType)
 	
 	if (ShutdownType == EToolShutdownType::Accept)
 	{
-		UStaticMeshComponent* StaticMeshComponent = CastChecked<UStaticMeshComponent>(UE::ToolTarget::GetTargetComponent(Targets[0]));
-		UStaticMesh* StaticMeshAsset = StaticMeshComponent->GetStaticMesh();
-		check(StaticMeshAsset);
-		const FString BaseName = UE::ToolTarget::GetTargetActor(Targets[0])->GetActorNameOrLabel();
-
-		bool bCreatedAssetOK = true;
-		for (const TTuple<EBakeMapType, TObjectPtr<UTexture2D>>& Result : Settings->Result)
-		{
-			FString TexName;
-			GetTextureName(Result.Get<0>(), BaseName, TexName);
-			bCreatedAssetOK = bCreatedAssetOK && UE::Modeling::CreateTextureObject(GetToolManager(), FCreateTextureObjectParams{ 0, StaticMeshAsset->GetWorld(), StaticMeshAsset, TexName, Result.Get<1>() }).IsOK();
-		}
-		ensure(bCreatedAssetOK);
+		IStaticMeshBackedTarget* StaticMeshTarget = Cast<IStaticMeshBackedTarget>(Targets[0]);
+		UObject* SourceAsset = StaticMeshTarget ? StaticMeshTarget->GetStaticMesh() : nullptr;
+		const UPrimitiveComponent* SourceComponent = UE::ToolTarget::GetTargetComponent(Targets[0]);
+		CreateTextureAssets(Settings->Result, SourceComponent->GetWorld(), SourceAsset);
 	}
 }
 
@@ -583,7 +582,9 @@ void UBakeMultiMeshAttributeMapsTool::UpdateResult()
 	BakeCacheSettings.Thickness = Settings->Thickness;
 	BakeCacheSettings.Multisampling = (int32)Settings->Multisampling;
 
-	// process the raw bitfield before caching which may add additional targets.
+	// Record the original map types and process the raw bitfield which may add
+	// additional targets.
+	BakeCacheSettings.SourceBakeMapTypes = static_cast<EBakeMapType>(Settings->MapTypes);
 	BakeCacheSettings.BakeMapTypes = GetMapTypes(Settings->MapTypes);
 
 	// update bake cache settings
@@ -608,19 +609,7 @@ void UBakeMultiMeshAttributeMapsTool::UpdateResult()
 
 	// This should be the only point of compute invalidation to
 	// minimize synchronization issues.
-	const bool bInvalidate = bInputsDirty || (bool)(OpState & EBakeOpState::Evaluate);
-	if (!Compute)
-	{
-		Compute = MakeUnique<TGenericDataBackgroundCompute<FMeshMapBaker>>();
-		Compute->Setup(this);
-		Compute->OnResultUpdated.AddLambda([this](const TUniquePtr<FMeshMapBaker>& NewResult) { OnMapsUpdated(NewResult, CachedBakeCacheSettings.SourceFormat); });
-		Compute->InvalidateResult();
-	}
-	else if (bInvalidate)
-	{
-		Compute->InvalidateResult();
-	}
-	bInputsDirty = false;
+	InvalidateCompute();
 }
 
 
@@ -712,6 +701,7 @@ void UBakeMultiMeshAttributeMapsTool::UpdateOnModeChange()
 	OnMapTypesUpdated(Settings);
 }
 
+
 void UBakeMultiMeshAttributeMapsTool::InvalidateResults()
 {
 	for (TTuple<EBakeMapType, TObjectPtr<UTexture2D>>& Result : Settings->Result)
@@ -721,5 +711,24 @@ void UBakeMultiMeshAttributeMapsTool::InvalidateResults()
 }
 
 
+void UBakeMultiMeshAttributeMapsTool::GatherAnalytics(FBakeAnalytics::FMeshSettings& Data)
+{
+	if (FEngineAnalytics::IsAvailable())
+	{
+		Data.NumTargetMeshTris = BaseMesh.TriangleCount();
+		Data.NumDetailMesh = 0;
+		Data.NumDetailMeshTris = 0;
+		DetailMeshScene.ProcessActorChildMeshes([&Data](const FActorAdapter* ActorAdapter, const FActorChildMesh* ChildMesh)
+		{
+			if (ChildMesh)
+			{
+				++Data.NumDetailMesh;
+				Data.NumDetailMeshTris += ChildMesh->MeshSpatial->GetTriangleCount();
+			}
+		});
+	}
+}
+
+
 
 #undef LOCTEXT_NAMESPACE
diff --git a/Engine/Plugins/Experimental/MeshModelingToolsetExp/Source/MeshModelingToolsExp/Public/BakeMeshAttributeMapsTool.h b/Engine/Plugins/Experimental/MeshModelingToolsetExp/Source/MeshModelingToolsExp/Public/BakeMeshAttributeMapsTool.h
--- a/Engine/Plugins/Experimental/MeshModelingToolsetExp/Source/MeshModelingToolsExp/Public/BakeMeshAttributeMapsTool.h
+++ b/Engine/Plugins/Experimental/MeshModelingToolsetExp/Source/MeshModelingToolsExp/Public/BakeMeshAttributeMapsTool.h
@@ -148,6 +148,8 @@ protected:
 	// Begin UBakeMeshAttributeMapsToolBase interface
 	virtual void UpdateResult() override;
 	virtual void UpdateVisualization() override;
+
+	virtual void GatherAnalytics(FBakeAnalytics::FMeshSettings& Data) override;
 	// End UBakeMeshAttributeMapsToolBase interface
 
 protected:
@@ -162,7 +164,6 @@ protected:
 	void UpdateDetailMesh();
 	bool bDetailMeshValid = false;
 
-	bool bInputsDirty = false;
 	void UpdateOnModeChange();
 
 	void InvalidateResults();
diff --git a/Engine/Plugins/Experimental/MeshModelingToolsetExp/Source/MeshModelingToolsExp/Public/BakeMeshAttributeMapsToolBase.h b/Engine/Plugins/Experimental/MeshModelingToolsetExp/Source/MeshModelingToolsExp/Public/BakeMeshAttributeMapsToolBase.h
--- a/Engine/Plugins/Experimental/MeshModelingToolsetExp/Source/MeshModelingToolsExp/Public/BakeMeshAttributeMapsToolBase.h
+++ b/Engine/Plugins/Experimental/MeshModelingToolsetExp/Source/MeshModelingToolsExp/Public/BakeMeshAttributeMapsToolBase.h
@@ -174,8 +174,12 @@ protected:
 	float SecondsBeforeWorkingMaterial = 0.75;
 
 protected:
-	/** To be invoked at end of client Setup methods. */
-	void SetupBaseToolProperties();
+	/**
+	 * Post-client setup function. Should be invoked at end of client Setup().
+	 * Initialize common base tool properties (ex. visualization properties) and
+	 * analytics.
+	 */
+	void PostSetup();
 	
 	/**
 	 * Process dirty props and update background compute.
@@ -189,6 +193,19 @@ protected:
 	 */
 	virtual void UpdateVisualization();
 
+	/**
+	 * Invalidates the background compute operator.
+	 */
+	void InvalidateCompute();
+
+	/**
+	 * Create texture assets from our result map of Texture2D
+	 * @param Textures the result map of textures to create
+	 * @param SourceWorld the source world to define where the texture assets will be stored.
+	 * @param SourceAsset if not null, result textures will be stored adjacent to this asset.
+	 */
+	void CreateTextureAssets(const TMap<EBakeMapType, TObjectPtr<UTexture2D>>& Textures, UWorld* SourceWorld, UObject* SourceAsset);
+
 
 protected:
 	//
@@ -200,11 +217,13 @@ protected:
 	UE::Geometry::FDynamicMesh3 BaseMesh;
 	TSharedPtr<UE::Geometry::TMeshTangents<double>, ESPMode::ThreadSafe> BaseMeshTangents;
 	UE::Geometry::FDynamicMeshAABBTree3 BaseSpatial;
-	
+
+	bool bInputsDirty = false;
 	EBakeOpState OpState = EBakeOpState::Evaluate;
 
 	struct FBakeCacheSettings
 	{
+		EBakeMapType SourceBakeMapTypes = EBakeMapType::None;
 		EBakeMapType BakeMapTypes = EBakeMapType::None;
 		FImageDimensions Dimensions;
 		EBakeTextureFormat SourceFormat = EBakeTextureFormat::ChannelBits8;
@@ -212,13 +231,15 @@ protected:
 		int32 DetailTimestamp = 0;
 		float Thickness = 3.0;
 		int32 Multisampling = 1;
+		bool bUseWorldSpace = false;
 
 		bool operator==(const FBakeCacheSettings& Other) const
 		{
 			return BakeMapTypes == Other.BakeMapTypes && Dimensions == Other.Dimensions &&
 				UVLayer == Other.UVLayer && DetailTimestamp == Other.DetailTimestamp &&
 				Thickness == Other.Thickness && Multisampling == Other.Multisampling &&
-				SourceFormat == Other.SourceFormat;
+				SourceFormat == Other.SourceFormat && SourceBakeMapTypes == Other.SourceBakeMapTypes &&
+				bUseWorldSpace == Other.bUseWorldSpace;
 		}
 	};
 	FBakeCacheSettings CachedBakeCacheSettings;
@@ -245,7 +266,7 @@ protected:
 	 * 
 	 * @param NewResult the resulting FMeshMapBaker from the background Compute
 	 */
-	void OnMapsUpdated(const TUniquePtr<UE::Geometry::FMeshMapBaker>& NewResult, EBakeTextureFormat Format);
+	void OnMapsUpdated(const TUniquePtr<UE::Geometry::FMeshMapBaker>& NewResult);
 
 
 	/**
@@ -277,6 +298,63 @@ protected:
 	template <typename PropertySet>
 	static void UpdateUVLayerNames(PropertySet& Properties, const FDynamicMesh3& Mesh);
 
+
+	//
+	// Analytics
+	//
+	struct FBakeAnalytics
+	{
+		double TotalBakeDuration = 0.0;
+		double WriteToImageDuration = 0.0;
+		double WriteToGutterDuration = 0.0;
+		int64 NumBakedPixels = 0;
+		int64 NumGutterPixels = 0;
+
+		struct FMeshSettings
+		{
+			int32 NumTargetMeshTris = 0;
+			int32 NumDetailMesh = 0;
+			int64 NumDetailMeshTris = 0;
+		};
+		FMeshSettings MeshSettings;
+
+		FBakeCacheSettings BakeSettings;
+		FOcclusionMapSettings OcclusionSettings;
+		FCurvatureMapSettings CurvatureSettings;
+	};
+	FBakeAnalytics BakeAnalytics;
+
+	/**
+	 * Computes the NumTargetMeshTris, NumDetailMesh and NumDetailMeshTris analytics.
+	 * @param Data the mesh analytics data to compute
+	 */
+	virtual void GatherAnalytics(FBakeAnalytics::FMeshSettings& Data);
+
+	/**
+	 * Records bake timing and settings data for analytics.
+	 * @param Result the result of the bake.
+	 * @param Settings The bake settings used for the bake.
+	 * @param Data the output bake analytics struct.
+	 */
+	static void GatherAnalytics(const UE::Geometry::FMeshMapBaker& Result,
+								const FBakeCacheSettings& Settings,
+								FBakeAnalytics& Data);
+
+	/**
+	 * Posts an analytics event using the given analytics struct.
+	 * @param Data the bake analytics struct to output.
+	 * @param EventName the name of the analytics event to output.
+	 */
+	static void RecordAnalytics(const FBakeAnalytics& Data, const FString& EventName);
+
+	
+	/**
+	 * @return the analytics event name for this tool.
+	 */
+	virtual FString GetAnalyticsEventName() const
+	{
+		return TEXT("BakeTexture");
+	}
 	
 	//
 	// Utilities
diff --git a/Engine/Plugins/Experimental/MeshModelingToolsetExp/Source/MeshModelingToolsExp/Public/BakeMultiMeshAttributeMapsTool.h b/Engine/Plugins/Experimental/MeshModelingToolsetExp/Source/MeshModelingToolsExp/Public/BakeMultiMeshAttributeMapsTool.h
--- a/Engine/Plugins/Experimental/MeshModelingToolsetExp/Source/MeshModelingToolsExp/Public/BakeMultiMeshAttributeMapsTool.h
+++ b/Engine/Plugins/Experimental/MeshModelingToolsetExp/Source/MeshModelingToolsExp/Public/BakeMultiMeshAttributeMapsTool.h
@@ -206,6 +206,8 @@ protected:
 	// Begin UBakeMeshAttributeMapsToolBase interface
 	virtual void UpdateResult() override;
 	virtual void UpdateVisualization() override;
+
+	virtual void GatherAnalytics(FBakeAnalytics::FMeshSettings& Data) override;
 	// End UBakeMeshAttributeMapsToolBase interface
 	
 
@@ -214,7 +216,6 @@ protected:
 
 	UE::Geometry::FMeshSceneAdapter DetailMeshScene;
 
-	bool bInputsDirty = false;
 	void UpdateOnModeChange();
 
 	void InvalidateResults();
@@ -228,4 +229,10 @@ protected:
 	TArray<TSharedPtr<UE::Geometry::TImageBuilder<FVector4f>>> CachedColorImages;
 	TArray<int> CachedColorUVLayers;
 	FTextureImageMap CachedMeshToColorImagesMap;
+
+	// Analytics
+	virtual FString GetAnalyticsEventName() const override
+	{
+		return TEXT("BakeAll");
+	}
 };
diff --git a/Engine/Plugins/Runtime/GeometryProcessing/Source/DynamicMesh/Private/Sampling/MeshMapBaker.cpp b/Engine/Plugins/Runtime/GeometryProcessing/Source/DynamicMesh/Private/Sampling/MeshMapBaker.cpp
--- a/Engine/Plugins/Runtime/GeometryProcessing/Source/DynamicMesh/Private/Sampling/MeshMapBaker.cpp
+++ b/Engine/Plugins/Runtime/GeometryProcessing/Source/DynamicMesh/Private/Sampling/MeshMapBaker.cpp
@@ -5,6 +5,8 @@
 #include "Sampling/MeshMapBakerQueue.h"
 #include "Image/ImageOccupancyMap.h"
 #include "Image/ImageTile.h"
+#include "Algo/Count.h"
+#include "ProfilingDebugging/ScopedTimers.h"
 
 #include "ExplicitUseGeometryMathTypes.h"		// using UE::Geometry::(math types)
 using namespace UE::Geometry;
@@ -13,8 +15,8 @@ using namespace UE::Geometry;
 // FMeshMapBaker
 //
 
-static const float BoxFilterRadius = 0.5f;
-static const float BCFilterRadius = 0.769f;
+static constexpr float BoxFilterRadius = 0.5f;
+static constexpr float BCFilterRadius = 0.769f;
 
 FBoxFilter FMeshMapBaker::BoxFilter(BoxFilterRadius);
 FBSplineFilter FMeshMapBaker::BSplineFilter(BCFilterRadius);
@@ -114,6 +116,9 @@ void FMeshMapBaker::InitBakeDefaults()
 void FMeshMapBaker::Bake()
 {
 	TRACE_CPUPROFILER_EVENT_SCOPE(FMeshMapBaker::Bake);
+
+	BakeAnalytics.Reset();
+	FScopedDurationTimer TotalBakeTimer(BakeAnalytics.TotalBakeDuration);
 	
 	if (Bakers.IsEmpty())
 	{
@@ -298,6 +303,8 @@ void FMeshMapBaker::Bake()
 		OccupancyMap.ComputeFromUVSpaceMesh(FlatMesh, [this](int32 TriangleID) { return FlatMesh.GetTriangleGroup(TriangleID); });
 		GutterTexelsPerTile[TileIdx] = OccupancyMap.GutterTexels;
 
+		BakeAnalytics.NumBakedPixels += Algo::CountIf(OccupancyMap.TexelInteriorSamples, [](const int32 Value){ return Value > 0; });
+
 		FMeshMapTileBuffer* TileBuffer = new FMeshMapTileBuffer(PaddedTile, BakeSampleBufferSize);
 
 		// Calculate interior texels
@@ -369,73 +376,80 @@ void FMeshMapBaker::Bake()
 		WriteQueuedOutput(OutputQueue);
 	}
 
-	// Normalize and convert ImageTileBuffer data to color data.
-	ParallelFor(NumTiles, [this, &Tiles, &ImageTileBuffer](int32 TileIdx)
 	{
-		TRACE_CPUPROFILER_EVENT_SCOPE(FMeshMapBaker::WriteToImageBuffer);
+		FScopedDurationTimer WriteToImageTimer(BakeAnalytics.WriteToImageDuration);
 		
-		const FImageTile Tile = Tiles.GetTile(TileIdx);
-		const int TileWidth = Tile.GetWidth();
-		const int TileHeight = Tile.GetHeight();
-		for (FVector2i TileCoords(0,0); TileCoords.Y < TileHeight; ++TileCoords.Y)
+		// Normalize and convert ImageTileBuffer data to color data.
+		ParallelFor(NumTiles, [this, &Tiles, &ImageTileBuffer](int32 TileIdx)
 		{
-			for (TileCoords.X = 0; TileCoords.X < TileWidth; ++TileCoords.X)
+			TRACE_CPUPROFILER_EVENT_SCOPE(FMeshMapBaker::WriteToImageBuffer);
+		
+			const FImageTile Tile = Tiles.GetTile(TileIdx);
+			const int TileWidth = Tile.GetWidth();
+			const int TileHeight = Tile.GetHeight();
+			for (FVector2i TileCoords(0,0); TileCoords.Y < TileHeight; ++TileCoords.Y)
 			{
-				if (CancelF())
+				for (TileCoords.X = 0; TileCoords.X < TileWidth; ++TileCoords.X)
 				{
-					return;
-				}
+					if (CancelF())
+					{
+						return;
+					}
 
-				const FVector2i ImageCoords = Tile.GetSourceCoords(TileCoords);
-				const int64 ImageLinearIdx = Dimensions.GetIndex(ImageCoords);
-				const float& PixelWeight = ImageTileBuffer.GetPixelWeight(ImageLinearIdx);
-				float* PixelBuffer = ImageTileBuffer.GetPixel(ImageLinearIdx);
+					const FVector2i ImageCoords = Tile.GetSourceCoords(TileCoords);
+					const int64 ImageLinearIdx = Dimensions.GetIndex(ImageCoords);
+					const float& PixelWeight = ImageTileBuffer.GetPixelWeight(ImageLinearIdx);
+					float* PixelBuffer = ImageTileBuffer.GetPixel(ImageLinearIdx);
 
-				auto WriteToPixel = [this, &PixelBuffer, &ImageLinearIdx](TArray<int32>& BakeIds, float OneOverWeight)
-				{
-					for (const int32 Idx : BakeIds)
+					auto WriteToPixel = [this, &PixelBuffer, &ImageLinearIdx](TArray<int32>& BakeIds, float OneOverWeight)
 					{
-						const FMeshMapEvaluator::FEvaluationContext& Context = BakeContexts[Idx];
-						const int32 NumData = Context.DataLayout.Num();
-						const int32 ResultOffset = BakeOffsets[Idx];
-						for (int32 DataIdx = 0; DataIdx < NumData; ++DataIdx)
+						for (const int32 Idx : BakeIds)
 						{
-							const int32 ResultIdx = ResultOffset + DataIdx;
-							const int32 Offset = BakeSampleOffsets[ResultIdx];
-							float* BufferPtr = &PixelBuffer[Offset];
-							const FMeshMapEvaluator::EComponents Stride = Context.DataLayout[DataIdx];
-
-							// Apply weight to raw float data.
-							for (int32 FloatIdx = 0; FloatIdx < (int32)Stride; ++FloatIdx)
+							const FMeshMapEvaluator::FEvaluationContext& Context = BakeContexts[Idx];
+							const int32 NumData = Context.DataLayout.Num();
+							const int32 ResultOffset = BakeOffsets[Idx];
+							for (int32 DataIdx = 0; DataIdx < NumData; ++DataIdx)
 							{
-								BufferPtr[FloatIdx] *= OneOverWeight;
+								const int32 ResultIdx = ResultOffset + DataIdx;
+								const int32 Offset = BakeSampleOffsets[ResultIdx];
+								float* BufferPtr = &PixelBuffer[Offset];
+								const FMeshMapEvaluator::EComponents Stride = Context.DataLayout[DataIdx];
+
+								// Apply weight to raw float data.
+								for (int32 FloatIdx = 0; FloatIdx < (int32)Stride; ++FloatIdx)
+								{
+									BufferPtr[FloatIdx] *= OneOverWeight;
+								}
+
+								// Convert float data to color.
+								FVector4f& Pixel = BakeResults[ResultIdx]->GetPixel(ImageLinearIdx);
+								Context.EvaluateColor(DataIdx, BufferPtr, Pixel, Context.EvalData);
 							}
-
-							// Convert float data to color.
-							FVector4f& Pixel = BakeResults[ResultIdx]->GetPixel(ImageLinearIdx);
-							Context.EvaluateColor(DataIdx, BufferPtr, Pixel, Context.EvalData);
 						}
-					}
-				};
+					};
 				
-				if (PixelWeight > 0.0)
-				{
-					WriteToPixel(BakeAccumulateLists[(int32)FMeshMapEvaluator::EAccumulateMode::Add], 1.0f / PixelWeight);
+					if (PixelWeight > 0.0)
+					{
+						WriteToPixel(BakeAccumulateLists[(int32)FMeshMapEvaluator::EAccumulateMode::Add], 1.0f / PixelWeight);
+					}
+					WriteToPixel(BakeAccumulateLists[(int32)FMeshMapEvaluator::EAccumulateMode::Overwrite], 1.0f);
 				}
-				WriteToPixel(BakeAccumulateLists[(int32)FMeshMapEvaluator::EAccumulateMode::Overwrite], 1.0f);
 			}
-		}
-	}, !bParallel ? EParallelForFlags::ForceSingleThread : EParallelForFlags::None);
+		}, !bParallel ? EParallelForFlags::ForceSingleThread : EParallelForFlags::None);
+	}
 
 	// Gutter Texel processing
 	if( bGutterEnabled )
 	{
+		FScopedDurationTimer WriteToGutterTimer(BakeAnalytics.WriteToGutterDuration);
+		
 		const int32 NumResults = BakeResults.Num();
 		ParallelFor(NumTiles, [this, &NumResults, &GutterTexelsPerTile](int32 TileIdx)
 		{
 			TRACE_CPUPROFILER_EVENT_SCOPE(FMeshMapBaker::WriteGutterPixels);
-			
-			for (int64 GutterIdx = 0; GutterIdx < GutterTexelsPerTile[TileIdx].Num(); ++GutterIdx)
+
+			const int NumGutter = GutterTexelsPerTile[TileIdx].Num();
+			for (int64 GutterIdx = 0; GutterIdx < NumGutter; ++GutterIdx)
 			{
 				int64 GutterPixelTo;
 				int64 GutterPixelFrom;
@@ -445,6 +459,8 @@ void FMeshMapBaker::Bake()
 					BakeResults[Idx]->CopyPixel(GutterPixelFrom, GutterPixelTo);
 				}
 			}
+
+			BakeAnalytics.NumGutterPixels += NumGutter;
 		}, !bParallel ? EParallelForFlags::ForceSingleThread : EParallelForFlags::None);
 	}
 }
@@ -544,7 +560,7 @@ int32 FMeshMapBaker::AddEvaluator(const TSharedPtr<FMeshMapEvaluator, ESPMode::T
 	return Bakers.Add(Eval);
 }
 
-FMeshMapEvaluator* FMeshMapBaker::GetEvaluator(const int32 EvalIdx)
+FMeshMapEvaluator* FMeshMapBaker::GetEvaluator(const int32 EvalIdx) const
 {
 	return Bakers[EvalIdx].Get();
 }
diff --git a/Engine/Plugins/Runtime/GeometryProcessing/Source/DynamicMesh/Public/Sampling/MeshBakerCommon.h b/Engine/Plugins/Runtime/GeometryProcessing/Source/DynamicMesh/Public/Sampling/MeshBakerCommon.h
--- a/Engine/Plugins/Runtime/GeometryProcessing/Source/DynamicMesh/Public/Sampling/MeshBakerCommon.h
+++ b/Engine/Plugins/Runtime/GeometryProcessing/Source/DynamicMesh/Public/Sampling/MeshBakerCommon.h
@@ -102,6 +102,9 @@ public:
 	/** Iterate over each mesh in the detail set. */
 	virtual void ProcessMeshes(TFunctionRef<void(const void*)> ProcessFn) const = 0;
 
+	/** @return the triangle count of a given mesh. */
+	virtual int32 GetTriangleCount(const void* Mesh) const = 0;
+
 	/** Associate a color map and UV layer index for a given mesh in the detail set */
 	virtual void SetColorMap(const void* Mesh, const FBakeDetailTexture& Map) = 0;
 
@@ -328,6 +331,12 @@ public:
 		ProcessFn(DetailMesh);
 	}
 
+	virtual int32 GetTriangleCount(const void* Mesh) const override
+	{
+		const FDynamicMesh3* DynamicMesh = static_cast<const FDynamicMesh3*>(Mesh);
+		return DynamicMesh->TriangleCount();
+	}
+
 	virtual void SetColorMap(const void* Mesh, const FBakeDetailTexture& Map) override
 	{
 		DetailColorMap = Map;
diff --git a/Engine/Plugins/Runtime/GeometryProcessing/Source/DynamicMesh/Public/Sampling/MeshMapBaker.h b/Engine/Plugins/Runtime/GeometryProcessing/Source/DynamicMesh/Public/Sampling/MeshMapBaker.h
--- a/Engine/Plugins/Runtime/GeometryProcessing/Source/DynamicMesh/Public/Sampling/MeshMapBaker.h
+++ b/Engine/Plugins/Runtime/GeometryProcessing/Source/DynamicMesh/Public/Sampling/MeshMapBaker.h
@@ -2,6 +2,8 @@
 
 #pragma once
 
+#include <atomic>
+
 #include "Sampling/MeshBaseBaker.h"
 #include "Sampling/MeshMapEvaluator.h"
 #include "Sampling/MeshSurfaceSampler.h"
@@ -34,7 +36,7 @@ public:
 	int32 AddEvaluator(const TSharedPtr<FMeshMapEvaluator, ESPMode::ThreadSafe>& Eval);
 
 	/** @return the evaluator at the given index. */
-	FMeshMapEvaluator* GetEvaluator(int32 EvalIdx);
+	FMeshMapEvaluator* GetEvaluator(int32 EvalIdx) const;
 
 	/** @return the number of bake evaluators on this baker. */
 	int32 NumEvaluators() const;
@@ -75,6 +77,28 @@ public:
 	EBakeFilterType GetFilter() const { return FilterType; }
 	int32 GetTileSize() const { return TileSize; }
 
+	//
+	// Analytics
+	//
+	struct FBakeAnalytics
+	{
+		double TotalBakeDuration = 0.0;
+		double WriteToImageDuration = 0.0;
+		double WriteToGutterDuration = 0.0;
+		std::atomic<int64> NumBakedPixels = 0;
+		std::atomic<int64> NumGutterPixels = 0;
+
+		void Reset()
+		{
+			TotalBakeDuration = 0.0;
+			WriteToImageDuration = 0.0;
+			WriteToGutterDuration = 0.0;
+			NumBakedPixels = 0;
+			NumGutterPixels = 0;
+		}
+	};
+	FBakeAnalytics BakeAnalytics;
+
 protected:
 	/** Evaluate this sample. */
 	void BakeSample(
diff --git a/Engine/Plugins/Runtime/MeshModelingToolset/Source/ModelingComponents/Private/Scene/MeshSceneAdapter.cpp b/Engine/Plugins/Runtime/MeshModelingToolset/Source/ModelingComponents/Private/Scene/MeshSceneAdapter.cpp
--- a/Engine/Plugins/Runtime/MeshModelingToolset/Source/ModelingComponents/Private/Scene/MeshSceneAdapter.cpp
+++ b/Engine/Plugins/Runtime/MeshModelingToolset/Source/ModelingComponents/Private/Scene/MeshSceneAdapter.cpp
@@ -282,7 +282,7 @@ public:
 		return true;
 	}
 
-	virtual int32 GetTriangleCount() override
+	virtual int32 GetTriangleCount() const override
 	{
 		return Mesh.TriangleCount();
 	}
@@ -670,7 +670,7 @@ public:
 		return false;
 	}
 
-	virtual int32 GetTriangleCount() override
+	virtual int32 GetTriangleCount() const override
 	{
 		if (!SourceMesh) return 0;
 		return Adapter->TriangleCount();
diff --git a/Engine/Plugins/Runtime/MeshModelingToolset/Source/ModelingComponents/Public/Scene/MeshSceneAdapter.h b/Engine/Plugins/Runtime/MeshModelingToolset/Source/ModelingComponents/Public/Scene/MeshSceneAdapter.h
--- a/Engine/Plugins/Runtime/MeshModelingToolset/Source/ModelingComponents/Public/Scene/MeshSceneAdapter.h
+++ b/Engine/Plugins/Runtime/MeshModelingToolset/Source/ModelingComponents/Public/Scene/MeshSceneAdapter.h
@@ -129,7 +129,7 @@ public:
 	virtual bool Build(const FMeshSceneAdapterBuildOptions& BuildOptions) = 0;
 
 	/*** @return triangle count for this mesh */
-	virtual int32 GetTriangleCount() = 0;
+	virtual int32 GetTriangleCount() const = 0;
 
 	/** Calculate bounding box for this Mesh */
 	virtual FAxisAlignedBox3d GetWorldBounds(TFunctionRef<FVector3d(const FVector3d&)> LocalToWorldFunc) = 0;
