ModelingTools: Add analytics to BakeVtx.

#rnx

diff --git a/Engine/Plugins/Experimental/MeshModelingToolsetExp/Source/MeshModelingToolsExp/Private/BakeMeshAttributeVertexTool.cpp b/Engine/Plugins/Experimental/MeshModelingToolsetExp/Source/MeshModelingToolsExp/Private/BakeMeshAttributeVertexTool.cpp
--- a/Engine/Plugins/Experimental/MeshModelingToolsetExp/Source/MeshModelingToolsExp/Private/BakeMeshAttributeVertexTool.cpp
+++ b/Engine/Plugins/Experimental/MeshModelingToolsetExp/Source/MeshModelingToolsExp/Private/BakeMeshAttributeVertexTool.cpp
@@ -25,6 +25,8 @@
 #include "ModelingToolTargetUtil.h"
 #include "AssetUtils/Texture2DUtil.h"
 
+#include "EngineAnalytics.h"
+
 #include "ExplicitUseGeometryMathTypes.h"		// using UE::Geometry::(math types)
 using namespace UE::Geometry;
 
@@ -348,13 +350,15 @@ void UBakeMeshAttributeVertexTool::Setup()
 
 	UpdateOnModeChange();
 
-	bDetailMeshValid = false;
+	UpdateDetailMesh();
 
 	SetToolDisplayName(LOCTEXT("ToolName", "Bake Vertex Colors"));
 	GetToolManager()->DisplayMessage(
 		LOCTEXT("OnStartTool",
 		        "Bake Vertex Colors. Select Bake Mesh (LowPoly) first, then (optionally) Detail Mesh second."),
 		EToolMessageLevel::UserNotification);
+
+	GatherAnalytics(BakeAnalytics.MeshSettings);
 }
 
 void UBakeMeshAttributeVertexTool::Shutdown(EToolShutdownType ShutdownType)
@@ -395,6 +399,8 @@ void UBakeMeshAttributeVertexTool::Shutdown(EToolShutdownType ShutdownType)
 		PreviewMesh->Disconnect();
 		PreviewMesh = nullptr;
 	}
+
+	RecordAnalytics(BakeAnalytics, TEXT("BakeVertex"));
 }
 
 void UBakeMeshAttributeVertexTool::OnTick(float DeltaTime)
@@ -469,6 +475,7 @@ void UBakeMeshAttributeVertexTool::UpdateDetailMesh()
 	DetailSpatial->SetMesh(DetailMesh.Get(), true);
 
 	OpState = EBakeOpState::Evaluate;
+	bDetailMeshValid = true;
 	DetailMeshTimestamp++;
 }
 
@@ -600,6 +607,8 @@ void UBakeMeshAttributeVertexTool::UpdateColorTopology()
 	{
 		BaseMesh.Attributes()->PrimaryColors()->Copy(*Mesh.Attributes()->PrimaryColors());
 	});
+
+	bColorTopologyValid = true;
 }
 
 void UBakeMeshAttributeVertexTool::UpdateResult()
@@ -607,13 +616,11 @@ void UBakeMeshAttributeVertexTool::UpdateResult()
 	if (!bDetailMeshValid)
 	{
 		UpdateDetailMesh();
-		bDetailMeshValid = true;
 	}
 
 	if (!bColorTopologyValid)
 	{
 		UpdateColorTopology();
-		bColorTopologyValid = true;
 	}
 
 	if (OpState == EBakeOpState::Complete)
@@ -626,6 +633,8 @@ void UBakeMeshAttributeVertexTool::UpdateResult()
 
 	FBakeSettings BakeSettings;
 	BakeSettings.VertexMode = Settings->VertexMode;
+	BakeSettings.bSplitAtNormalSeams = Settings->bSplitAtNormalSeams;
+	BakeSettings.bSplitAtUVSeams = Settings->bSplitAtUVSeams;
 	BakeSettings.bUseWorldSpace = Settings->bUseWorldSpace;
 	BakeSettings.Thickness = Settings->Thickness;
 	if (!(BakeSettings == CachedBakeSettings))
@@ -741,6 +750,8 @@ void UBakeMeshAttributeVertexTool::OnResultUpdated(const TUniquePtr<FMeshVertexB
 	UpdateVisualization();
 
 	OpState = EBakeOpState::Complete;
+
+	GatherAnalytics(*NewResult, CachedBakeSettings, CachedColorSettings, CachedChannelSettings, BakeAnalytics);
 }
 
 EBakeOpState UBakeMeshAttributeVertexTool::UpdateResult_Normal()
@@ -926,5 +937,190 @@ EBakeOpState UBakeMeshAttributeVertexTool::UpdateResult_MultiTexture()
 }
 
 
+void UBakeMeshAttributeVertexTool::GatherAnalytics(FBakeAnalytics::FMeshSettings& Data)
+{
+	if (!FEngineAnalytics::IsAvailable())
+	{
+		return;
+	}
+	
+	Data.NumTargetMeshVerts = BaseMesh.VertexCount();
+	Data.NumTargetMeshTris = BaseMesh.TriangleCount();
+	Data.NumDetailMesh = 1;
+	Data.NumDetailMeshTris = DetailMesh->TriangleCount();
+}
+
+
+void UBakeMeshAttributeVertexTool::GatherAnalytics(
+	const FMeshVertexBaker& Result,
+	const FBakeSettings& Settings,
+	const FBakeColorSettings& ColorSettings,
+	const FBakeChannelSettings& ChannelSettings,
+	FBakeAnalytics& Data)
+{
+	if (!FEngineAnalytics::IsAvailable())
+	{
+		return;
+	}
+	
+	Data.TotalBakeDuration = Result.TotalBakeDuration;
+	Data.BakeSettings = Settings;
+	Data.BakeColorSettings = ColorSettings;
+	Data.BakeChannelSettings = ChannelSettings;
+
+	auto GatherEvaluatorData = [&Data](const FMeshMapEvaluator* Eval)
+	{
+		if (Eval)
+		{
+			switch(Eval->Type())
+			{
+			case EMeshMapEvaluatorType::Occlusion:
+			{
+				const FMeshOcclusionMapEvaluator* OcclusionEval = static_cast<const FMeshOcclusionMapEvaluator*>(Eval);
+				Data.OcclusionSettings.OcclusionRays = OcclusionEval->NumOcclusionRays;
+				Data.OcclusionSettings.MaxDistance = OcclusionEval->MaxDistance;
+				Data.OcclusionSettings.SpreadAngle = OcclusionEval->SpreadAngle;
+				Data.OcclusionSettings.BiasAngle = OcclusionEval->BiasAngleDeg;
+				break;
+			}
+			case EMeshMapEvaluatorType::Curvature:
+			{
+				const FMeshCurvatureMapEvaluator* CurvatureEval = static_cast<const FMeshCurvatureMapEvaluator*>(Eval);
+				Data.CurvatureSettings.CurvatureType = static_cast<int>(CurvatureEval->UseCurvatureType);
+				Data.CurvatureSettings.RangeMultiplier = CurvatureEval->RangeScale;
+				Data.CurvatureSettings.MinRangeMultiplier = CurvatureEval->MinRangeScale;
+				Data.CurvatureSettings.ColorMode = static_cast<int>(CurvatureEval->UseColorMode);
+				Data.CurvatureSettings.ClampMode = static_cast<int>(CurvatureEval->UseClampMode);
+				break;
+			}
+			default:
+				break;
+			};
+		}
+	};
+
+	if (Result.BakeMode == FMeshVertexBaker::EBakeMode::Color)
+	{
+		GatherEvaluatorData(Result.ColorEvaluator.Get());
+	}
+	else // Result.BakeMode == FMeshVertexBaker::EBakeMode::Channel
+	{
+		for (int EvalId = 0; EvalId < 4; ++EvalId)
+		{
+			GatherEvaluatorData(Result.ChannelEvaluators[EvalId].Get());
+		}
+	}
+}
+
+
+void UBakeMeshAttributeVertexTool::RecordAnalytics(const FBakeAnalytics& Data, const FString& EventName)
+{
+	if (!FEngineAnalytics::IsAvailable())
+	{
+		return;
+	}
+	
+	TArray<FAnalyticsEventAttribute> Attributes;
+
+	// General
+	Attributes.Add(FAnalyticsEventAttribute(TEXT("Bake.Duration.Total.Seconds"), Data.TotalBakeDuration));
+
+	// Mesh data
+	Attributes.Add(FAnalyticsEventAttribute(TEXT("Input.TargetMesh.NumTriangles"), Data.MeshSettings.NumTargetMeshTris));
+	Attributes.Add(FAnalyticsEventAttribute(TEXT("Input.TargetMesh.NumVertices"), Data.MeshSettings.NumTargetMeshVerts));
+	Attributes.Add(FAnalyticsEventAttribute(TEXT("Input.DetailMesh.NumMeshes"), Data.MeshSettings.NumDetailMesh));
+	Attributes.Add(FAnalyticsEventAttribute(TEXT("Input.DetailMesh.NumTriangles"), Data.MeshSettings.NumDetailMeshTris));
+
+	// Bake settings
+	Attributes.Add(FAnalyticsEventAttribute(TEXT("Settings.Thickness"), Data.BakeSettings.Thickness));
+	Attributes.Add(FAnalyticsEventAttribute(TEXT("Settings.Split.NormalSeams"), Data.BakeSettings.bSplitAtNormalSeams));
+	Attributes.Add(FAnalyticsEventAttribute(TEXT("Settings.Split.UVSeams"), Data.BakeSettings.bSplitAtUVSeams));
+	Attributes.Add(FAnalyticsEventAttribute(TEXT("Settings.Thickness"), Data.BakeSettings.Thickness));
+	Attributes.Add(FAnalyticsEventAttribute(TEXT("Settings.UseWorldSpace"), Data.BakeSettings.bUseWorldSpace));
+
+	const FString OutputType = Data.BakeSettings.VertexMode == EBakeVertexMode::Color ? TEXT("RGBA") : TEXT("PerChannel");
+	Attributes.Add(FAnalyticsEventAttribute(TEXT("Settings.Output.Type"), OutputType));
+
+	auto RecordAmbientOcclusionSettings = [&Attributes, &Data](const FString& ModeName)
+	{
+		Attributes.Add(FAnalyticsEventAttribute(FString::Printf(TEXT("Settings.Output.%s.AmbientOcclusion.OcclusionRays"), *ModeName), Data.OcclusionSettings.OcclusionRays));
+		Attributes.Add(FAnalyticsEventAttribute(FString::Printf(TEXT("Settings.Output.%s.AmbientOcclusion.MaxDistance"), *ModeName), Data.OcclusionSettings.MaxDistance));
+		Attributes.Add(FAnalyticsEventAttribute(FString::Printf(TEXT("Settings.Output.%s.AmbientOcclusion.SpreadAngle"), *ModeName), Data.OcclusionSettings.SpreadAngle));
+		Attributes.Add(FAnalyticsEventAttribute(FString::Printf(TEXT("Settings.Output.%s.AmbientOcclusion.BiasAngle"), *ModeName), Data.OcclusionSettings.BiasAngle));
+	};
+
+	auto RecordBentNormalSettings = [&Attributes, &Data](const FString& ModeName)
+	{
+		Attributes.Add(FAnalyticsEventAttribute(FString::Printf(TEXT("Settings.Output.%s.BentNormal.OcclusionRays"), *ModeName), Data.OcclusionSettings.OcclusionRays));
+		Attributes.Add(FAnalyticsEventAttribute(FString::Printf(TEXT("Settings.Output.%s.BentNormal.MaxDistance"), *ModeName), Data.OcclusionSettings.MaxDistance));
+		Attributes.Add(FAnalyticsEventAttribute(FString::Printf(TEXT("Settings.Output.%s.BentNormal.SpreadAngle"), *ModeName), Data.OcclusionSettings.SpreadAngle));
+	};
+
+	auto RecordCurvatureSettings = [&Attributes, &Data](const FString& ModeName)
+	{
+		Attributes.Add(FAnalyticsEventAttribute(FString::Printf(TEXT("Settings.Output.%s.Curvature.CurvatureType"), *ModeName), Data.CurvatureSettings.CurvatureType));
+		Attributes.Add(FAnalyticsEventAttribute(FString::Printf(TEXT("Settings.Output.%s.Curvature.RangeMultiplier"), *ModeName), Data.CurvatureSettings.RangeMultiplier));
+		Attributes.Add(FAnalyticsEventAttribute(FString::Printf(TEXT("Settings.Output.%s.Curvature.MinRangeMultiplier"), *ModeName), Data.CurvatureSettings.MinRangeMultiplier));
+		Attributes.Add(FAnalyticsEventAttribute(FString::Printf(TEXT("Settings.Output.%s.Curvature.ClampMode"), *ModeName), Data.CurvatureSettings.ClampMode));
+		Attributes.Add(FAnalyticsEventAttribute(FString::Printf(TEXT("Settings.Output.%s.Curvature.ColorMode"), *ModeName), Data.CurvatureSettings.ColorMode));
+	};
+
+	if (Data.BakeSettings.VertexMode == EBakeVertexMode::Color)
+	{
+		const FString OutputName(TEXT("RGBA"));
+
+		FString OutputTypeName = StaticEnum<EBakeVertexTypeColor>()->GetNameStringByIndex(static_cast<int>(Data.BakeColorSettings.BakeType));
+		Attributes.Add(FAnalyticsEventAttribute(FString::Printf(TEXT("Settings.Output.%s.Type"), *OutputName), OutputTypeName));
+
+		switch (Data.BakeColorSettings.BakeType)
+		{
+		case EBakeVertexTypeColor::AmbientOcclusion:
+			RecordAmbientOcclusionSettings(OutputName);
+			break;
+		case EBakeVertexTypeColor::BentNormal:
+			RecordBentNormalSettings(OutputName);
+			break;
+		case EBakeVertexTypeColor::Curvature:
+			RecordCurvatureSettings(OutputName);
+			break;
+		default:
+			break;
+		}
+	}
+	else
+	{
+		ensure(Data.BakeSettings.VertexMode == EBakeVertexMode::PerChannel);
+		for (int EvalId = 0; EvalId < 4; ++EvalId)
+		{
+			FString OutputName = StaticEnum<EBakeVertexChannel>()->GetNameStringByIndex(EvalId);
+			FString OutputTypeName = StaticEnum<EBakeVertexTypeChannel>()->GetNameStringByIndex(static_cast<int>(Data.BakeChannelSettings.BakeType[EvalId]));
+			Attributes.Add(FAnalyticsEventAttribute(FString::Printf(TEXT("Settings.Output.%s.Type"), *OutputName), OutputTypeName));
+
+			switch (Data.BakeChannelSettings.BakeType[EvalId])
+			{
+			case EBakeVertexTypeChannel::AmbientOcclusion:
+				RecordAmbientOcclusionSettings(OutputName);
+				break;
+			case EBakeVertexTypeChannel::Curvature:
+				RecordCurvatureSettings(OutputName);
+				break;
+			default:
+				break;
+			}
+		}
+	}
+
+	FEngineAnalytics::GetProvider().RecordEvent(FString(TEXT("Editor.Usage.ModelingTools.")) + EventName, Attributes);
+
+	constexpr bool bLogAnalytics = false; 
+	if constexpr (bLogAnalytics)
+	{
+		for (const FAnalyticsEventAttribute& Attr : Attributes)
+		{
+			UE_LOG(LogGeometry, Log, TEXT("[%s] %s = %s"), *EventName, *Attr.GetName(), *Attr.GetValue());
+		}
+	}
+}
+
 
 #undef LOCTEXT_NAMESPACE
diff --git a/Engine/Plugins/Experimental/MeshModelingToolsetExp/Source/MeshModelingToolsExp/Public/BakeMeshAttributeVertexTool.h b/Engine/Plugins/Experimental/MeshModelingToolsetExp/Source/MeshModelingToolsExp/Public/BakeMeshAttributeVertexTool.h
--- a/Engine/Plugins/Experimental/MeshModelingToolsetExp/Source/MeshModelingToolsExp/Public/BakeMeshAttributeVertexTool.h
+++ b/Engine/Plugins/Experimental/MeshModelingToolsetExp/Source/MeshModelingToolsExp/Public/BakeMeshAttributeVertexTool.h
@@ -312,5 +312,54 @@ protected:
 
 	TMap<int32, TSharedPtr<UE::Geometry::TImageBuilder<FVector4f>, ESPMode::ThreadSafe>> CachedMultiTextures;
 	EBakeOpState UpdateResult_MultiTexture();
+
+	//
+	// Analytics
+	//
+	struct FBakeAnalytics
+	{
+		double TotalBakeDuration = 0.0;
+
+		struct FMeshSettings
+		{
+			int32 NumTargetMeshVerts = 0;
+			int32 NumTargetMeshTris = 0;
+			int32 NumDetailMesh = 0;
+			int64 NumDetailMeshTris = 0;
+		};
+		FMeshSettings MeshSettings;
+
+		FBakeSettings BakeSettings;
+		FBakeColorSettings BakeColorSettings;
+		FBakeChannelSettings BakeChannelSettings;
+		FOcclusionMapSettings OcclusionSettings;
+		FCurvatureMapSettings CurvatureSettings;
+	};
+	FBakeAnalytics BakeAnalytics;
+
+	/**
+	 * Computes the NumTargetMeshTris, NumDetailMesh and NumDetailMeshTris analytics.
+	 * @param Data the mesh analytics data to compute
+	 */
+	virtual void GatherAnalytics(FBakeAnalytics::FMeshSettings& Data);
+
+	/**
+	 * Records bake timing and settings data for analytics.
+	 * @param Result the result of the bake.
+	 * @param Settings The bake settings used for the bake.
+	 * @param Data the output bake analytics struct.
+	 */
+	static void GatherAnalytics(const UE::Geometry::FMeshVertexBaker& Result,
+								const FBakeSettings& Settings,
+								const FBakeColorSettings& ColorSettings,
+								const FBakeChannelSettings& ChannelSettings,
+								FBakeAnalytics& Data);
+
+	/**
+	 * Outputs an analytics event using the given analytics struct.
+	 * @param Data the bake analytics struct to output.
+	 * @param EventName the name of the analytics event to output.
+	 */
+	static void RecordAnalytics(const FBakeAnalytics& Data, const FString& EventName);
 };
 
diff --git a/Engine/Plugins/Runtime/GeometryProcessing/Source/DynamicMesh/Private/Sampling/MeshVertexBaker.cpp b/Engine/Plugins/Runtime/GeometryProcessing/Source/DynamicMesh/Private/Sampling/MeshVertexBaker.cpp
--- a/Engine/Plugins/Runtime/GeometryProcessing/Source/DynamicMesh/Private/Sampling/MeshVertexBaker.cpp
+++ b/Engine/Plugins/Runtime/GeometryProcessing/Source/DynamicMesh/Private/Sampling/MeshVertexBaker.cpp
@@ -3,6 +3,7 @@
 #include "Sampling/MeshVertexBaker.h"
 #include "Sampling/MeshBakerCommon.h"
 #include "DynamicMesh/MeshNormals.h"
+#include "ProfilingDebugging/ScopedTimers.h"
 
 #include "ExplicitUseGeometryMathTypes.h"		// using UE::Geometry::(math types)
 using namespace UE::Geometry;
@@ -12,6 +13,9 @@ FMeshConstantMapEvaluator FMeshVertexBaker::OneEvaluator(1.0f);
 
 void FMeshVertexBaker::Bake()
 {
+	TotalBakeDuration = 0.0;
+	FScopedDurationTimer Timer(TotalBakeDuration);
+	
 	if (!ensure(TargetMesh && TargetMesh->HasAttributes() && TargetMesh->Attributes()->HasPrimaryColors()) ||
 		!ensure(DetailSampler))
 	{
diff --git a/Engine/Plugins/Runtime/GeometryProcessing/Source/DynamicMesh/Public/Sampling/MeshVertexBaker.h b/Engine/Plugins/Runtime/GeometryProcessing/Source/DynamicMesh/Public/Sampling/MeshVertexBaker.h
--- a/Engine/Plugins/Runtime/GeometryProcessing/Source/DynamicMesh/Public/Sampling/MeshVertexBaker.h
+++ b/Engine/Plugins/Runtime/GeometryProcessing/Source/DynamicMesh/Public/Sampling/MeshVertexBaker.h
@@ -30,6 +30,11 @@ public:
 	/** Evaluators to use for per channel (RGBA) bakes */
 	TSharedPtr<FMeshMapEvaluator, ESPMode::ThreadSafe> ChannelEvaluators[4];
 
+	//
+	// Analytics
+	//
+	double TotalBakeDuration = 0.0;
+
 public:	
 	//
 	// Bake
