ModelingTools: Replace duplicate BakeMultiMapType enum with ValidEnumValues.

#rnx

diff --git a/Engine/Plugins/Experimental/MeshModelingToolsetExp/Source/MeshModelingToolsExp/Public/BakeMultiMeshAttributeMapsTool.h b/Engine/Plugins/Experimental/MeshModelingToolsetExp/Source/MeshModelingToolsExp/Public/BakeMultiMeshAttributeMapsTool.h
--- a/Engine/Plugins/Experimental/MeshModelingToolsetExp/Source/MeshModelingToolsExp/Public/BakeMultiMeshAttributeMapsTool.h
+++ b/Engine/Plugins/Experimental/MeshModelingToolsetExp/Source/MeshModelingToolsExp/Public/BakeMultiMeshAttributeMapsTool.h
@@ -33,41 +33,6 @@ class MESHMODELINGTOOLSEXP_API UBakeMultiMeshAttributeMapsToolBuilder : public U
 };
 
 
-// This enumeration must match EBakeMapType. This duplicate enum is
-// intended to only be used for UI purposes to hide unsupported map types
-// but is internally converted to EBakeMapType.
-UENUM(meta=(Bitflags, UseEnumValuesAsMaskValuesInEditor="true"))
-enum class EBakeMultiMapType
-{
-	None                   = 0 UMETA(Hidden),
-	/* Sample normals from the detail mesh in tangent space */
-	TangentSpaceNormals    = 1 << 0,
-	/* Sample ambient occlusion from the detail mesh */
-	AmbientOcclusion       = 1 << 1 UMETA(Hidden),
-	/* Sample normals skewed towards the least occluded direction from the detail mesh */
-	BentNormals            = 1 << 2 UMETA(Hidden),
-	/* Sample mesh curvatures from the detail mesh */
-	Curvature              = 1 << 3 UMETA(Hidden),
-	/* Sample a source texture from the detail mesh UVs */
-	Texture                = 1 << 4,
-	/* Sample object space normals from the detail mesh */
-	ObjectSpaceNormals     = 1 << 5 UMETA(Hidden),
-	/* Sample object space face normals from the detail mesh */
-	FaceNormals            = 1 << 6 UMETA(Hidden),
-	/* Sample bounding box relative positions from the detail mesh */
-	Position               = 1 << 7 UMETA(Hidden),
-	/* Sample material IDs as unique colors from the detail mesh */
-	MaterialID             = 1 << 8 UMETA(DisplayName="Material ID", Hidden),
-	/* Sample a source texture per material ID on the detail mesh */
-	MultiTexture           = 1 << 9 UMETA(Hidden),
-	/* Sample the interpolated vertex colors from the detail mesh */
-	VertexColors           = 1 << 10 UMETA(Hidden),
-	Occlusion              = (AmbientOcclusion | BentNormals) UMETA(Hidden),
-	All                    = 0x7FF UMETA(Hidden)
-};
-ENUM_CLASS_FLAGS(EBakeMultiMapType);
-
-
 UCLASS()
 class MESHMODELINGTOOLSEXP_API UBakeMultiMeshAttributeMapsToolProperties : public UInteractiveToolPropertySet
 {
@@ -75,7 +40,7 @@ class MESHMODELINGTOOLSEXP_API UBakeMultiMeshAttributeMapsToolProperties : publi
 
 public:
 	/** The map types to generate */
-	UPROPERTY(EditAnywhere, Category = MapSettings, meta=(Bitmask, BitmaskEnum=EBakeMultiMapType))
+	UPROPERTY(EditAnywhere, Category = MapSettings, meta=(Bitmask, BitmaskEnum=EBakeMapType, ValidEnumValues="TangentSpaceNormal, Texture"))
 	int32 MapTypes = (int32) EBakeMapType::None;
 
 	/** The map type index to preview */
