ModelingTools: Cleanup unused Curvature settings

#rnx

diff --git a/Engine/Plugins/Experimental/MeshModelingToolsetExp/Source/MeshModelingToolsExp/Public/BakeMeshAttributeToolCommon.h b/Engine/Plugins/Experimental/MeshModelingToolsetExp/Source/MeshModelingToolsExp/Public/BakeMeshAttributeToolCommon.h
--- a/Engine/Plugins/Experimental/MeshModelingToolsetExp/Source/MeshModelingToolsExp/Public/BakeMeshAttributeToolCommon.h
+++ b/Engine/Plugins/Experimental/MeshModelingToolsetExp/Source/MeshModelingToolsExp/Public/BakeMeshAttributeToolCommon.h
@@ -235,17 +235,15 @@ struct FOcclusionMapSettings
 struct FCurvatureMapSettings
 {
 	FImageDimensions Dimensions;
-	int32 RayCount = 1;
 	int32 CurvatureType = 0;
 	float RangeMultiplier = 1.0;
 	float MinRangeMultiplier = 0.0;
 	int32 ColorMode = 0;
 	int32 ClampMode = 0;
-	float MaxDistance = 1.0;
 
 	bool operator==(const FCurvatureMapSettings& Other) const
 	{
-		return Dimensions == Other.Dimensions && RayCount == Other.RayCount && CurvatureType == Other.CurvatureType && RangeMultiplier == Other.RangeMultiplier && MinRangeMultiplier == Other.MinRangeMultiplier && ColorMode == Other.ColorMode && ClampMode == Other.ClampMode && MaxDistance == Other.MaxDistance;
+		return Dimensions == Other.Dimensions && CurvatureType == Other.CurvatureType && RangeMultiplier == Other.RangeMultiplier && MinRangeMultiplier == Other.MinRangeMultiplier && ColorMode == Other.ColorMode && ClampMode == Other.ClampMode;
 	}
 };
 
