ModelingTools: Replace Enum to string conversion tables with UEnum queries.

#rnx

diff --git a/Engine/Plugins/Experimental/MeshModelingToolsetExp/Source/MeshModelingToolsExp/Private/BakeMeshAttributeMapsTool.cpp b/Engine/Plugins/Experimental/MeshModelingToolsetExp/Source/MeshModelingToolsExp/Private/BakeMeshAttributeMapsTool.cpp
--- a/Engine/Plugins/Experimental/MeshModelingToolsetExp/Source/MeshModelingToolsExp/Private/BakeMeshAttributeMapsTool.cpp
+++ b/Engine/Plugins/Experimental/MeshModelingToolsetExp/Source/MeshModelingToolsExp/Private/BakeMeshAttributeMapsTool.cpp
@@ -916,11 +916,7 @@ void UBakeMeshAttributeMapsTool::UpdateVisualization()
 		}
 	}
 
-	// Set the preview material according to the preview name.
-	if (const EBakeMapType* PreviewType = NameToMapTypeMap.Find(Settings->MapPreview))
-	{
-		UpdatePreview(*PreviewType);
-	}
+	UpdatePreview(Settings);
 }
 
 
diff --git a/Engine/Plugins/Experimental/MeshModelingToolsetExp/Source/MeshModelingToolsExp/Private/BakeMeshAttributeMapsToolBase.cpp b/Engine/Plugins/Experimental/MeshModelingToolsetExp/Source/MeshModelingToolsExp/Private/BakeMeshAttributeMapsToolBase.cpp
--- a/Engine/Plugins/Experimental/MeshModelingToolsetExp/Source/MeshModelingToolsExp/Private/BakeMeshAttributeMapsToolBase.cpp
+++ b/Engine/Plugins/Experimental/MeshModelingToolsetExp/Source/MeshModelingToolsExp/Private/BakeMeshAttributeMapsToolBase.cpp
@@ -30,48 +30,6 @@ using namespace UE::Geometry;
 
 #define LOCTEXT_NAMESPACE "UBakeMeshAttributeMapsToolBase"
 
-// Initialize map type name map. These names should match the
-// auto-generated enum names. 
-static const char* BakeType_TangentSpaceNormals = "Tangent Space Normal";
-static const char* BakeType_AmbientOcclusion = "Ambient Occlusion";
-static const char* BakeType_BentNormals = "Bent Normal";
-static const char* BakeType_Curvature = "Curvature";
-static const char* BakeType_Texture = "Texture";
-static const char* BakeType_ObjectSpaceNormals = "Object Space Normal";
-static const char* BakeType_FaceNormals = "Face Normal";
-static const char* BakeType_Position = "Position";
-static const char* BakeType_MaterialID = "Material ID";
-static const char* BakeType_MultiTexture = "Multi Texture";
-static const char* BakeType_VertexColors = "Vertex Color";
-
-TMap<FString, EBakeMapType> UBakeMeshAttributeMapsToolBase::NameToMapTypeMap = {
-	TTuple<FString, EBakeMapType>(BakeType_TangentSpaceNormals, EBakeMapType::TangentSpaceNormal),
-	TTuple<FString, EBakeMapType>(BakeType_AmbientOcclusion, EBakeMapType::AmbientOcclusion),
-	TTuple<FString, EBakeMapType>(BakeType_BentNormals, EBakeMapType::BentNormal),
-	TTuple<FString, EBakeMapType>(BakeType_Curvature, EBakeMapType::Curvature),
-	TTuple<FString, EBakeMapType>(BakeType_Texture, EBakeMapType::Texture),
-	TTuple<FString, EBakeMapType>(BakeType_ObjectSpaceNormals, EBakeMapType::ObjectSpaceNormal),
-	TTuple<FString, EBakeMapType>(BakeType_FaceNormals, EBakeMapType::FaceNormal),
-	TTuple<FString, EBakeMapType>(BakeType_Position, EBakeMapType::Position),
-	TTuple<FString, EBakeMapType>(BakeType_MaterialID, EBakeMapType::MaterialID),
-	TTuple<FString, EBakeMapType>(BakeType_MultiTexture, EBakeMapType::MultiTexture),
-	TTuple<FString, EBakeMapType>(BakeType_VertexColors, EBakeMapType::VertexColor)
-};
-
-TMap<EBakeMapType, FString> UBakeMeshAttributeMapsToolBase::MapTypeToNameMap = {
-	TTuple<EBakeMapType, FString>(EBakeMapType::TangentSpaceNormal, BakeType_TangentSpaceNormals),
-	TTuple<EBakeMapType, FString>(EBakeMapType::AmbientOcclusion, BakeType_AmbientOcclusion),
-	TTuple<EBakeMapType, FString>(EBakeMapType::BentNormal, BakeType_BentNormals),
-	TTuple<EBakeMapType, FString>(EBakeMapType::Curvature, BakeType_Curvature),
-	TTuple<EBakeMapType, FString>(EBakeMapType::Texture, BakeType_Texture),
-	TTuple<EBakeMapType, FString>(EBakeMapType::ObjectSpaceNormal, BakeType_ObjectSpaceNormals),
-	TTuple<EBakeMapType, FString>(EBakeMapType::FaceNormal, BakeType_FaceNormals),
-	TTuple<EBakeMapType, FString>(EBakeMapType::Position, BakeType_Position),
-	TTuple<EBakeMapType, FString>(EBakeMapType::MaterialID, BakeType_MaterialID),
-	TTuple<EBakeMapType, FString>(EBakeMapType::MultiTexture, BakeType_MultiTexture),
-	TTuple<EBakeMapType, FString>(EBakeMapType::VertexColor, BakeType_VertexColors)
-};
-
 
 void UBakeMeshAttributeMapsToolBase::Setup()
 {
@@ -227,7 +185,6 @@ void UBakeMeshAttributeMapsToolBase::CreateTextureAssets(const TMap<EBakeMapType
 }
 
 
-
 void UBakeMeshAttributeMapsToolBase::UpdatePreview(const EBakeMapType PreviewMapType)
 {
 	if (PreviewMapType == EBakeMapType::None)
diff --git a/Engine/Plugins/Experimental/MeshModelingToolsetExp/Source/MeshModelingToolsExp/Private/BakeMultiMeshAttributeMapsTool.cpp b/Engine/Plugins/Experimental/MeshModelingToolsetExp/Source/MeshModelingToolsExp/Private/BakeMultiMeshAttributeMapsTool.cpp
--- a/Engine/Plugins/Experimental/MeshModelingToolsetExp/Source/MeshModelingToolsExp/Private/BakeMultiMeshAttributeMapsTool.cpp
+++ b/Engine/Plugins/Experimental/MeshModelingToolsetExp/Source/MeshModelingToolsExp/Private/BakeMultiMeshAttributeMapsTool.cpp
@@ -688,11 +688,7 @@ void UBakeMultiMeshAttributeMapsTool::UpdateVisualization()
 		}
 	}
 
-	// Set the preview material according to the preview name.
-	if (const EBakeMapType* PreviewType = NameToMapTypeMap.Find(Settings->MapPreview))
-	{
-		UpdatePreview(*PreviewType);
-	}
+	UpdatePreview(Settings);
 }
 
 
diff --git a/Engine/Plugins/Experimental/MeshModelingToolsetExp/Source/MeshModelingToolsExp/Public/BakeMeshAttributeMapsTool.h b/Engine/Plugins/Experimental/MeshModelingToolsetExp/Source/MeshModelingToolsExp/Public/BakeMeshAttributeMapsTool.h
--- a/Engine/Plugins/Experimental/MeshModelingToolsetExp/Source/MeshModelingToolsExp/Public/BakeMeshAttributeMapsTool.h
+++ b/Engine/Plugins/Experimental/MeshModelingToolsetExp/Source/MeshModelingToolsExp/Public/BakeMeshAttributeMapsTool.h
@@ -92,6 +92,7 @@ public:
 	const TArray<FString>& GetMapPreviewNamesFunc();
 	UPROPERTY(meta = (TransientToolProperty))
 	TArray<FString> MapPreviewNamesList;
+	TMap<FString, FString> MapPreviewNamesMap;
 
 	UPROPERTY(VisibleAnywhere, Category = MapSettings, meta = (TransientToolProperty))
 	TMap<EBakeMapType, TObjectPtr<UTexture2D>> Result;
diff --git a/Engine/Plugins/Experimental/MeshModelingToolsetExp/Source/MeshModelingToolsExp/Public/BakeMeshAttributeMapsToolBase.h b/Engine/Plugins/Experimental/MeshModelingToolsetExp/Source/MeshModelingToolsExp/Public/BakeMeshAttributeMapsToolBase.h
--- a/Engine/Plugins/Experimental/MeshModelingToolsetExp/Source/MeshModelingToolsExp/Public/BakeMeshAttributeMapsToolBase.h
+++ b/Engine/Plugins/Experimental/MeshModelingToolsetExp/Source/MeshModelingToolsExp/Public/BakeMeshAttributeMapsToolBase.h
@@ -143,13 +143,15 @@ public:
 	void SetWorld(UWorld* World);
 
 public:
-	/** @return An enumerated map type from a bitfield */
+	/**
+	 * Process a bitfield into an EBakeMapType. This function
+	 * may inject additional map types based on the enabled bits.
+	 * For example, enabling AmbientOcclusion if BentNormal is
+	 * active.
+	 * @return An enumerated map type from a bitfield
+	 */
 	static EBakeMapType GetMapTypes(const int32& MapTypes);
 
-	/** Internal maps of name to BakeMapType */
-	static TMap<FString, EBakeMapType> NameToMapTypeMap;
-	static TMap<EBakeMapType, FString> MapTypeToNameMap;
-
 protected:
 	//
 	// Tool property sets
@@ -270,7 +272,16 @@ protected:
 
 
 	/**
-	 * Update the preview material parameters for a given a result index.
+	 * Update the preview material parameters for a given Bake type
+	 * display name.
+	 * @param PreviewDisplayName the display name of a Bake type to preview.
+	 */
+	template <typename PropertySet>
+	void UpdatePreview(PropertySet& Properties);
+
+
+	/**
+	 * Update the preview material parameters for a given a Bake type.
 	 * @param PreviewMapType EBakeMapType to preview
 	 */
 	void UpdatePreview(EBakeMapType PreviewMapType);
@@ -449,11 +460,27 @@ void UBakeMeshAttributeMapsToolBase::OnMapTypesUpdated(PropertySet& Properties)
 }
 
 
+template <typename PropertySet>
+void UBakeMeshAttributeMapsToolBase::UpdatePreview(PropertySet& Properties)
+{
+	const FString& PreviewDisplayName = Properties->MapPreview;
+	if (const FString* PreviewNameString = Properties->MapPreviewNamesMap.Find(PreviewDisplayName))
+	{
+		const int64 PreviewValue = StaticEnum<EBakeMapType>()->GetValueByNameString(*PreviewNameString);
+		if (PreviewValue != INDEX_NONE)
+		{
+			UpdatePreview(static_cast<EBakeMapType>(PreviewValue));
+		}
+	}
+}
+
+
 template <typename PropertySet>
 void UBakeMeshAttributeMapsToolBase::UpdatePreviewNames(PropertySet& Properties)
 {
 	// Update our preview names list.
 	Properties->MapPreviewNamesList.Reset();
+	Properties->MapPreviewNamesMap.Reset();
 	bool bFoundMapType = false;
 	for (const TTuple<EBakeMapType, TObjectPtr<UTexture2D>>& Map : CachedMaps)
 	{
@@ -461,13 +488,15 @@ void UBakeMeshAttributeMapsToolBase::UpdatePreviewNames(PropertySet& Properties)
 		// AO may have only been added for preview of other types (ex. BentNormal)
 		if (IsRequestedMapType(Properties, Map.Get<0>()))
 		{
-			if (const FString* MapTypeName = MapTypeToNameMap.Find(Map.Get<0>()))
+			const UEnum* BakeTypeEnum = StaticEnum<EBakeMapType>();
+			const int64 BakeEnumValue = static_cast<int64>(Map.Get<0>());
+			const FString BakeTypeDisplayName = BakeTypeEnum->GetDisplayNameTextByValue(BakeEnumValue).ToString();
+			const FString BakeTypeNameString = BakeTypeEnum->GetNameStringByValue(BakeEnumValue);
+			Properties->MapPreviewNamesList.Add(BakeTypeDisplayName);
+			Properties->MapPreviewNamesMap.Emplace(BakeTypeDisplayName, BakeTypeNameString);
+			if (Properties->MapPreview == Properties->MapPreviewNamesList.Last())
 			{
-				Properties->MapPreviewNamesList.Add(*MapTypeName);
-				if (Properties->MapPreview == Properties->MapPreviewNamesList.Last())
-				{
-					bFoundMapType = true;
-				}
+				bFoundMapType = true;
 			}
 		}
 	}
diff --git a/Engine/Plugins/Experimental/MeshModelingToolsetExp/Source/MeshModelingToolsExp/Public/BakeMultiMeshAttributeMapsTool.h b/Engine/Plugins/Experimental/MeshModelingToolsetExp/Source/MeshModelingToolsExp/Public/BakeMultiMeshAttributeMapsTool.h
--- a/Engine/Plugins/Experimental/MeshModelingToolsetExp/Source/MeshModelingToolsExp/Public/BakeMultiMeshAttributeMapsTool.h
+++ b/Engine/Plugins/Experimental/MeshModelingToolsetExp/Source/MeshModelingToolsExp/Public/BakeMultiMeshAttributeMapsTool.h
@@ -76,6 +76,7 @@ public:
 	const TArray<FString>& GetMapPreviewNamesFunc();
 	UPROPERTY(meta = (TransientToolProperty))
 	TArray<FString> MapPreviewNamesList;
+	TMap<FString, FString> MapPreviewNamesMap;
 
 	UPROPERTY(VisibleAnywhere, Category = MapSettings, meta = (TransientToolProperty))
 	TMap<EBakeMapType, TObjectPtr<UTexture2D>> Result;
