ModelingTools: Unify control flow in the bake tools to use BakeOpState. This will address the perceived hangs in BakeVertex.

#rnx

diff --git a/Engine/Plugins/Experimental/MeshModelingToolsetExp/Source/MeshModelingToolsExp/Private/BakeMeshAttributeMapsTool.cpp b/Engine/Plugins/Experimental/MeshModelingToolsetExp/Source/MeshModelingToolsExp/Private/BakeMeshAttributeMapsTool.cpp
--- a/Engine/Plugins/Experimental/MeshModelingToolsetExp/Source/MeshModelingToolsExp/Private/BakeMeshAttributeMapsTool.cpp
+++ b/Engine/Plugins/Experimental/MeshModelingToolsetExp/Source/MeshModelingToolsExp/Private/BakeMeshAttributeMapsTool.cpp
@@ -282,14 +282,14 @@ void UBakeMeshAttributeMapsTool::Setup()
 	UpdateUVLayerNames(Settings, BaseMesh);
 	AddToolPropertySource(Settings);
 
-	Settings->WatchProperty(Settings->MapTypes, [this](int32) { bInputsDirty = true; UpdateOnModeChange(); });
+	Settings->WatchProperty(Settings->MapTypes, [this](int32) { OpState |= EBakeOpState::Evaluate; UpdateOnModeChange(); });
 	Settings->WatchProperty(Settings->MapPreview, [this](FString) { UpdateVisualization(); GetToolManager()->PostInvalidation(); });
-	Settings->WatchProperty(Settings->Resolution, [this](EBakeTextureResolution) { bInputsDirty = true; });
-	Settings->WatchProperty(Settings->SourceFormat, [this](EBakeTextureFormat) { bInputsDirty = true; });
-	Settings->WatchProperty(Settings->UVLayer, [this](FString) { bInputsDirty = true; });
-	Settings->WatchProperty(Settings->bUseWorldSpace, [this](bool) { bDetailMeshValid = false; bInputsDirty = true; });
-	Settings->WatchProperty(Settings->Thickness, [this](float) { bInputsDirty = true; });
-	Settings->WatchProperty(Settings->Multisampling, [this](EBakeMultisampling) { bInputsDirty = true; });
+	Settings->WatchProperty(Settings->Resolution, [this](EBakeTextureResolution) { OpState |= EBakeOpState::Evaluate; });
+	Settings->WatchProperty(Settings->SourceFormat, [this](EBakeTextureFormat) { OpState |= EBakeOpState::Evaluate; });
+	Settings->WatchProperty(Settings->UVLayer, [this](FString) { OpState |= EBakeOpState::Evaluate; });
+	Settings->WatchProperty(Settings->bUseWorldSpace, [this](bool) { OpState |= EBakeOpState::EvaluateDetailMesh; });
+	Settings->WatchProperty(Settings->Thickness, [this](float) { OpState |= EBakeOpState::Evaluate; });
+	Settings->WatchProperty(Settings->Multisampling, [this](EBakeMultisampling) { OpState |= EBakeOpState::Evaluate; });
 
 
 	UToolTarget* DetailTarget = Targets[bIsBakeToSelf ? 0 : 1];
@@ -304,15 +304,15 @@ void UBakeMeshAttributeMapsTool::Setup()
 	DetailMeshProps->DetailStaticMesh = DetailStaticMesh;
 	DetailMeshProps->DetailSkeletalMesh = DetailSkeletalMesh;
 	DetailMeshProps->DetailMeshNormalMap = nullptr;
-	DetailMeshProps->WatchProperty(DetailMeshProps->DetailNormalUVLayer, [this](int) { bInputsDirty = true; });
+	DetailMeshProps->WatchProperty(DetailMeshProps->DetailNormalUVLayer, [this](int) { OpState |= EBakeOpState::Evaluate; });
 	DetailMeshProps->WatchProperty(DetailMeshProps->DetailMeshNormalMap, [this](UTexture2D*)
 	{
 		// Only invalidate detail mesh if we need to recompute tangents.
 		if (!DetailMeshTangents)
 		{
-			bDetailMeshValid = false;
+			OpState |= EBakeOpState::EvaluateDetailMesh;
 		}
-		bInputsDirty = true;
+		OpState = EBakeOpState::Evaluate;
 	});
 	
 
@@ -320,39 +320,39 @@ void UBakeMeshAttributeMapsTool::Setup()
 	OcclusionMapProps->RestoreProperties(this);
 	AddToolPropertySource(OcclusionMapProps);
 	SetToolPropertySourceEnabled(OcclusionMapProps, false);
-	OcclusionMapProps->WatchProperty(OcclusionMapProps->OcclusionRays, [this](int32) { bInputsDirty = true; });
-	OcclusionMapProps->WatchProperty(OcclusionMapProps->MaxDistance, [this](float) { bInputsDirty = true; });
-	OcclusionMapProps->WatchProperty(OcclusionMapProps->SpreadAngle, [this](float) { bInputsDirty = true; });
-	OcclusionMapProps->WatchProperty(OcclusionMapProps->BiasAngle, [this](float) { bInputsDirty = true; });
+	OcclusionMapProps->WatchProperty(OcclusionMapProps->OcclusionRays, [this](int32) { OpState |= EBakeOpState::Evaluate; });
+	OcclusionMapProps->WatchProperty(OcclusionMapProps->MaxDistance, [this](float) { OpState |= EBakeOpState::Evaluate; });
+	OcclusionMapProps->WatchProperty(OcclusionMapProps->SpreadAngle, [this](float) { OpState |= EBakeOpState::Evaluate; });
+	OcclusionMapProps->WatchProperty(OcclusionMapProps->BiasAngle, [this](float) { OpState |= EBakeOpState::Evaluate; });
 
 
 	CurvatureMapProps = NewObject<UBakedCurvatureMapToolProperties>(this);
 	CurvatureMapProps->RestoreProperties(this);
 	AddToolPropertySource(CurvatureMapProps);
 	SetToolPropertySourceEnabled(CurvatureMapProps, false);
-	CurvatureMapProps->WatchProperty(CurvatureMapProps->RangeMultiplier, [this](float) { bInputsDirty = true; });
-	CurvatureMapProps->WatchProperty(CurvatureMapProps->MinRangeMultiplier, [this](float) { bInputsDirty = true; });
-	CurvatureMapProps->WatchProperty(CurvatureMapProps->CurvatureType, [this](EBakedCurvatureTypeMode) { bInputsDirty = true; });
-	CurvatureMapProps->WatchProperty(CurvatureMapProps->ColorMode, [this](EBakedCurvatureColorMode) { bInputsDirty = true; });
-	CurvatureMapProps->WatchProperty(CurvatureMapProps->Clamping, [this](EBakedCurvatureClampMode) { bInputsDirty = true; });
+	CurvatureMapProps->WatchProperty(CurvatureMapProps->RangeMultiplier, [this](float) { OpState |= EBakeOpState::Evaluate; });
+	CurvatureMapProps->WatchProperty(CurvatureMapProps->MinRangeMultiplier, [this](float) { OpState |= EBakeOpState::Evaluate; });
+	CurvatureMapProps->WatchProperty(CurvatureMapProps->CurvatureType, [this](EBakedCurvatureTypeMode) { OpState |= EBakeOpState::Evaluate; });
+	CurvatureMapProps->WatchProperty(CurvatureMapProps->ColorMode, [this](EBakedCurvatureColorMode) { OpState |= EBakeOpState::Evaluate; });
+	CurvatureMapProps->WatchProperty(CurvatureMapProps->Clamping, [this](EBakedCurvatureClampMode) { OpState |= EBakeOpState::Evaluate; });
 
 
 	Texture2DProps = NewObject<UBakedTexture2DImageProperties>(this);
 	Texture2DProps->RestoreProperties(this);
 	AddToolPropertySource(Texture2DProps);
 	SetToolPropertySourceEnabled(Texture2DProps, false);
-	Texture2DProps->WatchProperty(Texture2DProps->UVLayer, [this](float) { bInputsDirty = true; });
-	Texture2DProps->WatchProperty(Texture2DProps->SourceTexture, [this](UTexture2D*) { bInputsDirty = true; });
+	Texture2DProps->WatchProperty(Texture2DProps->UVLayer, [this](float) { OpState |= EBakeOpState::Evaluate; });
+	Texture2DProps->WatchProperty(Texture2DProps->SourceTexture, [this](UTexture2D*) { OpState |= EBakeOpState::Evaluate; });
 
 	MultiTextureProps = NewObject<UBakedMultiTexture2DImageProperties>(this);
 	MultiTextureProps->RestoreProperties(this);
 	AddToolPropertySource(MultiTextureProps);
 	SetToolPropertySourceEnabled(MultiTextureProps, false);
 
-	auto SetDirtyCallback = [this](decltype(MultiTextureProps->MaterialIDSourceTextureMap)) { bInputsDirty = true; };
+	auto SetDirtyCallback = [this](decltype(MultiTextureProps->MaterialIDSourceTextureMap)) { OpState |= EBakeOpState::Evaluate; };
 	auto NotEqualsCallback = [](const decltype(MultiTextureProps->MaterialIDSourceTextureMap)& A, const decltype(MultiTextureProps->MaterialIDSourceTextureMap)& B) -> bool { return !(A.OrderIndependentCompareEqual(B)); };
 	MultiTextureProps->WatchProperty(MultiTextureProps->MaterialIDSourceTextureMap, SetDirtyCallback, NotEqualsCallback);
-	MultiTextureProps->WatchProperty(MultiTextureProps->UVLayer, [this](float) { bInputsDirty = true; });
+	MultiTextureProps->WatchProperty(MultiTextureProps->UVLayer, [this](float) { OpState |= EBakeOpState::Evaluate; });
 
 	UpdateOnModeChange();
 
@@ -536,8 +536,9 @@ void UBakeMeshAttributeMapsTool::UpdateDetailMesh()
 		}
 	});
 
-	bDetailMeshValid = true;
-	bInputsDirty = true;
+	// Clear detail mesh evaluation flag and mark evaluation.
+	OpState &= ~EBakeOpState::EvaluateDetailMesh;
+	OpState |= EBakeOpState::Evaluate;
 	CachedBakeCacheSettings = FBakeCacheSettings();
 	DetailMeshTimestamp++;
 }
@@ -545,14 +546,12 @@ void UBakeMeshAttributeMapsTool::UpdateDetailMesh()
 
 void UBakeMeshAttributeMapsTool::UpdateResult()
 {
-	if (!bDetailMeshValid)
+	if (static_cast<bool>(OpState & EBakeOpState::EvaluateDetailMesh))
 	{
 		UpdateDetailMesh();
 	}
 
-	// bInputsDirty ensures that we only validate parameters once per param
-	// change. Parameter validation can be expensive (ex. UpdateResult_Texture2DImage).
-	if (!bInputsDirty)
+	if (OpState == EBakeOpState::Clean)
 	{
 		return;
 	}
@@ -590,7 +589,7 @@ void UBakeMeshAttributeMapsTool::UpdateResult()
 	}
 
 	// Clear our invalid bitflag to check again for valid inputs.
-	OpState = EBakeOpState::Evaluate;
+	OpState &= ~EBakeOpState::Invalid;
 
 	OpState |= UpdateResult_DetailNormalMap();
 
@@ -640,7 +639,7 @@ void UBakeMeshAttributeMapsTool::UpdateResult()
 
 EBakeOpState UBakeMeshAttributeMapsTool::UpdateResult_DetailNormalMap()
 {
-	EBakeOpState ResultState = EBakeOpState::Complete;
+	EBakeOpState ResultState = EBakeOpState::Clean;
 	
 	const FDynamicMeshUVOverlay* UVOverlay = DetailMesh->Attributes()->GetUVLayer(DetailMeshProps->DetailNormalUVLayer);
 	if (UVOverlay == nullptr)
@@ -678,7 +677,7 @@ EBakeOpState UBakeMeshAttributeMapsTool::UpdateResult_DetailNormalMap()
 
 EBakeOpState UBakeMeshAttributeMapsTool::UpdateResult_Normal()
 {
-	EBakeOpState ResultState = EBakeOpState::Complete;
+	EBakeOpState ResultState = EBakeOpState::Clean;
 
 	const int32 ImageSize = (int32)Settings->Resolution;
 	const FImageDimensions Dimensions(ImageSize, ImageSize);
@@ -697,7 +696,7 @@ EBakeOpState UBakeMeshAttributeMapsTool::UpdateResult_Normal()
 
 EBakeOpState UBakeMeshAttributeMapsTool::UpdateResult_Occlusion()
 {
-	EBakeOpState ResultState = EBakeOpState::Complete;
+	EBakeOpState ResultState = EBakeOpState::Clean;
 
 	const int32 ImageSize = (int32)Settings->Resolution;
 	const FImageDimensions Dimensions(ImageSize, ImageSize);
@@ -720,7 +719,7 @@ EBakeOpState UBakeMeshAttributeMapsTool::UpdateResult_Occlusion()
 
 EBakeOpState UBakeMeshAttributeMapsTool::UpdateResult_Curvature()
 {
-	EBakeOpState ResultState = EBakeOpState::Complete;
+	EBakeOpState ResultState = EBakeOpState::Clean;
 
 	const int32 ImageSize = (int32)Settings->Resolution;
 	const FImageDimensions Dimensions(ImageSize, ImageSize);
@@ -783,7 +782,7 @@ EBakeOpState UBakeMeshAttributeMapsTool::UpdateResult_Curvature()
 
 EBakeOpState UBakeMeshAttributeMapsTool::UpdateResult_MeshProperty()
 {
-	EBakeOpState ResultState = EBakeOpState::Complete;
+	EBakeOpState ResultState = EBakeOpState::Clean;
 
 	const int32 ImageSize = (int32)Settings->Resolution;
 	const FImageDimensions Dimensions(ImageSize, ImageSize);
@@ -802,7 +801,7 @@ EBakeOpState UBakeMeshAttributeMapsTool::UpdateResult_MeshProperty()
 
 EBakeOpState UBakeMeshAttributeMapsTool::UpdateResult_Texture2DImage()
 {
-	EBakeOpState ResultState = EBakeOpState::Complete;
+	EBakeOpState ResultState = EBakeOpState::Clean;
 
 	const int32 ImageSize = (int32)Settings->Resolution;
 	const FImageDimensions Dimensions(ImageSize, ImageSize);
@@ -844,7 +843,7 @@ EBakeOpState UBakeMeshAttributeMapsTool::UpdateResult_Texture2DImage()
 
 EBakeOpState UBakeMeshAttributeMapsTool::UpdateResult_MultiTexture()
 {
-	EBakeOpState ResultState = EBakeOpState::Complete;
+	EBakeOpState ResultState = EBakeOpState::Clean;
 
 	const int32 ImageSize = (int32)Settings->Resolution;
 	const FImageDimensions Dimensions(ImageSize, ImageSize);
diff --git a/Engine/Plugins/Experimental/MeshModelingToolsetExp/Source/MeshModelingToolsExp/Private/BakeMeshAttributeMapsToolBase.cpp b/Engine/Plugins/Experimental/MeshModelingToolsetExp/Source/MeshModelingToolsExp/Private/BakeMeshAttributeMapsToolBase.cpp
--- a/Engine/Plugins/Experimental/MeshModelingToolsetExp/Source/MeshModelingToolsExp/Private/BakeMeshAttributeMapsToolBase.cpp
+++ b/Engine/Plugins/Experimental/MeshModelingToolsetExp/Source/MeshModelingToolsExp/Private/BakeMeshAttributeMapsToolBase.cpp
@@ -150,7 +150,7 @@ void UBakeMeshAttributeMapsToolBase::UpdateVisualization()
 
 void UBakeMeshAttributeMapsToolBase::InvalidateCompute()
 {
-	const bool bInvalidate = bInputsDirty || static_cast<bool>(OpState & EBakeOpState::Evaluate);
+	const bool bInvalidate = static_cast<bool>(OpState & EBakeOpState::Evaluate);
 	if (!Compute)
 	{
 		Compute = MakeUnique<TGenericDataBackgroundCompute<FMeshMapBaker>>();
@@ -162,7 +162,7 @@ void UBakeMeshAttributeMapsToolBase::InvalidateCompute()
 	{
 		Compute->InvalidateResult();
 	}
-	bInputsDirty = false;
+	OpState = EBakeOpState::Clean;
 }
 
 
diff --git a/Engine/Plugins/Experimental/MeshModelingToolsetExp/Source/MeshModelingToolsExp/Private/BakeMeshAttributeVertexTool.cpp b/Engine/Plugins/Experimental/MeshModelingToolsetExp/Source/MeshModelingToolsExp/Private/BakeMeshAttributeVertexTool.cpp
--- a/Engine/Plugins/Experimental/MeshModelingToolsetExp/Source/MeshModelingToolsExp/Private/BakeMeshAttributeVertexTool.cpp
+++ b/Engine/Plugins/Experimental/MeshModelingToolsetExp/Source/MeshModelingToolsExp/Private/BakeMeshAttributeVertexTool.cpp
@@ -295,53 +295,53 @@ void UBakeMeshAttributeVertexTool::Setup()
 	Settings->RestoreProperties(this);
 	AddToolPropertySource(Settings);
 
-	Settings->WatchProperty(Settings->VertexMode, [this](EBakeVertexMode) { OpState = EBakeOpState::Evaluate; UpdateOnModeChange(); });
+	Settings->WatchProperty(Settings->VertexMode, [this](EBakeVertexMode) { OpState |= EBakeOpState::Evaluate; UpdateOnModeChange(); });
 	Settings->WatchProperty(Settings->VertexChannelPreview, [this](EBakeVertexChannel) { UpdateVisualization(); });
-	Settings->WatchProperty(Settings->Thickness, [this](float) { OpState = EBakeOpState::Evaluate; });
-	Settings->WatchProperty(Settings->bUseWorldSpace, [this](bool)	{ bDetailMeshValid = false; OpState = EBakeOpState::Evaluate; });
-	Settings->WatchProperty(Settings->bSplitAtNormalSeams, [this](bool) { bColorTopologyValid = false; OpState = EBakeOpState::Evaluate; });
-	Settings->WatchProperty(Settings->bSplitAtUVSeams, [this](bool) { bColorTopologyValid = false; OpState = EBakeOpState::Evaluate; });
+	Settings->WatchProperty(Settings->Thickness, [this](float) { OpState |= EBakeOpState::Evaluate; });
+	Settings->WatchProperty(Settings->bUseWorldSpace, [this](bool)	{ OpState |= EBakeOpState::EvaluateDetailMesh; });
+	Settings->WatchProperty(Settings->bSplitAtNormalSeams, [this](bool) { bColorTopologyValid = false; OpState |= EBakeOpState::Evaluate; });
+	Settings->WatchProperty(Settings->bSplitAtUVSeams, [this](bool) { bColorTopologyValid = false; OpState |= EBakeOpState::Evaluate; });
 
 	ColorSettings = NewObject<UBakeMeshAttributeVertexToolColorProperties>(this);
 	ColorSettings->RestoreProperties(this);
 	AddToolPropertySource(ColorSettings);
 	SetToolPropertySourceEnabled(ColorSettings, false);
-	ColorSettings->WatchProperty(ColorSettings->BakeType, [this](EBakeVertexTypeColor) { OpState = EBakeOpState::Evaluate; UpdateOnModeChange(); });
+	ColorSettings->WatchProperty(ColorSettings->BakeType, [this](EBakeVertexTypeColor) { OpState |= EBakeOpState::Evaluate; UpdateOnModeChange(); });
 
 	PerChannelSettings = NewObject<UBakeMeshAttributeVertexToolChannelProperties>(this);
 	PerChannelSettings->RestoreProperties(this);
 	AddToolPropertySource(PerChannelSettings);
 	SetToolPropertySourceEnabled(PerChannelSettings, false);
-	PerChannelSettings->WatchProperty(PerChannelSettings->BakeTypeR, [this](EBakeVertexTypeChannel) { OpState = EBakeOpState::Evaluate; UpdateOnModeChange(); });
-	PerChannelSettings->WatchProperty(PerChannelSettings->BakeTypeG, [this](EBakeVertexTypeChannel) { OpState = EBakeOpState::Evaluate; UpdateOnModeChange(); });
-	PerChannelSettings->WatchProperty(PerChannelSettings->BakeTypeB, [this](EBakeVertexTypeChannel) { OpState = EBakeOpState::Evaluate; UpdateOnModeChange(); });
-	PerChannelSettings->WatchProperty(PerChannelSettings->BakeTypeA, [this](EBakeVertexTypeChannel) { OpState = EBakeOpState::Evaluate; UpdateOnModeChange(); });
+	PerChannelSettings->WatchProperty(PerChannelSettings->BakeTypeR, [this](EBakeVertexTypeChannel) { OpState |= EBakeOpState::Evaluate; UpdateOnModeChange(); });
+	PerChannelSettings->WatchProperty(PerChannelSettings->BakeTypeG, [this](EBakeVertexTypeChannel) { OpState |= EBakeOpState::Evaluate; UpdateOnModeChange(); });
+	PerChannelSettings->WatchProperty(PerChannelSettings->BakeTypeB, [this](EBakeVertexTypeChannel) { OpState |= EBakeOpState::Evaluate; UpdateOnModeChange(); });
+	PerChannelSettings->WatchProperty(PerChannelSettings->BakeTypeA, [this](EBakeVertexTypeChannel) { OpState |= EBakeOpState::Evaluate; UpdateOnModeChange(); });
 
 	OcclusionSettings = NewObject<UBakedOcclusionMapToolProperties>(this);
 	OcclusionSettings->RestoreProperties(this);
 	AddToolPropertySource(OcclusionSettings);
 	SetToolPropertySourceEnabled(OcclusionSettings, false);
-	OcclusionSettings->WatchProperty(OcclusionSettings->OcclusionRays, [this](int32) { OpState = EBakeOpState::Evaluate; });
-	OcclusionSettings->WatchProperty(OcclusionSettings->MaxDistance, [this](float) { OpState = EBakeOpState::Evaluate; });
-	OcclusionSettings->WatchProperty(OcclusionSettings->SpreadAngle, [this](float) { OpState = EBakeOpState::Evaluate; });
-	OcclusionSettings->WatchProperty(OcclusionSettings->BiasAngle, [this](float) { OpState = EBakeOpState::Evaluate; });
+	OcclusionSettings->WatchProperty(OcclusionSettings->OcclusionRays, [this](int32) { OpState |= EBakeOpState::Evaluate; });
+	OcclusionSettings->WatchProperty(OcclusionSettings->MaxDistance, [this](float) { OpState |= EBakeOpState::Evaluate; });
+	OcclusionSettings->WatchProperty(OcclusionSettings->SpreadAngle, [this](float) { OpState |= EBakeOpState::Evaluate; });
+	OcclusionSettings->WatchProperty(OcclusionSettings->BiasAngle, [this](float) { OpState |= EBakeOpState::Evaluate; });
 
 	CurvatureSettings = NewObject<UBakedCurvatureMapToolProperties>(this);
 	CurvatureSettings->RestoreProperties(this);
 	AddToolPropertySource(CurvatureSettings);
 	SetToolPropertySourceEnabled(CurvatureSettings, false);
-	CurvatureSettings->WatchProperty(CurvatureSettings->RangeMultiplier, [this](float) { OpState = EBakeOpState::Evaluate; });
-	CurvatureSettings->WatchProperty(CurvatureSettings->MinRangeMultiplier, [this](float) { OpState = EBakeOpState::Evaluate; });
-	CurvatureSettings->WatchProperty(CurvatureSettings->CurvatureType, [this](EBakedCurvatureTypeMode) { OpState = EBakeOpState::Evaluate; });
-	CurvatureSettings->WatchProperty(CurvatureSettings->ColorMode, [this](EBakedCurvatureColorMode) { OpState = EBakeOpState::Evaluate; });
-	CurvatureSettings->WatchProperty(CurvatureSettings->Clamping, [this](EBakedCurvatureClampMode) { OpState = EBakeOpState::Evaluate; });
+	CurvatureSettings->WatchProperty(CurvatureSettings->RangeMultiplier, [this](float) { OpState |= EBakeOpState::Evaluate; });
+	CurvatureSettings->WatchProperty(CurvatureSettings->MinRangeMultiplier, [this](float) { OpState |= EBakeOpState::Evaluate; });
+	CurvatureSettings->WatchProperty(CurvatureSettings->CurvatureType, [this](EBakedCurvatureTypeMode) { OpState |= EBakeOpState::Evaluate; });
+	CurvatureSettings->WatchProperty(CurvatureSettings->ColorMode, [this](EBakedCurvatureColorMode) { OpState |= EBakeOpState::Evaluate; });
+	CurvatureSettings->WatchProperty(CurvatureSettings->Clamping, [this](EBakedCurvatureClampMode) { OpState |= EBakeOpState::Evaluate; });
 
 	TextureSettings = NewObject<UBakedTexture2DImageProperties>(this);
 	TextureSettings->RestoreProperties(this);
 	AddToolPropertySource(TextureSettings);
 	SetToolPropertySourceEnabled(TextureSettings, false);
-	TextureSettings->WatchProperty(TextureSettings->UVLayer, [this](float) { OpState = EBakeOpState::Evaluate; });
-	TextureSettings->WatchProperty(TextureSettings->SourceTexture, [this](UTexture2D*) { OpState = EBakeOpState::Evaluate; });
+	TextureSettings->WatchProperty(TextureSettings->UVLayer, [this](float) { OpState |= EBakeOpState::Evaluate; });
+	TextureSettings->WatchProperty(TextureSettings->SourceTexture, [this](UTexture2D*) { OpState |= EBakeOpState::Evaluate; });
 
 	MultiTextureSettings = NewObject<UBakedMultiTexture2DImageProperties>(this);
 	MultiTextureSettings->RestoreProperties(this);
@@ -474,8 +474,8 @@ void UBakeMeshAttributeVertexTool::UpdateDetailMesh()
 	DetailSpatial = MakeShared<FDynamicMeshAABBTree3, ESPMode::ThreadSafe>();
 	DetailSpatial->SetMesh(DetailMesh.Get(), true);
 
-	OpState = EBakeOpState::Evaluate;
-	bDetailMeshValid = true;
+	OpState &= ~EBakeOpState::EvaluateDetailMesh;
+	OpState |= EBakeOpState::Evaluate;
 	DetailMeshTimestamp++;
 }
 
@@ -613,7 +613,7 @@ void UBakeMeshAttributeVertexTool::UpdateColorTopology()
 
 void UBakeMeshAttributeVertexTool::UpdateResult()
 {
-	if (!bDetailMeshValid)
+	if (static_cast<bool>(OpState & EBakeOpState::EvaluateDetailMesh))
 	{
 		UpdateDetailMesh();
 	}
@@ -623,7 +623,7 @@ void UBakeMeshAttributeVertexTool::UpdateResult()
 		UpdateColorTopology();
 	}
 
-	if (OpState == EBakeOpState::Complete)
+	if (OpState == EBakeOpState::Clean)
 	{
 		return;
 	}
@@ -663,7 +663,7 @@ void UBakeMeshAttributeVertexTool::UpdateResult()
 	}
 
 	// Clear our invalid bitflag to check again for valid inputs.
-	OpState = EBakeOpState::Evaluate;
+	OpState &= ~EBakeOpState::Invalid;
 
 	// Validate bake inputs
 	if (CachedBakeSettings.VertexMode == EBakeVertexMode::Color)
@@ -671,26 +671,26 @@ void UBakeMeshAttributeVertexTool::UpdateResult()
 		switch(CachedColorSettings.BakeType)
 		{
 		case EBakeVertexTypeColor::TangentSpaceNormal:
-			OpState = UpdateResult_Normal();
+			OpState |= UpdateResult_Normal();
 			break;
 		case EBakeVertexTypeColor::AmbientOcclusion:
 		case EBakeVertexTypeColor::BentNormal:
-			OpState = UpdateResult_Occlusion();
+			OpState |= UpdateResult_Occlusion();
 			break;
 		case EBakeVertexTypeColor::Curvature:
-			OpState = UpdateResult_Curvature();
+			OpState |= UpdateResult_Curvature();
 			break;
 		case EBakeVertexTypeColor::ObjectSpaceNormal:
 		case EBakeVertexTypeColor::FaceNormal:
 		case EBakeVertexTypeColor::Position:
 		case EBakeVertexTypeColor::MaterialID:
-			OpState = UpdateResult_MeshProperty();
+			OpState |= UpdateResult_MeshProperty();
 			break;
 		case EBakeVertexTypeColor::Texture:
-			OpState = UpdateResult_Texture2DImage();
+			OpState |= UpdateResult_Texture2DImage();
 			break;
 		case EBakeVertexTypeColor::MultiTexture:
-			OpState = UpdateResult_MultiTexture();
+			OpState |= UpdateResult_MultiTexture();
 			break;
 		}
 	}
@@ -712,19 +712,18 @@ void UBakeMeshAttributeVertexTool::UpdateResult()
 	{
 		return;
 	}
-	else // ((bool)(OpState & EBakeOpState::Evaluate))
+
+	if (!Compute)
 	{
-		if (!Compute)
+		Compute = MakeUnique<TGenericDataBackgroundCompute<FMeshVertexBaker>>();
+		Compute->Setup(this);
+		Compute->OnResultUpdated.AddLambda([this](const TUniquePtr<FMeshVertexBaker>& NewResult)
 		{
-			Compute = MakeUnique<TGenericDataBackgroundCompute<FMeshVertexBaker>>();
-			Compute->Setup(this);
-			Compute->OnResultUpdated.AddLambda([this](const TUniquePtr<FMeshVertexBaker>& NewResult)
-			{
-				OnResultUpdated(NewResult);
-			});
-		}
-		Compute->InvalidateResult();
+			OnResultUpdated(NewResult);
+		});
 	}
+	Compute->InvalidateResult();
+	OpState = EBakeOpState::Clean;
 }
 
 void UBakeMeshAttributeVertexTool::OnResultUpdated(const TUniquePtr<FMeshVertexBaker>& NewResult)
@@ -749,8 +748,6 @@ void UBakeMeshAttributeVertexTool::OnResultUpdated(const TUniquePtr<FMeshVertexB
 	PreviewMesh->NotifyDeferredEditCompleted(UPreviewMesh::ERenderUpdateMode::FastUpdate, EMeshRenderAttributeFlags::VertexColors, false);
 	UpdateVisualization();
 
-	OpState = EBakeOpState::Complete;
-
 	GatherAnalytics(*NewResult, CachedBakeSettings, CachedColorSettings, CachedChannelSettings, BakeAnalytics);
 }
 
@@ -762,7 +759,7 @@ EBakeOpState UBakeMeshAttributeVertexTool::UpdateResult_Normal()
 
 EBakeOpState UBakeMeshAttributeVertexTool::UpdateResult_Occlusion()
 {
-	EBakeOpState ResultState = EBakeOpState::Complete;
+	EBakeOpState ResultState = EBakeOpState::Clean;
 
 	FOcclusionMapSettings OcclusionMapSettings;
 	OcclusionMapSettings.MaxDistance = (OcclusionSettings->MaxDistance == 0) ? TNumericLimits<float>::Max() : OcclusionSettings->MaxDistance;
@@ -780,7 +777,7 @@ EBakeOpState UBakeMeshAttributeVertexTool::UpdateResult_Occlusion()
 
 EBakeOpState UBakeMeshAttributeVertexTool::UpdateResult_Curvature()
 {
-	EBakeOpState ResultState = EBakeOpState::Complete;
+	EBakeOpState ResultState = EBakeOpState::Clean;
 
 	FCurvatureMapSettings CurvatureMapSettings;
 	CurvatureMapSettings.RangeMultiplier = CurvatureSettings->RangeMultiplier;
@@ -844,7 +841,7 @@ EBakeOpState UBakeMeshAttributeVertexTool::UpdateResult_MeshProperty()
 
 EBakeOpState UBakeMeshAttributeVertexTool::UpdateResult_Texture2DImage()
 {
-	EBakeOpState ResultState = EBakeOpState::Complete;
+	EBakeOpState ResultState = EBakeOpState::Clean;
 
 	FTexture2DImageSettings NewSettings;
 	NewSettings.UVLayer = 0;
@@ -882,7 +879,7 @@ EBakeOpState UBakeMeshAttributeVertexTool::UpdateResult_Texture2DImage()
 
 EBakeOpState UBakeMeshAttributeVertexTool::UpdateResult_MultiTexture()
 {
-	EBakeOpState ResultState = EBakeOpState::Complete;
+	EBakeOpState ResultState = EBakeOpState::Clean;
 
 	FTexture2DImageSettings NewSettings;
 	NewSettings.UVLayer = MultiTextureSettings->UVLayer;
diff --git a/Engine/Plugins/Experimental/MeshModelingToolsetExp/Source/MeshModelingToolsExp/Private/BakeMultiMeshAttributeMapsTool.cpp b/Engine/Plugins/Experimental/MeshModelingToolsetExp/Source/MeshModelingToolsExp/Private/BakeMultiMeshAttributeMapsTool.cpp
--- a/Engine/Plugins/Experimental/MeshModelingToolsetExp/Source/MeshModelingToolsExp/Private/BakeMultiMeshAttributeMapsTool.cpp
+++ b/Engine/Plugins/Experimental/MeshModelingToolsetExp/Source/MeshModelingToolsExp/Private/BakeMultiMeshAttributeMapsTool.cpp
@@ -431,13 +431,13 @@ void UBakeMultiMeshAttributeMapsTool::Setup()
 	UpdateUVLayerNames(Settings, BaseMesh);
 	AddToolPropertySource(Settings);
 
-	Settings->WatchProperty(Settings->MapTypes, [this](int32) { bInputsDirty = true; UpdateOnModeChange(); });
+	Settings->WatchProperty(Settings->MapTypes, [this](int32) { OpState |= EBakeOpState::Evaluate; UpdateOnModeChange(); });
 	Settings->WatchProperty(Settings->MapPreview, [this](FString) { UpdateVisualization(); GetToolManager()->PostInvalidation(); });
-	Settings->WatchProperty(Settings->Resolution, [this](EBakeTextureResolution) { bInputsDirty = true; });
-	Settings->WatchProperty(Settings->SourceFormat, [this](EBakeTextureFormat) { bInputsDirty = true; });
-	Settings->WatchProperty(Settings->UVLayer, [this](FString) { bInputsDirty = true; });
-	Settings->WatchProperty(Settings->Thickness, [this](float) { bInputsDirty = true; });
-	Settings->WatchProperty(Settings->Multisampling, [this](EBakeMultisampling) { bInputsDirty = true; });
+	Settings->WatchProperty(Settings->Resolution, [this](EBakeTextureResolution) { OpState |= EBakeOpState::Evaluate; });
+	Settings->WatchProperty(Settings->SourceFormat, [this](EBakeTextureFormat) { OpState |= EBakeOpState::Evaluate; });
+	Settings->WatchProperty(Settings->UVLayer, [this](FString) { OpState |= EBakeOpState::Evaluate; });
+	Settings->WatchProperty(Settings->Thickness, [this](float) { OpState |= EBakeOpState::Evaluate; });
+	Settings->WatchProperty(Settings->Multisampling, [this](EBakeMultisampling) { OpState |= EBakeOpState::Evaluate; });
 
 	DetailProps = NewObject<UBakeMultiMeshDetailToolProperties>(this);
 	DetailProps->RestoreProperties(this);
@@ -475,13 +475,13 @@ void UBakeMultiMeshAttributeMapsTool::Setup()
 		DetailMeshProp.DetailMesh = DetailStaticMesh;
 		DetailMeshProp.DetailColorMap = DetailColorTexture;
 		DetailProps->DetailProperties.Add(DetailMeshProp);
-		DetailProps->WatchProperty(DetailProps->DetailProperties[Idx-1].DetailColorMap, [this](UTexture2D*) { bInputsDirty = true; });
-		DetailProps->WatchProperty(DetailProps->DetailProperties[Idx-1].DetailColorMapUVLayer, [this](int) { bInputsDirty = true; });
+		DetailProps->WatchProperty(DetailProps->DetailProperties[Idx-1].DetailColorMap, [this](UTexture2D*) { OpState |= EBakeOpState::Evaluate; });
+		DetailProps->WatchProperty(DetailProps->DetailProperties[Idx-1].DetailColorMapUVLayer, [this](int) { OpState |= EBakeOpState::Evaluate; });
 	}
 
 	UpdateOnModeChange();
 
-	bInputsDirty = true;
+	OpState |= EBakeOpState::Evaluate;
 
 	SetToolDisplayName(LOCTEXT("ToolName", "Bake Textures"));
 	GetToolManager()->DisplayMessage(
@@ -562,9 +562,7 @@ void UBakeMultiMeshAttributeMapsTool::Shutdown(EToolShutdownType ShutdownType)
 
 void UBakeMultiMeshAttributeMapsTool::UpdateResult()
 {
-	// bInputsDirty ensures that we only validate parameters once per param
-	// change. Parameter validation can be expensive (ex. UpdateResult_Texture2DImage).
-	if (!bInputsDirty)
+	if (OpState == EBakeOpState::Clean)
 	{
 		return;
 	}
@@ -595,7 +593,7 @@ void UBakeMultiMeshAttributeMapsTool::UpdateResult()
 	}
 
 	// Clear our invalid bitflag to check again for valid inputs.
-	OpState = EBakeOpState::Evaluate;
+	OpState &= ~EBakeOpState::Invalid;
 
 	// Update map type settings
 	OpState |= UpdateResult_DetailMeshes();
@@ -615,9 +613,6 @@ void UBakeMultiMeshAttributeMapsTool::UpdateResult()
 
 EBakeOpState UBakeMultiMeshAttributeMapsTool::UpdateResult_DetailMeshes()
 {
-	// This method will always force a re-evaluation.
-	EBakeOpState ResultState = EBakeOpState::Evaluate;
-
 	FBakeMultiMeshDetailSettings NewSettings;
 
 	// Iterate through our detail props to build our detail mesh data.
@@ -670,8 +665,9 @@ EBakeOpState UBakeMultiMeshAttributeMapsTool::UpdateResult_DetailMeshes()
 		}
 	};
 	DetailMeshScene.ProcessActorChildMeshes(BuildMeshToDataMaps);
-	
-	return ResultState;
+
+	// This method will always force a re-evaluation.
+	return EBakeOpState::Evaluate;
 }
 
 
diff --git a/Engine/Plugins/Experimental/MeshModelingToolsetExp/Source/MeshModelingToolsExp/Public/BakeMeshAttributeMapsTool.h b/Engine/Plugins/Experimental/MeshModelingToolsetExp/Source/MeshModelingToolsExp/Public/BakeMeshAttributeMapsTool.h
--- a/Engine/Plugins/Experimental/MeshModelingToolsetExp/Source/MeshModelingToolsExp/Public/BakeMeshAttributeMapsTool.h
+++ b/Engine/Plugins/Experimental/MeshModelingToolsetExp/Source/MeshModelingToolsExp/Public/BakeMeshAttributeMapsTool.h
@@ -163,7 +163,6 @@ protected:
 	TSharedPtr<UE::Geometry::TMeshTangents<double>, ESPMode::ThreadSafe> DetailMeshTangents;
 	int32 DetailMeshTimestamp = 0;
 	void UpdateDetailMesh();
-	bool bDetailMeshValid = false;
 
 	void UpdateOnModeChange();
 
diff --git a/Engine/Plugins/Experimental/MeshModelingToolsetExp/Source/MeshModelingToolsExp/Public/BakeMeshAttributeMapsToolBase.h b/Engine/Plugins/Experimental/MeshModelingToolsetExp/Source/MeshModelingToolsExp/Public/BakeMeshAttributeMapsToolBase.h
--- a/Engine/Plugins/Experimental/MeshModelingToolsetExp/Source/MeshModelingToolsExp/Public/BakeMeshAttributeMapsToolBase.h
+++ b/Engine/Plugins/Experimental/MeshModelingToolsetExp/Source/MeshModelingToolsExp/Public/BakeMeshAttributeMapsToolBase.h
@@ -220,7 +220,6 @@ protected:
 	TSharedPtr<UE::Geometry::TMeshTangents<double>, ESPMode::ThreadSafe> BaseMeshTangents;
 	UE::Geometry::FDynamicMeshAABBTree3 BaseSpatial;
 
-	bool bInputsDirty = false;
 	EBakeOpState OpState = EBakeOpState::Evaluate;
 
 	struct FBakeCacheSettings
diff --git a/Engine/Plugins/Experimental/MeshModelingToolsetExp/Source/MeshModelingToolsExp/Public/BakeMeshAttributeToolCommon.h b/Engine/Plugins/Experimental/MeshModelingToolsetExp/Source/MeshModelingToolsExp/Public/BakeMeshAttributeToolCommon.h
--- a/Engine/Plugins/Experimental/MeshModelingToolsetExp/Source/MeshModelingToolsExp/Public/BakeMeshAttributeToolCommon.h
+++ b/Engine/Plugins/Experimental/MeshModelingToolsetExp/Source/MeshModelingToolsExp/Public/BakeMeshAttributeToolCommon.h
@@ -274,9 +274,10 @@ struct FTexture2DImageSettings
  */
 enum class EBakeOpState
 {
-	Complete = 0,		// Inputs valid & Result is valid - no-op.
-	Evaluate = 1 << 0,	// Inputs valid & Result is invalid - re-evaluate.
-	Invalid	 = 1 << 1	// Inputs invalid - pause eval.
+	Clean              = 0,      // No-op - evaluation already launched/complete.
+	Evaluate           = 1 << 0, // Inputs are modified and valid, re-evaluate.
+	EvaluateDetailMesh = 1 << 1, // Detail mesh input is modified, re-evaluate the detail mesh.
+	Invalid            = 1 << 2  // Inputs are modified and invalid - retry eval until valid.
 };
 ENUM_CLASS_FLAGS(EBakeOpState);
 
diff --git a/Engine/Plugins/Experimental/MeshModelingToolsetExp/Source/MeshModelingToolsExp/Public/BakeMeshAttributeVertexTool.h b/Engine/Plugins/Experimental/MeshModelingToolsetExp/Source/MeshModelingToolsExp/Public/BakeMeshAttributeVertexTool.h
--- a/Engine/Plugins/Experimental/MeshModelingToolsetExp/Source/MeshModelingToolsExp/Public/BakeMeshAttributeVertexTool.h
+++ b/Engine/Plugins/Experimental/MeshModelingToolsetExp/Source/MeshModelingToolsExp/Public/BakeMeshAttributeVertexTool.h
@@ -242,7 +242,6 @@ protected:
 	TSharedPtr<UE::Geometry::FDynamicMesh3, ESPMode::ThreadSafe> DetailMesh;
 	TSharedPtr<UE::Geometry::FDynamicMeshAABBTree3, ESPMode::ThreadSafe> DetailSpatial;
 	int32 DetailMeshTimestamp = 0;
-	bool bDetailMeshValid = false;
 	void UpdateDetailMesh();
 
 	bool bColorTopologyValid = false;
