ModelingTools: Fix Bake analytics event prefix.

#rnx

diff --git a/Engine/Plugins/Experimental/MeshModelingToolsetExp/Source/MeshModelingToolsExp/Private/BakeMeshAttributeMapsToolBase.cpp b/Engine/Plugins/Experimental/MeshModelingToolsetExp/Source/MeshModelingToolsExp/Private/BakeMeshAttributeMapsToolBase.cpp
--- a/Engine/Plugins/Experimental/MeshModelingToolsetExp/Source/MeshModelingToolsExp/Private/BakeMeshAttributeMapsToolBase.cpp
+++ b/Engine/Plugins/Experimental/MeshModelingToolsetExp/Source/MeshModelingToolsExp/Private/BakeMeshAttributeMapsToolBase.cpp
@@ -658,7 +658,7 @@ void UBakeMeshAttributeMapsToolBase::RecordAnalytics(const FBakeAnalytics& Data,
 	const bool bVertexColor = static_cast<bool>(Data.BakeSettings.SourceBakeMapTypes & EBakeMapType::VertexColor);
 	Attributes.Add(FAnalyticsEventAttribute(TEXT("Settings.VertexColor.Enabled"), bVertexColor));
 
-	FEngineAnalytics::GetProvider().RecordEvent(FString(TEXT("Editor.Usage.ModelingTools.")) + EventName, Attributes);
+	FEngineAnalytics::GetProvider().RecordEvent(FString(TEXT("Editor.Usage.MeshModelingMode.")) + EventName, Attributes);
 
 	constexpr bool bLogAnalytics = false; 
 	if constexpr (bLogAnalytics)
diff --git a/Engine/Plugins/Experimental/MeshModelingToolsetExp/Source/MeshModelingToolsExp/Private/BakeMeshAttributeVertexTool.cpp b/Engine/Plugins/Experimental/MeshModelingToolsetExp/Source/MeshModelingToolsExp/Private/BakeMeshAttributeVertexTool.cpp
--- a/Engine/Plugins/Experimental/MeshModelingToolsetExp/Source/MeshModelingToolsExp/Private/BakeMeshAttributeVertexTool.cpp
+++ b/Engine/Plugins/Experimental/MeshModelingToolsetExp/Source/MeshModelingToolsExp/Private/BakeMeshAttributeVertexTool.cpp
@@ -1107,7 +1107,7 @@ void UBakeMeshAttributeVertexTool::RecordAnalytics(const FBakeAnalytics& Data, c
 		}
 	}
 
-	FEngineAnalytics::GetProvider().RecordEvent(FString(TEXT("Editor.Usage.ModelingTools.")) + EventName, Attributes);
+	FEngineAnalytics::GetProvider().RecordEvent(FString(TEXT("Editor.Usage.MeshModelingMode.")) + EventName, Attributes);
 
 	constexpr bool bLogAnalytics = false; 
 	if constexpr (bLogAnalytics)
