ModelingTools: Fix BakeVertex MultiTexture not prepopulating the source maps.
Fix BakeVertex MultiTexture not updating when changing source maps.

#rnx

diff --git a/Engine/Plugins/Experimental/MeshModelingToolsetExp/Source/MeshModelingToolsExp/Private/BakeMeshAttributeMapsToolBase.cpp b/Engine/Plugins/Experimental/MeshModelingToolsetExp/Source/MeshModelingToolsExp/Private/BakeMeshAttributeMapsToolBase.cpp
--- a/Engine/Plugins/Experimental/MeshModelingToolsetExp/Source/MeshModelingToolsExp/Private/BakeMeshAttributeMapsToolBase.cpp
+++ b/Engine/Plugins/Experimental/MeshModelingToolsetExp/Source/MeshModelingToolsExp/Private/BakeMeshAttributeMapsToolBase.cpp
@@ -53,12 +53,6 @@ void UBakeMeshAttributeMapsToolBase::Setup()
 	{
 		BentNormalPreviewMaterial = UMaterialInstanceDynamic::Create(BentNormalMaterial, GetToolManager());
 	}
-	UMaterial* WorkingMaterial = LoadObject<UMaterial>(nullptr, TEXT("/MeshModelingToolsetExp/Materials/InProgressMaterial"));
-	check(WorkingMaterial);
-	if (WorkingMaterial != nullptr)
-	{
-		WorkingPreviewMaterial = UMaterialInstanceDynamic::Create(WorkingMaterial, GetToolManager());
-	}
 
 	// Initialize preview mesh
 	UE::ToolTarget::HideSourceObject(Targets[0]);
@@ -132,12 +126,6 @@ TUniquePtr<UE::Geometry::TGenericDataOperator<FMeshMapBaker>> UBakeMeshAttribute
 }
 
 
-void UBakeMeshAttributeMapsToolBase::SetWorld(UWorld* World)
-{
-	TargetWorld = World;
-}
-
-
 void UBakeMeshAttributeMapsToolBase::UpdateResult()
 {
 }
@@ -456,52 +444,6 @@ void UBakeMeshAttributeMapsToolBase::GetTextureName(const EBakeMapType MapType,
 }
 
 
-int UBakeMeshAttributeMapsToolBase::SelectColorTextureToBake(const TArray<UTexture*>& Textures)
-{
-	TArray<int> TextureVotes;
-	TextureVotes.Init(0, Textures.Num());
-
-	for (int TextureIndex = 0; TextureIndex < Textures.Num(); ++TextureIndex)
-	{
-		UTexture* Tex = Textures[TextureIndex];
-		UTexture2D* Tex2D = Cast<UTexture2D>(Tex);
-
-		if (Tex2D)
-		{
-			// Texture uses SRGB
-			if (Tex->SRGB != 0)
-			{
-				++TextureVotes[TextureIndex];
-			}
-
-#if WITH_EDITORONLY_DATA
-			// Texture has multiple channels
-			ETextureSourceFormat Format = Tex->Source.GetFormat();
-			if (Format == TSF_BGRA8 || Format == TSF_BGRE8 || Format == TSF_RGBA16 || Format == TSF_RGBA16F)
-			{
-				++TextureVotes[TextureIndex];
-			}
-#endif
-
-			// What else? Largest texture? Most layers? Most mipmaps?
-		}
-	}
-
-	int MaxIndex = -1;
-	int MaxVotes = -1;
-	for (int TextureIndex = 0; TextureIndex < Textures.Num(); ++TextureIndex)
-	{
-		if (TextureVotes[TextureIndex] > MaxVotes)
-		{
-			MaxIndex = TextureIndex;
-			MaxVotes = TextureVotes[TextureIndex];
-		}
-	}
-
-	return MaxIndex;
-}
-
-
 void UBakeMeshAttributeMapsToolBase::InitializeEmptyMaps()
 {
 	FTexture2DBuilder NormalsBuilder;
diff --git a/Engine/Plugins/Experimental/MeshModelingToolsetExp/Source/MeshModelingToolsExp/Private/BakeMeshAttributeTool.cpp b/Engine/Plugins/Experimental/MeshModelingToolsetExp/Source/MeshModelingToolsExp/Private/BakeMeshAttributeTool.cpp
new file mode 100644
--- /dev/null
+++ b/Engine/Plugins/Experimental/MeshModelingToolsetExp/Source/MeshModelingToolsExp/Private/BakeMeshAttributeTool.cpp
@@ -0,0 +1,77 @@
+// Copyright Epic Games, Inc. All Rights Reserved.
+
+#include "BakeMeshAttributeTool.h"
+#include "InteractiveToolManager.h"
+
+using namespace UE::Geometry;
+
+#define LOCTEXT_NAMESPACE "UBakeMeshAttributeTool"
+
+
+void UBakeMeshAttributeTool::Setup()
+{
+	Super::Setup();
+
+	// Setup preview materials
+	UMaterial* WorkingMaterial = LoadObject<UMaterial>(nullptr, TEXT("/MeshModelingToolsetExp/Materials/InProgressMaterial"));
+	check(WorkingMaterial);
+	if (WorkingMaterial != nullptr)
+	{
+		WorkingPreviewMaterial = UMaterialInstanceDynamic::Create(WorkingMaterial, GetToolManager());
+	}
+}
+
+
+void UBakeMeshAttributeTool::SetWorld(UWorld* World)
+{
+	TargetWorld = World;
+}
+
+
+int UBakeMeshAttributeTool::SelectColorTextureToBake(const TArray<UTexture*>& Textures)
+{
+	TArray<int> TextureVotes;
+	TextureVotes.Init(0, Textures.Num());
+
+	for (int TextureIndex = 0; TextureIndex < Textures.Num(); ++TextureIndex)
+	{
+		UTexture* Tex = Textures[TextureIndex];
+		UTexture2D* Tex2D = Cast<UTexture2D>(Tex);
+
+		if (Tex2D)
+		{
+			// Texture uses SRGB
+			if (Tex->SRGB != 0)
+			{
+				++TextureVotes[TextureIndex];
+			}
+
+#if WITH_EDITORONLY_DATA
+			// Texture has multiple channels
+			ETextureSourceFormat Format = Tex->Source.GetFormat();
+			if (Format == TSF_BGRA8 || Format == TSF_BGRE8 || Format == TSF_RGBA16 || Format == TSF_RGBA16F)
+			{
+				++TextureVotes[TextureIndex];
+			}
+#endif
+
+			// What else? Largest texture? Most layers? Most mipmaps?
+		}
+	}
+
+	int MaxIndex = -1;
+	int MaxVotes = -1;
+	for (int TextureIndex = 0; TextureIndex < Textures.Num(); ++TextureIndex)
+	{
+		if (TextureVotes[TextureIndex] > MaxVotes)
+		{
+			MaxIndex = TextureIndex;
+			MaxVotes = TextureVotes[TextureIndex];
+		}
+	}
+
+	return MaxIndex;
+}
+
+#undef LOCTEXT_NAMESPACE
+
diff --git a/Engine/Plugins/Experimental/MeshModelingToolsetExp/Source/MeshModelingToolsExp/Private/BakeMeshAttributeVertexTool.cpp b/Engine/Plugins/Experimental/MeshModelingToolsetExp/Source/MeshModelingToolsExp/Private/BakeMeshAttributeVertexTool.cpp
--- a/Engine/Plugins/Experimental/MeshModelingToolsetExp/Source/MeshModelingToolsExp/Private/BakeMeshAttributeVertexTool.cpp
+++ b/Engine/Plugins/Experimental/MeshModelingToolsetExp/Source/MeshModelingToolsExp/Private/BakeMeshAttributeVertexTool.cpp
@@ -260,13 +260,6 @@ void UBakeMeshAttributeVertexTool::Setup()
 		PreviewAlphaMaterial = UMaterialInstanceDynamic::Create(AlphaMaterial, GetToolManager());
 	}
 
-	UMaterial* WorkingMaterial = LoadObject<UMaterial>(nullptr, TEXT("/MeshModelingToolsetExp/Materials/InProgressMaterial"));
-	check(WorkingMaterial);
-	if (WorkingMaterial != nullptr)
-	{
-		WorkingPreviewMaterial = UMaterialInstanceDynamic::Create(WorkingMaterial, GetToolManager());
-	}
-
 	bIsBakeToSelf = (Targets.Num() == 1);
 
 	UE::ToolTarget::HideSourceObject(Targets[0]);
@@ -347,6 +340,10 @@ void UBakeMeshAttributeVertexTool::Setup()
 	MultiTextureSettings->RestoreProperties(this);
 	AddToolPropertySource(MultiTextureSettings);
 	SetToolPropertySourceEnabled(MultiTextureSettings, false);
+	auto SetDirtyCallback = [this](decltype(MultiTextureSettings->MaterialIDSourceTextureMap)) { OpState |= EBakeOpState::Evaluate; };
+	auto NotEqualsCallback = [](const decltype(MultiTextureSettings->MaterialIDSourceTextureMap)& A, const decltype(MultiTextureSettings->MaterialIDSourceTextureMap)& B) -> bool { return !(A.OrderIndependentCompareEqual(B)); };
+	MultiTextureSettings->WatchProperty(MultiTextureSettings->MaterialIDSourceTextureMap, SetDirtyCallback, NotEqualsCallback);
+	MultiTextureSettings->WatchProperty(MultiTextureSettings->UVLayer, [this](float) { OpState |= EBakeOpState::Evaluate; });
 
 	UpdateOnModeChange();
 
@@ -449,11 +446,6 @@ TUniquePtr<UE::Geometry::TGenericDataOperator<FMeshVertexBaker>> UBakeMeshAttrib
 	return Op;
 }
 
-void UBakeMeshAttributeVertexTool::SetWorld(UWorld* World)
-{
-	TargetWorld = World;
-}
-
 void UBakeMeshAttributeVertexTool::UpdateDetailMesh()
 {
 	IPrimitiveComponentBackedTarget* TargetComponent = TargetComponentInterface(0);
@@ -474,6 +466,36 @@ void UBakeMeshAttributeVertexTool::UpdateDetailMesh()
 	DetailSpatial = MakeShared<FDynamicMeshAABBTree3, ESPMode::ThreadSafe>();
 	DetailSpatial->SetMesh(DetailMesh.Get(), true);
 
+	UToolTarget* DetailTarget = Targets[bIsBakeToSelf ? 0 : 1];
+	ProcessComponentTextures(UE::ToolTarget::GetTargetComponent(DetailTarget), [this](const int MaterialID, const TArray<UTexture*>& Textures)
+	{
+		for (UTexture* Tex : Textures)
+		{
+			UTexture2D* Tex2D = Cast<UTexture2D>(Tex);
+			if (Tex2D)
+			{
+				MultiTextureSettings->AllSourceTextures.Add(Tex2D);
+			}
+		}
+
+		constexpr bool bGuessAtTextures = true;
+		if (bGuessAtTextures)
+		{
+			const int SelectedTextureIndex = SelectColorTextureToBake(Textures);
+			if (SelectedTextureIndex >= 0)
+			{
+				UTexture2D* Tex2D = Cast<UTexture2D>(Textures[SelectedTextureIndex]);
+
+				// if cast fails, this will set the value to nullptr, which is fine
+				MultiTextureSettings->MaterialIDSourceTextureMap.Add(MaterialID, Tex2D);	
+			}
+		}
+		else
+		{
+			MultiTextureSettings->MaterialIDSourceTextureMap.Add(MaterialID, nullptr);
+		}
+	});
+
 	OpState &= ~EBakeOpState::EvaluateDetailMesh;
 	OpState |= EBakeOpState::Evaluate;
 	DetailMeshTimestamp++;
diff --git a/Engine/Plugins/Experimental/MeshModelingToolsetExp/Source/MeshModelingToolsExp/Public/BakeMeshAttributeMapsToolBase.h b/Engine/Plugins/Experimental/MeshModelingToolsetExp/Source/MeshModelingToolsExp/Public/BakeMeshAttributeMapsToolBase.h
--- a/Engine/Plugins/Experimental/MeshModelingToolsetExp/Source/MeshModelingToolsExp/Public/BakeMeshAttributeMapsToolBase.h
+++ b/Engine/Plugins/Experimental/MeshModelingToolsetExp/Source/MeshModelingToolsExp/Public/BakeMeshAttributeMapsToolBase.h
@@ -16,7 +16,7 @@
 #include "ModelingOperators.h"
 #include "MeshOpPreviewHelpers.h"
 #include "PreviewMesh.h"
-#include "BakeMeshAttributeToolCommon.h"
+#include "BakeMeshAttributeTool.h"
 #include "BakeMeshAttributeMapsToolBase.generated.h"
 
 /**
@@ -117,7 +117,7 @@ enum class EBakeMultisampling
  * Base Bake Maps tool
  */
 UCLASS()
-class MESHMODELINGTOOLSEXP_API UBakeMeshAttributeMapsToolBase : public UMultiSelectionTool, public IInteractiveToolExclusiveToolAPI, public UE::Geometry::IGenericDataOperatorFactory<UE::Geometry::FMeshMapBaker>
+class MESHMODELINGTOOLSEXP_API UBakeMeshAttributeMapsToolBase : public UBakeMeshAttributeTool, public UE::Geometry::IGenericDataOperatorFactory<UE::Geometry::FMeshMapBaker>
 {
 	GENERATED_BODY()
 
@@ -130,18 +130,12 @@ public:
 
 	virtual void OnTick(float DeltaTime) override;
 	virtual void Render(IToolsContextRenderAPI* RenderAPI) override;
-
-	virtual bool HasCancel() const override { return true; }
-	virtual bool HasAccept() const override { return true; }
-	virtual bool CanAccept() const override { return true; }
 	// End UInteractiveTool interface
 
 	// Begin IGenericDataOperatorFactory interface
 	virtual TUniquePtr<UE::Geometry::TGenericDataOperator<UE::Geometry::FMeshMapBaker>> MakeNewOperator() override;
 	// End IGenericDataOperatorFactory interface
 
-	void SetWorld(UWorld* World);
-
 public:
 	/**
 	 * Process a bitfield into an EBakeMapType. This function
@@ -171,10 +165,6 @@ protected:
 	UPROPERTY()
 	TObjectPtr<UMaterialInstanceDynamic> BentNormalPreviewMaterial;
 
-	UPROPERTY()
-	TObjectPtr<UMaterialInstanceDynamic> WorkingPreviewMaterial;
-	float SecondsBeforeWorkingMaterial = 0.75;
-
 protected:
 	/**
 	 * Post-client setup function. Should be invoked at end of client Setup().
@@ -213,15 +203,10 @@ protected:
 	//
 	// Bake parameters
 	//
-	UPROPERTY()
-	TObjectPtr<UWorld> TargetWorld = nullptr;
-
 	UE::Geometry::FDynamicMesh3 BaseMesh;
 	TSharedPtr<UE::Geometry::TMeshTangents<double>, ESPMode::ThreadSafe> BaseMeshTangents;
 	UE::Geometry::FDynamicMeshAABBTree3 BaseSpatial;
 
-	EBakeOpState OpState = EBakeOpState::Evaluate;
-
 	struct FBakeCacheSettings
 	{
 		EBakeMapType SourceBakeMapTypes = EBakeMapType::None;
@@ -273,7 +258,7 @@ protected:
 	/**
 	 * Update the preview material parameters for a given Bake type
 	 * display name.
-	 * @param PreviewDisplayName the display name of a Bake type to preview.
+	 * @param Properties the properties containing the preview property.
 	 */
 	template <typename PropertySet>
 	void UpdatePreview(PropertySet& Properties);
@@ -377,22 +362,6 @@ protected:
 	/** @return the texture name given a base name and map type */
 	static void GetTextureName(EBakeMapType MapType, const FString& BaseName, FString& TexName);
 
-	/**
-	 * Given an array of textures associated with a material,
-	 * use heuristics to identify the color/albedo texture.
-	 * @param Textures array of textures associated with a material.
-	 * @return integer index into the Textures array representing the color/albedo texture
-	 */
-	static int SelectColorTextureToBake(const TArray<UTexture*>& Textures);
-
-	/**
-	 * Iterate through a primitive component's textures by material ID.
-	 * @param Component the component to query
-	 * @param ProcessFunc enumeration function with signature: void(int MaterialID, const TArray<UTexture*>& Textures)
-	 */
-	template <typename ProcessFn>
-	static void ProcessComponentTextures(const UPrimitiveComponent* Component, ProcessFn&& ProcessFunc);
-
 	/**
 	 * @param Properties the tool property set to validate the map type against.
 	 * @param MapType the map type to validate.
@@ -526,31 +495,6 @@ void UBakeMeshAttributeMapsToolBase::UpdateUVLayerNames(PropertySet& Properties,
 }
 
 
-template <typename ProcessFn>
-void UBakeMeshAttributeMapsToolBase::ProcessComponentTextures(const UPrimitiveComponent* Component, ProcessFn&& ProcessFunc)
-{
-	if (!Component)
-	{
-		return;
-	}
-
-	TArray<UMaterialInterface*> Materials;
-	Component->GetUsedMaterials(Materials);
-	
-	for (int32 MaterialID = 0; MaterialID < Materials.Num(); ++MaterialID)	// TODO: This won't match MaterialIDs on the FDynamicMesh3 in general, will it?
-	{
-		UMaterialInterface* MaterialInterface = Materials[MaterialID];
-		if (MaterialInterface == nullptr)
-		{
-			continue;
-		}
-
-		TArray<UTexture*> Textures;
-		MaterialInterface->GetUsedTextures(Textures, EMaterialQualityLevel::High, true, ERHIFeatureLevel::SM5, true);
-		ProcessFunc(MaterialID, Textures);
-	}
-}
-
 template <typename PropertySet>
 bool UBakeMeshAttributeMapsToolBase::IsRequestedMapType(PropertySet& Properties, EBakeMapType MapType)
 {
diff --git a/Engine/Plugins/Experimental/MeshModelingToolsetExp/Source/MeshModelingToolsExp/Public/BakeMeshAttributeTool.h b/Engine/Plugins/Experimental/MeshModelingToolsetExp/Source/MeshModelingToolsExp/Public/BakeMeshAttributeTool.h
new file mode 100644
--- /dev/null
+++ b/Engine/Plugins/Experimental/MeshModelingToolsetExp/Source/MeshModelingToolsExp/Public/BakeMeshAttributeTool.h
@@ -0,0 +1,98 @@
+// Copyright Epic Games, Inc. All Rights Reserved.
+
+#pragma once
+
+#include "CoreMinimal.h"
+#include "Engine/Classes/Engine/Texture2D.h"
+#include "Materials/MaterialInstanceDynamic.h"
+#include "MultiSelectionTool.h"
+#include "InteractiveToolQueryInterfaces.h" // for UInteractiveToolExclusiveToolAPI
+#include "DynamicMesh/DynamicMesh3.h"
+#include "PreviewMesh.h"
+#include "BakeMeshAttributeToolCommon.h"
+#include "BakeMeshAttributeTool.generated.h"
+
+/**
+ * Base Mesh Bake tool
+ */
+UCLASS()
+class MESHMODELINGTOOLSEXP_API UBakeMeshAttributeTool : public UMultiSelectionTool, public IInteractiveToolExclusiveToolAPI
+{
+	GENERATED_BODY()
+
+public:
+	UBakeMeshAttributeTool() = default;
+
+	// Begin UInteractiveTool interface
+	virtual void Setup() override;
+
+	virtual bool HasCancel() const override { return true; }
+	virtual bool HasAccept() const override { return true; }
+	virtual bool CanAccept() const override { return true; }
+	// End UInteractiveTool interface
+
+	void SetWorld(UWorld* World);
+
+protected:
+	//
+	// Preview materials
+	//
+	UPROPERTY()
+	TObjectPtr<UMaterialInstanceDynamic> WorkingPreviewMaterial;
+	float SecondsBeforeWorkingMaterial = 0.75;
+
+protected:
+	//
+	// Bake parameters
+	//
+	UPROPERTY()
+	TObjectPtr<UWorld> TargetWorld = nullptr;
+
+	EBakeOpState OpState = EBakeOpState::Evaluate;
+
+
+protected:
+	/**
+	 * Given an array of textures associated with a material,
+	 * use heuristics to identify the color/albedo texture.
+	 * @param Textures array of textures associated with a material.
+	 * @return integer index into the Textures array representing the color/albedo texture
+	 */
+	static int SelectColorTextureToBake(const TArray<UTexture*>& Textures);
+
+	/**
+	 * Iterate through a primitive component's textures by material ID.
+	 * @param Component the component to query
+	 * @param ProcessFunc enumeration function with signature: void(int MaterialID, const TArray<UTexture*>& Textures)
+	 */
+	template <typename ProcessFn>
+	static void ProcessComponentTextures(const UPrimitiveComponent* Component, ProcessFn&& ProcessFunc);
+};
+
+
+template <typename ProcessFn>
+void UBakeMeshAttributeTool::ProcessComponentTextures(const UPrimitiveComponent* Component, ProcessFn&& ProcessFunc)
+{
+	if (!Component)
+	{
+		return;
+	}
+
+	TArray<UMaterialInterface*> Materials;
+	Component->GetUsedMaterials(Materials);
+	
+	for (int32 MaterialID = 0; MaterialID < Materials.Num(); ++MaterialID)	// TODO: This won't match MaterialIDs on the FDynamicMesh3 in general, will it?
+	{
+		UMaterialInterface* MaterialInterface = Materials[MaterialID];
+		if (MaterialInterface == nullptr)
+		{
+			continue;
+		}
+
+		TArray<UTexture*> Textures;
+		MaterialInterface->GetUsedTextures(Textures, EMaterialQualityLevel::High, true, ERHIFeatureLevel::SM5, true);
+		ProcessFunc(MaterialID, Textures);
+	}
+}
+
+
diff --git a/Engine/Plugins/Experimental/MeshModelingToolsetExp/Source/MeshModelingToolsExp/Public/BakeMeshAttributeVertexTool.h b/Engine/Plugins/Experimental/MeshModelingToolsetExp/Source/MeshModelingToolsExp/Public/BakeMeshAttributeVertexTool.h
--- a/Engine/Plugins/Experimental/MeshModelingToolsetExp/Source/MeshModelingToolsExp/Public/BakeMeshAttributeVertexTool.h
+++ b/Engine/Plugins/Experimental/MeshModelingToolsetExp/Source/MeshModelingToolsExp/Public/BakeMeshAttributeVertexTool.h
@@ -12,7 +12,7 @@
 #include "MeshOpPreviewHelpers.h"
 #include "PreviewMesh.h"
 #include "Sampling/MeshVertexBaker.h"
-#include "BakeMeshAttributeToolCommon.h"
+#include "BakeMeshAttributeTool.h"
 #include "BakeMeshAttributeVertexTool.generated.h"
 
 // predeclarations
@@ -166,7 +166,7 @@ public:
  * Vertex Baking Tool
  */
 UCLASS()
-class MESHMODELINGTOOLSEXP_API UBakeMeshAttributeVertexTool : public UMultiSelectionTool, public IInteractiveToolExclusiveToolAPI, public UE::Geometry::IGenericDataOperatorFactory<UE::Geometry::FMeshVertexBaker>
+class MESHMODELINGTOOLSEXP_API UBakeMeshAttributeVertexTool : public UBakeMeshAttributeTool, public UE::Geometry::IGenericDataOperatorFactory<UE::Geometry::FMeshVertexBaker>
 {
 	GENERATED_BODY()
 
@@ -189,8 +189,6 @@ public:
 	virtual TUniquePtr<UE::Geometry::TGenericDataOperator<UE::Geometry::FMeshVertexBaker>> MakeNewOperator() override;
 	// End IGenericDataOperatorFactory interface
 
-	void SetWorld(UWorld* World);
-
 protected:
 	UPROPERTY()
 	TObjectPtr<UBakeMeshAttributeVertexToolProperties> Settings;
@@ -225,13 +223,6 @@ protected:
 	UPROPERTY()
 	TObjectPtr<UMaterialInstanceDynamic> PreviewAlphaMaterial;
 
-	UPROPERTY()
-	TObjectPtr<UMaterialInstanceDynamic> WorkingPreviewMaterial;
-	float SecondsBeforeWorkingMaterial = 0.75;
-
-	UPROPERTY()
-	TObjectPtr<UWorld> TargetWorld = nullptr;
-
 	TUniquePtr<TGenericDataBackgroundCompute<UE::Geometry::FMeshVertexBaker>> Compute = nullptr;
 	void OnResultUpdated(const TUniquePtr<UE::Geometry::FMeshVertexBaker>& NewResult);
 
@@ -253,8 +244,6 @@ protected:
 
 	const bool bPreferPlatformData = false;
 
-	EBakeOpState OpState = EBakeOpState::Evaluate;
-
 	struct FBakeSettings
 	{
 		EBakeVertexMode VertexMode = EBakeVertexMode::Color;
