ModelingTools: Fix broken max ray distance test in FMeshBakerMeshSceneSampler.

#rnx

diff --git a/Engine/Plugins/Experimental/MeshModelingToolsetExp/Source/MeshModelingToolsExp/Private/BakeMultiMeshAttributeMapsTool.cpp b/Engine/Plugins/Experimental/MeshModelingToolsetExp/Source/MeshModelingToolsExp/Private/BakeMultiMeshAttributeMapsTool.cpp
--- a/Engine/Plugins/Experimental/MeshModelingToolsetExp/Source/MeshModelingToolsExp/Private/BakeMultiMeshAttributeMapsTool.cpp
+++ b/Engine/Plugins/Experimental/MeshModelingToolsetExp/Source/MeshModelingToolsExp/Private/BakeMultiMeshAttributeMapsTool.cpp
@@ -165,7 +165,10 @@ public:
 		FMeshSceneRayHit HitResult;
 		const bool bHit = MeshScene->FindNearestRayIntersection(Ray, HitResult);
 		const void* HitMesh = nullptr;
-		if (bHit && NearestT < Options.MaxDistance)
+
+		// Use TNumericLimits<float>::Max() for consistency with MeshAABBTree3.
+		NearestT = (Options.MaxDistance < TNumericLimits<float>::Max()) ? Options.MaxDistance : TNumericLimits<float>::Max();
+		if (bHit && HitResult.RayDistance < Options.MaxDistance)
 		{
 			TriId = HitResult.HitMeshTriIndex;
 			NearestT = HitResult.RayDistance;
