remove UE::Geometry::TRay3 and replace usages with UE::Math::TRay

#rnx

diff --git a/Engine/Plugins/Experimental/GeometryScripting/Source/GeometryScriptingCore/Private/MeshSpatialFunctions.cpp b/Engine/Plugins/Experimental/GeometryScripting/Source/GeometryScriptingCore/Private/MeshSpatialFunctions.cpp
--- a/Engine/Plugins/Experimental/GeometryScripting/Source/GeometryScriptingCore/Private/MeshSpatialFunctions.cpp
+++ b/Engine/Plugins/Experimental/GeometryScripting/Source/GeometryScriptingCore/Private/MeshSpatialFunctions.cpp
@@ -194,7 +194,7 @@ UDynamicMesh* UGeometryScriptLibrary_MeshSpatial::FindNearestRayIntersectionWith
 					bHaveValidBVH = true;
 					IMeshSpatial::FQueryOptions QueryOptions;
 					QueryOptions.MaxDistance = (Options.MaxDistance == 0) ? TNumericLimits<float>::Max() : Options.MaxDistance;
-					UE::Geometry::FRay3d Ray((FVector3d)RayOrigin, Normalized((FVector3d)RayDirection));
+					FRay3d Ray((FVector3d)RayOrigin, Normalized((FVector3d)RayDirection));
 					int HitTID = QueryBVH.Spatial->FindNearestHitTriangle(Ray, QueryOptions);
 					if (HitTID >= 0)
 					{
diff --git a/Engine/Plugins/Experimental/MeshModelingToolsetExp/Source/GeometryProcessingAdapters/Private/ApproximateActorsImpl.cpp b/Engine/Plugins/Experimental/MeshModelingToolsetExp/Source/GeometryProcessingAdapters/Private/ApproximateActorsImpl.cpp
--- a/Engine/Plugins/Experimental/MeshModelingToolsetExp/Source/GeometryProcessingAdapters/Private/ApproximateActorsImpl.cpp
+++ b/Engine/Plugins/Experimental/MeshModelingToolsetExp/Source/GeometryProcessingAdapters/Private/ApproximateActorsImpl.cpp
@@ -166,7 +166,7 @@ static void BakeTexturesFromPhotoCapture(
 		FVector3d RayDir = ImagePosWorld - SurfPos;
 		double Dist = Normalize(RayDir);
 		FVector3d RayOrigin = SurfPos + RayOffsetHackDist * RayDir;
-		int32 HitTID = Spatial.FindNearestHitTriangle(UE::Geometry::FRay3d(RayOrigin, RayDir), IMeshSpatial::FQueryOptions(Dist));
+		int32 HitTID = Spatial.FindNearestHitTriangle(FRay3d(RayOrigin, RayDir), IMeshSpatial::FQueryOptions(Dist));
 		return (HitTID == IndexConstants::InvalidID);
 	};
 
diff --git a/Engine/Plugins/Experimental/MeshModelingToolsetExp/Source/MeshModelingToolsEditorOnly/Private/PolygonOnMeshTool.cpp b/Engine/Plugins/Experimental/MeshModelingToolsetExp/Source/MeshModelingToolsEditorOnly/Private/PolygonOnMeshTool.cpp
--- a/Engine/Plugins/Experimental/MeshModelingToolsetExp/Source/MeshModelingToolsEditorOnly/Private/PolygonOnMeshTool.cpp
+++ b/Engine/Plugins/Experimental/MeshModelingToolsetExp/Source/MeshModelingToolsEditorOnly/Private/PolygonOnMeshTool.cpp
@@ -424,9 +424,9 @@ bool UPolygonOnMeshTool::HitTest(const FRay& Ray, FHitResult& OutHit)
 	if (DrawPolygonMechanic != nullptr)
 	{
 		FFrame3d HitPoint;
-		if (DrawPolygonMechanic->IsHitByRay(UE::Geometry::FRay3d(Ray), HitPoint))
+		if (DrawPolygonMechanic->IsHitByRay(FRay3d(Ray), HitPoint))
 		{
-			OutHit.Distance = UE::Geometry::FRay3d(Ray).Project(HitPoint.Origin);
+			OutHit.Distance = FRay3d(Ray).GetParameter(HitPoint.Origin);
 			OutHit.ImpactPoint = (FVector)HitPoint.Origin;
 			OutHit.ImpactNormal = (FVector)HitPoint.Z();
 			return true;
@@ -452,7 +452,7 @@ void UPolygonOnMeshTool::OnClicked(const FInputDeviceRay& ClickPos)
 {
 	if (DrawPolygonMechanic != nullptr)
 	{
-		if (DrawPolygonMechanic->TryAddPointFromRay(ClickPos.WorldRay))
+		if (DrawPolygonMechanic->TryAddPointFromRay((FRay3d)ClickPos.WorldRay))
 		{
 			if (DrawPolygonMechanic->IsDone())
 			{
@@ -486,7 +486,7 @@ bool UPolygonOnMeshTool::OnUpdateHover(const FInputDeviceRay& DevicePos)
 {
 	if (DrawPolygonMechanic != nullptr)
 	{
-		DrawPolygonMechanic->UpdatePreviewPoint(DevicePos.WorldRay);
+		DrawPolygonMechanic->UpdatePreviewPoint((FRay3d)DevicePos.WorldRay);
 	}
 	return true;
 }
diff --git a/Engine/Plugins/Experimental/MeshModelingToolsetExp/Source/MeshModelingToolsExp/Private/BakeMultiMeshAttributeMapsTool.cpp b/Engine/Plugins/Experimental/MeshModelingToolsetExp/Source/MeshModelingToolsExp/Private/BakeMultiMeshAttributeMapsTool.cpp
--- a/Engine/Plugins/Experimental/MeshModelingToolsetExp/Source/MeshModelingToolsExp/Private/BakeMultiMeshAttributeMapsTool.cpp
+++ b/Engine/Plugins/Experimental/MeshModelingToolsetExp/Source/MeshModelingToolsExp/Private/BakeMultiMeshAttributeMapsTool.cpp
@@ -154,7 +154,7 @@ public:
 	}
 
 	virtual const void* FindNearestHitTriangle(
-		const UE::Geometry::FRay3d& Ray,
+		const FRay3d& Ray,
 		double& NearestT,
 		int& TriId,
 		FVector3d& TriBaryCoords,
@@ -178,7 +178,7 @@ public:
 	}
 
 	virtual bool TestAnyHitTriangle(
-		const UE::Geometry::FRay3d& Ray,
+		const FRay3d& Ray,
 		const IMeshSpatial::FQueryOptions& Options = IMeshSpatial::FQueryOptions()) const override
 	{
 		// TODO: Add proper test any hit triangle support (for Occlusion)
diff --git a/Engine/Plugins/Experimental/MeshModelingToolsetExp/Source/MeshModelingToolsExp/Private/CubeGridTool.cpp b/Engine/Plugins/Experimental/MeshModelingToolsetExp/Source/MeshModelingToolsExp/Private/CubeGridTool.cpp
--- a/Engine/Plugins/Experimental/MeshModelingToolsetExp/Source/MeshModelingToolsExp/Private/CubeGridTool.cpp
+++ b/Engine/Plugins/Experimental/MeshModelingToolsetExp/Source/MeshModelingToolsExp/Private/CubeGridTool.cpp
@@ -772,7 +772,7 @@ void UCubeGridTool::Setup()
 	};
 	MiddleClickDragBehavior->OnClickPressFunc = [this](const FInputDeviceRay& ClickPos) {
 		PrepForSelectionChange();
-		RayCastSelectionPlane(ClickPos.WorldRay, MiddleClickDragStart);
+		RayCastSelectionPlane((FRay3d)ClickPos.WorldRay, MiddleClickDragStart);
 	};
 	MiddleClickDragBehavior->OnClickDragFunc = [this](const FInputDeviceRay& ClickPos) {
 		OnMiddleClickDrag(ClickPos);
@@ -1055,13 +1055,13 @@ void UCubeGridTool::OnCtrlMiddleClick(const FInputDeviceRay& ClickPos)
 	FVector3d GridSpaceRayOrigin = CubeGrid->ToGridPoint(ClickPos.WorldRay.Origin);
 	FVector3d GridSpaceRayDirection = CubeGrid->GetFrame().ToFrameVector(FVector3d(ClickPos.WorldRay.Direction));
 	GridSpaceRayDirection.Normalize();
-	UE::Geometry::FRay3d GizmoSpaceRay(GridSpaceRayOrigin, GridSpaceRayDirection);
+	FRay3d GizmoSpaceRay(GridSpaceRayOrigin, GridSpaceRayDirection);
 
-	double MinDistSquared = GizmoSpaceRay.DistanceSquared(Corners[0]);
+	double MinDistSquared = GizmoSpaceRay.DistSquared(Corners[0]);
 	int32 ClosestCornerIndex = 0;
 	for (int32 i = 1; i < 4; ++i)
 	{
-		double DistSquared = GizmoSpaceRay.DistanceSquared(Corners[i]);
+		double DistSquared = GizmoSpaceRay.DistSquared(Corners[i]);
 		if (DistSquared < MinDistSquared)
 		{
 			MinDistSquared = DistSquared;
@@ -1101,7 +1101,7 @@ void UCubeGridTool::OnCtrlMiddleClick(const FInputDeviceRay& ClickPos)
 }
 
 // Tries to intersect the selected box. Used for middle mouse dragging the selection.
-FInputRayHit UCubeGridTool::RayCastSelectionPlane(const UE::Geometry::FRay3d& WorldRay,
+FInputRayHit UCubeGridTool::RayCastSelectionPlane(const FRay3d& WorldRay,
 	FVector3d& HitPointOut)
 {
 	FVector3d Normal = CubeGrid->GetFrame().FromFrameVector(
@@ -1115,7 +1115,7 @@ FInputRayHit UCubeGridTool::RayCastSelectionPlane(const UE::Geometry::FRay3d& Wo
 	if (HitResult.bHit)
 	{
 		HitPointOut = HitPoint;
-		HitResult = FInputRayHit(WorldRay.Project(HitPointOut));
+		HitResult = FInputRayHit(WorldRay.GetParameter(HitPointOut));
 	}
 	return HitResult;
 }
@@ -1129,7 +1129,7 @@ FInputRayHit UCubeGridTool::CanBeginMiddleClickDrag(const FInputDeviceRay& Click
 	}
 
 	FVector3d WorldHitPoint;
-	HitResult = RayCastSelectionPlane(ClickPos.WorldRay, WorldHitPoint);
+	HitResult = RayCastSelectionPlane((FRay3d)ClickPos.WorldRay, WorldHitPoint);
 
 	int FlatDim = FCubeGrid::DirToFlatDim(Selection.Direction);
 
@@ -1148,7 +1148,7 @@ void UCubeGridTool::OnMiddleClickDrag(const FInputDeviceRay& DragPos)
 	}
 
 	FVector3d MiddleClickDragEnd;
-	RayCastSelectionPlane(DragPos.WorldRay, MiddleClickDragEnd);
+	RayCastSelectionPlane((FRay3d)DragPos.WorldRay, MiddleClickDragEnd);
 	FVector3d DisplacementInGridFrame = CubeGrid->GetFrame().ToFrameVector(MiddleClickDragEnd - MiddleClickDragStart); 
 	
 	// Clamp the relevant dimension in the displacement vector
@@ -1256,7 +1256,7 @@ bool UCubeGridTool::GetHitGridFace(const FRay& WorldRay, FCubeGrid::FCubeFace& F
 
 	if (MeshSpatial)
 	{
-		UE::Geometry::FRay3d LocalRay(CurrentMeshTransform.InverseTransformPosition((FVector3d)WorldRay.Origin),
+		FRay3d LocalRay(CurrentMeshTransform.InverseTransformPosition((FVector3d)WorldRay.Origin),
 			CurrentMeshTransform.InverseTransformVectorNoScale((FVector3d)WorldRay.Direction));
 
 		int32 Tid;
@@ -1300,7 +1300,7 @@ void UCubeGridTool::OnClickPress(const FInputDeviceRay& PressPos)
 		{
 			PreDragCornerSelectedFlags[i] = CornerSelectedFlags[i];
 		}
-		AttemptToSelectCorner(PressPos.WorldRay);
+		AttemptToSelectCorner((FRay3d)PressPos.WorldRay);
 		return;
 	}
 
@@ -1309,10 +1309,10 @@ void UCubeGridTool::OnClickPress(const FInputDeviceRay& PressPos)
 		MouseState = EMouseState::DraggingExtrudeDistance;
 		if (bHaveSelection)
 		{
-			DragProjectionAxis = UE::Geometry::FRay3d(CubeGrid->GetFrame().FromFramePoint(Selection.Box.Center()),
+			DragProjectionAxis = FRay3d(CubeGrid->GetFrame().FromFramePoint(Selection.Box.Center()),
 				CubeGrid->GetFrame().FromFrameVector(FCubeGrid::DirToNormal(Selection.Direction)), true);
 
-			FDistLine3Ray3d DistanceCalculator(UE::Geometry::FLine3d(DragProjectionAxis.Origin, DragProjectionAxis.Direction), PressPos.WorldRay);
+			FDistLine3Ray3d DistanceCalculator(UE::Geometry::FLine3d(DragProjectionAxis.Origin, DragProjectionAxis.Direction), (FRay3d)PressPos.WorldRay);
 			DistanceCalculator.ComputeResult();
 			DragProjectedStartParam = DistanceCalculator.LineParameter;
 		}
@@ -1367,7 +1367,7 @@ void UCubeGridTool::OnClickDrag(const FInputDeviceRay& DragPos)
 		}
 
 		FDistLine3Ray3d DistanceCalculator(
-			UE::Geometry::FLine3d(DragProjectionAxis.Origin, DragProjectionAxis.Direction), DragPos.WorldRay);
+			UE::Geometry::FLine3d(DragProjectionAxis.Origin, DragProjectionAxis.Direction), (FRay3d)DragPos.WorldRay);
 		DistanceCalculator.ComputeResult();
 
 		double ParamDelta = DistanceCalculator.LineParameter - DragProjectedStartParam;
@@ -1382,7 +1382,7 @@ void UCubeGridTool::OnClickDrag(const FInputDeviceRay& DragPos)
 	}
 	else if (MouseState == EMouseState::DraggingCornerSelection)
 	{
-		AttemptToSelectCorner(DragPos.WorldRay);
+		AttemptToSelectCorner((FRay3d)DragPos.WorldRay);
 		return;
 	}
 	else // Grid selection
@@ -1424,7 +1424,7 @@ void UCubeGridTool::OnTerminateDragSequence()
 	MouseState = EMouseState::NotDragging;
 }
 
-void UCubeGridTool::AttemptToSelectCorner(const UE::Geometry::FRay3d& WorldRay)
+void UCubeGridTool::AttemptToSelectCorner(const FRay3d& WorldRay)
 {
 	TArray<FGeometrySet3::FNearest> HitCorners;
 	CornersGeometrySet.CollectPointsNearRay(WorldRay, HitCorners, [this](const FVector3d& Position1, const FVector3d& Position2) {
diff --git a/Engine/Plugins/Experimental/MeshModelingToolsetExp/Source/MeshModelingToolsExp/Private/DeformMeshPolygonsTool.cpp b/Engine/Plugins/Experimental/MeshModelingToolsetExp/Source/MeshModelingToolsExp/Private/DeformMeshPolygonsTool.cpp
--- a/Engine/Plugins/Experimental/MeshModelingToolsetExp/Source/MeshModelingToolsExp/Private/DeformMeshPolygonsTool.cpp
+++ b/Engine/Plugins/Experimental/MeshModelingToolsetExp/Source/MeshModelingToolsExp/Private/DeformMeshPolygonsTool.cpp
@@ -1041,13 +1041,13 @@ bool UDeformMeshPolygonsTool::HitTest(const FRay& WorldRay, FHitResult& OutHit)
 	if (Selection.SelectedCornerIDs.Num() > 0)
 	{
 		OutHit.FaceIndex   = Selection.GetASelectedCornerID();
-		OutHit.Distance    = LocalRay.Project(LocalPosition);
+		OutHit.Distance    = LocalRay.GetParameter(LocalPosition);
 		OutHit.ImpactPoint = (FVector)WorldTransform.TransformPosition(LocalRay.PointAt(OutHit.Distance));
 	}
 	else if (Selection.SelectedEdgeIDs.Num() > 0)
 	{
 		OutHit.FaceIndex   = Selection.GetASelectedEdgeID();
-		OutHit.Distance    = LocalRay.Project(LocalPosition);
+		OutHit.Distance    = LocalRay.GetParameter(LocalPosition);
 		OutHit.ImpactPoint = (FVector)WorldTransform.TransformPosition(LocalRay.PointAt(OutHit.Distance));
 	}
 	else
diff --git a/Engine/Plugins/Experimental/MeshModelingToolsetExp/Source/MeshModelingToolsExp/Private/DrawPolyPathTool.cpp b/Engine/Plugins/Experimental/MeshModelingToolsetExp/Source/MeshModelingToolsExp/Private/DrawPolyPathTool.cpp
--- a/Engine/Plugins/Experimental/MeshModelingToolsetExp/Source/MeshModelingToolsExp/Private/DrawPolyPathTool.cpp
+++ b/Engine/Plugins/Experimental/MeshModelingToolsetExp/Source/MeshModelingToolsExp/Private/DrawPolyPathTool.cpp
@@ -279,9 +279,9 @@ bool UDrawPolyPathTool::HitTest(const FRay& Ray, FHitResult& OutHit)
 	if (SurfacePathMechanic != nullptr)
 	{
 		FFrame3d HitPoint;
-		if (SurfacePathMechanic->IsHitByRay(UE::Geometry::FRay3d(Ray), HitPoint))
+		if (SurfacePathMechanic->IsHitByRay(FRay3d(Ray), HitPoint))
 		{
-			OutHit.Distance = UE::Geometry::FRay3d(Ray).Project(HitPoint.Origin);
+			OutHit.Distance = FRay3d(Ray).GetParameter(HitPoint.Origin);
 			OutHit.ImpactPoint = (FVector)HitPoint.Origin;
 			OutHit.ImpactNormal = (FVector)HitPoint.Z();
 			return true;
@@ -323,7 +323,7 @@ void UDrawPolyPathTool::OnClicked(const FInputDeviceRay& ClickPos)
 {
 	if (SurfacePathMechanic != nullptr)
 	{
-		if (SurfacePathMechanic->TryAddPointFromRay(ClickPos.WorldRay))
+		if (SurfacePathMechanic->TryAddPointFromRay((FRay3d)ClickPos.WorldRay))
 		{
 			if (SurfacePathMechanic->IsDone())
 			{
@@ -368,7 +368,7 @@ bool UDrawPolyPathTool::OnUpdateHover(const FInputDeviceRay& DevicePos)
 {
 	if (SurfacePathMechanic != nullptr)
 	{
-		SurfacePathMechanic->UpdatePreviewPoint(DevicePos.WorldRay);
+		SurfacePathMechanic->UpdatePreviewPoint((FRay3d)DevicePos.WorldRay);
 	} 
 	else if (CurveDistMechanic != nullptr)
 	{
diff --git a/Engine/Plugins/Experimental/MeshModelingToolsetExp/Source/MeshModelingToolsExp/Private/DynamicMeshBrushTool.cpp b/Engine/Plugins/Experimental/MeshModelingToolsetExp/Source/MeshModelingToolsExp/Private/DynamicMeshBrushTool.cpp
--- a/Engine/Plugins/Experimental/MeshModelingToolsetExp/Source/MeshModelingToolsExp/Private/DynamicMeshBrushTool.cpp
+++ b/Engine/Plugins/Experimental/MeshModelingToolsetExp/Source/MeshModelingToolsExp/Private/DynamicMeshBrushTool.cpp
@@ -78,7 +78,7 @@ void UDynamicMeshBrushTool::Shutdown(EToolShutdownType ShutdownType)
 
 bool UDynamicMeshBrushTool::HitTest(const FRay& Ray, FHitResult& OutHit)
 {
-	return PreviewMesh->FindRayIntersection(UE::Geometry::FRay3d(Ray), OutHit);
+	return PreviewMesh->FindRayIntersection(FRay3d(Ray), OutHit);
 }
 
 
diff --git a/Engine/Plugins/Experimental/MeshModelingToolsetExp/Source/MeshModelingToolsExp/Private/MeshSelectionTool.cpp b/Engine/Plugins/Experimental/MeshModelingToolsetExp/Source/MeshModelingToolsExp/Private/MeshSelectionTool.cpp
--- a/Engine/Plugins/Experimental/MeshModelingToolsetExp/Source/MeshModelingToolsExp/Private/MeshSelectionTool.cpp
+++ b/Engine/Plugins/Experimental/MeshModelingToolsetExp/Source/MeshModelingToolsExp/Private/MeshSelectionTool.cpp
@@ -596,7 +596,7 @@ void UMeshSelectionTool::UpdateFaceSelection(const FBrushStampData& Stamp, const
 		for (int tid : TriangleROI)
 		{
 			FVector3d Centroid = Mesh->GetTriCentroid(tid);
-			int HitTID = GetOctree()->FindNearestHitObject(UE::Geometry::FRay3d(LocalEyePosition, UE::Geometry::Normalized(Centroid - LocalEyePosition)));
+			int HitTID = GetOctree()->FindNearestHitObject(FRay3d(LocalEyePosition, UE::Geometry::Normalized(Centroid - LocalEyePosition)));
 			if (HitTID == tid)
 			{
 				LocalROI.Add(HitTID);
diff --git a/Engine/Plugins/Experimental/MeshModelingToolsetExp/Source/MeshModelingToolsExp/Private/Sculpting/MeshSculptToolBase.cpp b/Engine/Plugins/Experimental/MeshModelingToolsetExp/Source/MeshModelingToolsExp/Private/Sculpting/MeshSculptToolBase.cpp
--- a/Engine/Plugins/Experimental/MeshModelingToolsetExp/Source/MeshModelingToolsExp/Private/Sculpting/MeshSculptToolBase.cpp
+++ b/Engine/Plugins/Experimental/MeshModelingToolsetExp/Source/MeshModelingToolsExp/Private/Sculpting/MeshSculptToolBase.cpp
@@ -579,7 +579,7 @@ void UMeshSculptToolBase::OnEndDrag(const FRay& Ray)
 
 
 
-UE::Geometry::FRay3d UMeshSculptToolBase::GetLocalRay(const FRay& WorldRay) const
+FRay3d UMeshSculptToolBase::GetLocalRay(const FRay& WorldRay) const
 {
 	FRay3d LocalRay(CurTargetTransform.InverseTransformPosition((FVector3d)WorldRay.Origin),
 		CurTargetTransform.InverseTransformVector((FVector3d)WorldRay.Direction));
diff --git a/Engine/Plugins/Experimental/MeshModelingToolsetExp/Source/MeshModelingToolsExp/Public/CubeGridTool.h b/Engine/Plugins/Experimental/MeshModelingToolsetExp/Source/MeshModelingToolsExp/Public/CubeGridTool.h
--- a/Engine/Plugins/Experimental/MeshModelingToolsetExp/Source/MeshModelingToolsExp/Public/CubeGridTool.h
+++ b/Engine/Plugins/Experimental/MeshModelingToolsetExp/Source/MeshModelingToolsExp/Public/CubeGridTool.h
@@ -428,7 +428,7 @@ protected:
 	bool bSelectionToggle = false;
 	bool bChangeSideToggle = false;
 	bool bMouseDragShouldPushPull = false;
-	UE::Geometry::FRay3d DragProjectionAxis;
+	FRay3d DragProjectionAxis;
 	double DragProjectedStartParam;
 
 	void ApplyFlipSelection();
@@ -480,7 +480,7 @@ protected:
 	bool bInGizmoDrag = false;
 
 	FVector3d MiddleClickDragStart;
-	FInputRayHit RayCastSelectionPlane(const UE::Geometry::FRay3d& WorldRay, 
+	FInputRayHit RayCastSelectionPlane(const FRay3d& WorldRay, 
 		FVector3d& HitPointOut);
 	FInputRayHit CanBeginMiddleClickDrag(const FInputDeviceRay& ClickPos);
 	void OnMiddleClickDrag(const FInputDeviceRay& DragPos);
@@ -502,7 +502,7 @@ protected:
 	void StartCornerMode();
 	void ApplyCornerMode(bool bDontWaitForTick = false);
 	void CancelCornerMode();
-	void AttemptToSelectCorner(const UE::Geometry::FRay3d& WorldRay);
+	void AttemptToSelectCorner(const FRay3d& WorldRay);
 
 	// Used to see if we need to update the asset that we've been modifying
 	bool bChangesMade = false;
diff --git a/Engine/Plugins/Experimental/MeshModelingToolsetExp/Source/MeshModelingToolsExp/Public/DeformMeshPolygonsTool.h b/Engine/Plugins/Experimental/MeshModelingToolsetExp/Source/MeshModelingToolsExp/Public/DeformMeshPolygonsTool.h
--- a/Engine/Plugins/Experimental/MeshModelingToolsetExp/Source/MeshModelingToolsExp/Public/DeformMeshPolygonsTool.h
+++ b/Engine/Plugins/Experimental/MeshModelingToolsetExp/Source/MeshModelingToolsExp/Public/DeformMeshPolygonsTool.h
@@ -149,8 +149,6 @@ class MESHMODELINGTOOLSEXP_API UDeformMeshPolygonsTool : public UMeshSurfacePoin
 {
 	GENERATED_BODY()
 
-	using FRay3d = UE::Geometry::FRay3d;	// allows usage of FRay3d in cpp without fully-qualified name, to be removed after LWC conversion
-
 public:
 	UDeformMeshPolygonsTool();
 
diff --git a/Engine/Plugins/Experimental/MeshModelingToolsetExp/Source/MeshModelingToolsExp/Public/DynamicMeshSculptTool.h b/Engine/Plugins/Experimental/MeshModelingToolsetExp/Source/MeshModelingToolsExp/Public/DynamicMeshSculptTool.h
--- a/Engine/Plugins/Experimental/MeshModelingToolsetExp/Source/MeshModelingToolsExp/Public/DynamicMeshSculptTool.h
+++ b/Engine/Plugins/Experimental/MeshModelingToolsetExp/Source/MeshModelingToolsExp/Public/DynamicMeshSculptTool.h
@@ -223,7 +223,6 @@ class MESHMODELINGTOOLSEXP_API UDynamicMeshSculptTool : public UMeshSurfacePoint
 
 public:
 	using FFrame3d = UE::Geometry::FFrame3d;
-	using FRay3d = UE::Geometry::FRay3d;
 
 	UDynamicMeshSculptTool();
 
diff --git a/Engine/Plugins/Experimental/MeshModelingToolsetExp/Source/MeshModelingToolsExp/Public/Sculpting/MeshSculptToolBase.h b/Engine/Plugins/Experimental/MeshModelingToolsetExp/Source/MeshModelingToolsExp/Public/Sculpting/MeshSculptToolBase.h
--- a/Engine/Plugins/Experimental/MeshModelingToolsetExp/Source/MeshModelingToolsExp/Public/Sculpting/MeshSculptToolBase.h
+++ b/Engine/Plugins/Experimental/MeshModelingToolsetExp/Source/MeshModelingToolsExp/Public/Sculpting/MeshSculptToolBase.h
@@ -226,7 +226,6 @@ class MESHMODELINGTOOLSEXP_API UMeshSculptToolBase : public UMeshSurfacePointToo
 protected:
 	using FTransform3d = UE::Geometry::FTransform3d;
 	using FFrame3d = UE::Geometry::FFrame3d;
-	using FRay3d = UE::Geometry::FRay3d;
 public:
 
 	virtual void RegisterActions(FInteractiveToolActionSet& ActionSet) override;
diff --git a/Engine/Plugins/Experimental/UVEditor/Source/UVEditorTools/Private/UVEditorSeamTool.cpp b/Engine/Plugins/Experimental/UVEditor/Source/UVEditorTools/Private/UVEditorSeamTool.cpp
--- a/Engine/Plugins/Experimental/UVEditor/Source/UVEditorTools/Private/UVEditorSeamTool.cpp
+++ b/Engine/Plugins/Experimental/UVEditor/Source/UVEditorTools/Private/UVEditorSeamTool.cpp
@@ -1007,9 +1007,9 @@ int32 UUVEditorSeamTool::Get3DHitVertex(const FRay& WorldRayIn, int32* IndexOf3D
 		for (int32 i = 0; i < Spatials3D.Num(); ++i)
 		{
 			FTransform Transform = Targets[i]->AppliedPreview->PreviewMesh->GetTransform();
-			FRay LocalRay(
-				Transform.InverseTransformPosition(WorldRayIn.Origin),
-				Transform.InverseTransformVector(WorldRayIn.Direction));
+			FRay3d LocalRay(
+				(FVector3d)Transform.InverseTransformPosition(WorldRayIn.Origin),
+				(FVector3d)Transform.InverseTransformVector(WorldRayIn.Direction));
 
 			int32 Tid = IndexConstants::InvalidID;
 			double RayT = TNumericLimits<double>::Max();
@@ -1050,9 +1050,9 @@ int32 UUVEditorSeamTool::Get3DHitVertex(const FRay& WorldRayIn, int32* IndexOf3D
 	{
 		// Only need to hit test the mesh we've already started on, and we always snap to nearest vert
 		FTransform Transform = Targets[ClickedMeshIndex]->AppliedPreview->PreviewMesh->GetTransform();
-		FRay LocalRay(
-			Transform.InverseTransformPosition(WorldRayIn.Origin),
-			Transform.InverseTransformVector(WorldRayIn.Direction));
+		FRay3d LocalRay(
+			(FVector3d)Transform.InverseTransformPosition(WorldRayIn.Origin),
+			(FVector3d)Transform.InverseTransformVector(WorldRayIn.Direction));
 
 		double RayT = TNumericLimits<double>::Max();
 		int32 Tid = IndexConstants::InvalidID;
diff --git a/Engine/Plugins/MeshPainting/Source/MeshPaintingToolset/Private/BaseMeshPaintComponentAdapter.cpp b/Engine/Plugins/MeshPainting/Source/MeshPaintingToolset/Private/BaseMeshPaintComponentAdapter.cpp
--- a/Engine/Plugins/MeshPainting/Source/MeshPaintingToolset/Private/BaseMeshPaintComponentAdapter.cpp
+++ b/Engine/Plugins/MeshPainting/Source/MeshPaintingToolset/Private/BaseMeshPaintComponentAdapter.cpp
@@ -188,7 +188,7 @@ bool FBaseMeshPaintComponentAdapter::RayIntersectAdapter(UE::Geometry::FIndex3i&
 	double NearestT;
 	FVector3d Direction((FVector3d)End - (FVector3d)Start);
 	UE::Geometry::Normalize(Direction);
-	UE::Geometry::FRay3d Ray((FVector3d)Start, Direction);
+	FRay3d Ray((FVector3d)Start, Direction);
 	bool bHit = AABBTree->FindNearestHitTriangle(Ray, NearestT, HitTriangleID);
 	if (bHit)
 	{
diff --git a/Engine/Plugins/MeshPainting/Source/MeshPaintingToolset/Private/MeshPaintStaticMeshAdapter.cpp b/Engine/Plugins/MeshPainting/Source/MeshPaintingToolset/Private/MeshPaintStaticMeshAdapter.cpp
--- a/Engine/Plugins/MeshPainting/Source/MeshPaintingToolset/Private/MeshPaintStaticMeshAdapter.cpp
+++ b/Engine/Plugins/MeshPainting/Source/MeshPaintingToolset/Private/MeshPaintStaticMeshAdapter.cpp
@@ -219,7 +219,7 @@ bool FMeshPaintStaticMeshComponentAdapter::LineTraceComponent(struct FHitResult&
 		if (TriNormal.SquaredLength() > (double)SMALL_NUMBER)
 		{
 			FVector3d RayDirection = ((FVector3d)LocalEnd - (FVector3d)LocalStart);
-			UE::Geometry::FRay3d LocalRay((FVector3d)LocalStart, UE::Geometry::Normalized(RayDirection));
+			FRay3d LocalRay((FVector3d)LocalStart, UE::Geometry::Normalized(RayDirection));
 			UE::Geometry::FIntrRay3Triangle3d RayTriIntersection(LocalRay, Triangle);
 			if (RayTriIntersection.Find())
 			{
diff --git a/Engine/Plugins/Runtime/MeshModelingToolset/Source/MeshModelingTools/Private/EditMeshPolygonsTool.cpp b/Engine/Plugins/Runtime/MeshModelingToolset/Source/MeshModelingTools/Private/EditMeshPolygonsTool.cpp
--- a/Engine/Plugins/Runtime/MeshModelingToolset/Source/MeshModelingTools/Private/EditMeshPolygonsTool.cpp
+++ b/Engine/Plugins/Runtime/MeshModelingToolset/Source/MeshModelingTools/Private/EditMeshPolygonsTool.cpp
@@ -653,7 +653,7 @@ FBox UEditMeshPolygonsTool::GetWorldSpaceFocusBox()
 
 bool UEditMeshPolygonsTool::GetWorldSpaceFocusPoint(const FRay& WorldRay, FVector& PointOut)
 {
-	UE::Geometry::FRay3d LocalRay(WorldTransform.InverseTransformPosition((FVector3d)WorldRay.Origin),
+	FRay3d LocalRay(WorldTransform.InverseTransformPosition((FVector3d)WorldRay.Origin),
 		WorldTransform.InverseTransformNormal((FVector3d)WorldRay.Direction));
 	UE::Geometry::Normalize(LocalRay.Direction);
 
diff --git a/Engine/Plugins/Runtime/MeshModelingToolset/Source/MeshModelingTools/Private/ToolActivities/PolyEditCutFacesActivity.cpp b/Engine/Plugins/Runtime/MeshModelingToolset/Source/MeshModelingTools/Private/ToolActivities/PolyEditCutFacesActivity.cpp
--- a/Engine/Plugins/Runtime/MeshModelingToolset/Source/MeshModelingTools/Private/ToolActivities/PolyEditCutFacesActivity.cpp
+++ b/Engine/Plugins/Runtime/MeshModelingToolset/Source/MeshModelingTools/Private/ToolActivities/PolyEditCutFacesActivity.cpp
@@ -229,7 +229,7 @@ FInputRayHit UPolyEditCutFacesActivity::IsHitByClick(const FInputDeviceRay& Clic
 
 void UPolyEditCutFacesActivity::OnClicked(const FInputDeviceRay& ClickPos)
 {
-	if (bIsRunning && SurfacePathMechanic->TryAddPointFromRay(ClickPos.WorldRay))
+	if (bIsRunning && SurfacePathMechanic->TryAddPointFromRay((FRay3d)ClickPos.WorldRay))
 	{
 		if (SurfacePathMechanic->IsDone())
 		{
@@ -247,7 +247,7 @@ FInputRayHit UPolyEditCutFacesActivity::BeginHoverSequenceHitTest(const FInputDe
 
 bool UPolyEditCutFacesActivity::OnUpdateHover(const FInputDeviceRay& DevicePos)
 {
-	SurfacePathMechanic->UpdatePreviewPoint(DevicePos.WorldRay);
+	SurfacePathMechanic->UpdatePreviewPoint((FRay3d)DevicePos.WorldRay);
 	return bIsRunning;
 }
 
diff --git a/Engine/Plugins/Runtime/MeshModelingToolset/Source/MeshModelingTools/Private/ToolActivities/PolyEditInsertEdgeLoopActivity.cpp b/Engine/Plugins/Runtime/MeshModelingToolset/Source/MeshModelingTools/Private/ToolActivities/PolyEditInsertEdgeLoopActivity.cpp
--- a/Engine/Plugins/Runtime/MeshModelingToolset/Source/MeshModelingTools/Private/ToolActivities/PolyEditInsertEdgeLoopActivity.cpp
+++ b/Engine/Plugins/Runtime/MeshModelingToolset/Source/MeshModelingTools/Private/ToolActivities/PolyEditInsertEdgeLoopActivity.cpp
@@ -351,7 +351,7 @@ FInputRayHit UPolyEditInsertEdgeLoopActivity::HitTest(const FRay& WorldRay)
 	FInputRayHit Hit;
 
 	// See if we hit an edge
-	UE::Geometry::FRay3d LocalRay((FVector3d)TargetTransform.InverseTransformPosition(WorldRay.Origin),
+	FRay3d LocalRay((FVector3d)TargetTransform.InverseTransformPosition(WorldRay.Origin),
 		(FVector3d)TargetTransform.InverseTransformVector(WorldRay.Direction), false);
 	FGroupTopologySelection Selection;
 	FVector3d Position, Normal;
@@ -365,7 +365,7 @@ FInputRayHit UPolyEditInsertEdgeLoopActivity::HitTest(const FRay& WorldRay)
 		const FGroupTopology::FGroupEdge& GroupEdge = ActivityContext->CurrentTopology->Edges[GroupEdgeID];
 		if (GroupEdge.EndpointCorners.A != FDynamicMesh3::InvalidID)
 		{
-			Hit = FInputRayHit(LocalRay.Project(Position));
+			Hit = FInputRayHit(LocalRay.GetParameter(Position));
 		}
 	}
 
@@ -375,7 +375,7 @@ FInputRayHit UPolyEditInsertEdgeLoopActivity::HitTest(const FRay& WorldRay)
 bool UPolyEditInsertEdgeLoopActivity::UpdateHoveredItem(const FRay& WorldRay)
 {
 	// Check that we hit an edge
-	UE::Geometry::FRay3d LocalRay((FVector3d)TargetTransform.InverseTransformPosition(WorldRay.Origin),
+	FRay3d LocalRay((FVector3d)TargetTransform.InverseTransformPosition(WorldRay.Origin),
 		(FVector3d)TargetTransform.InverseTransformVector(WorldRay.Direction), false);
 
 	FGroupTopologySelection Selection;
diff --git a/Engine/Plugins/Runtime/MeshModelingToolset/Source/MeshModelingTools/Private/ToolActivities/PolyEditPlanarProjectionUVActivity.cpp b/Engine/Plugins/Runtime/MeshModelingToolset/Source/MeshModelingTools/Private/ToolActivities/PolyEditPlanarProjectionUVActivity.cpp
--- a/Engine/Plugins/Runtime/MeshModelingToolset/Source/MeshModelingTools/Private/ToolActivities/PolyEditPlanarProjectionUVActivity.cpp
+++ b/Engine/Plugins/Runtime/MeshModelingToolset/Source/MeshModelingTools/Private/ToolActivities/PolyEditPlanarProjectionUVActivity.cpp
@@ -256,7 +256,7 @@ FInputRayHit UPolyEditPlanarProjectionUVActivity::IsHitByClick(const FInputDevic
 
 void UPolyEditPlanarProjectionUVActivity::OnClicked(const FInputDeviceRay& ClickPos)
 {
-	if (bIsRunning && SurfacePathMechanic->TryAddPointFromRay(ClickPos.WorldRay))
+	if (bIsRunning && SurfacePathMechanic->TryAddPointFromRay((FRay3d)ClickPos.WorldRay))
 	{
 		if (SurfacePathMechanic->IsDone())
 		{
@@ -274,7 +274,7 @@ FInputRayHit UPolyEditPlanarProjectionUVActivity::BeginHoverSequenceHitTest(cons
 
 bool UPolyEditPlanarProjectionUVActivity::OnUpdateHover(const FInputDeviceRay& DevicePos)
 {
-	SurfacePathMechanic->UpdatePreviewPoint(DevicePos.WorldRay);
+	SurfacePathMechanic->UpdatePreviewPoint((FRay3d)DevicePos.WorldRay);
 	bPreviewUpdatePending = true;
 	return bIsRunning;
 }
diff --git a/Engine/Plugins/Runtime/MeshModelingToolset/Source/MeshModelingTools/Public/ToolActivities/PolyEditInsertEdgeActivity.h b/Engine/Plugins/Runtime/MeshModelingToolset/Source/MeshModelingTools/Public/ToolActivities/PolyEditInsertEdgeActivity.h
--- a/Engine/Plugins/Runtime/MeshModelingToolset/Source/MeshModelingTools/Public/ToolActivities/PolyEditInsertEdgeActivity.h
+++ b/Engine/Plugins/Runtime/MeshModelingToolset/Source/MeshModelingTools/Public/ToolActivities/PolyEditInsertEdgeActivity.h
@@ -57,7 +57,6 @@ class MESHMODELINGTOOLS_API UPolyEditInsertEdgeActivity : public UInteractiveToo
 
 	using FGroupEdgeInserter = UE::Geometry::FGroupEdgeInserter;
 	using FDynamicMesh3 = UE::Geometry::FDynamicMesh3;
-	using FRay3d = UE::Geometry::FRay3d;
 
 	enum class EState
 	{
diff --git a/Engine/Plugins/Runtime/MeshModelingToolset/Source/ModelingComponents/Private/Mechanics/CollectSurfacePathMechanic.cpp b/Engine/Plugins/Runtime/MeshModelingToolset/Source/ModelingComponents/Private/Mechanics/CollectSurfacePathMechanic.cpp
--- a/Engine/Plugins/Runtime/MeshModelingToolset/Source/ModelingComponents/Private/Mechanics/CollectSurfacePathMechanic.cpp
+++ b/Engine/Plugins/Runtime/MeshModelingToolset/Source/ModelingComponents/Private/Mechanics/CollectSurfacePathMechanic.cpp
@@ -188,7 +188,7 @@ bool UCollectSurfacePathMechanic::RayToPathPoint(const FRay3d& Ray, FFrame3d& Po
 			FIntrRay3Triangle3d Hit = TMeshQueries<FDynamicMesh3>::RayTriangleIntersection(TargetSurface, HitTri, Ray);
 			NearestHitFrame = TargetSurface.GetTriFrame(HitTri);
 			NearestHitFrame.Origin = Hit.Triangle.BarycentricPoint(Hit.TriangleBaryCoords);
-			NearestHitDistSqr = Ray.Project(NearestHitFrame.Origin);
+			NearestHitDistSqr = Ray.GetParameter(NearestHitFrame.Origin);
 			bHaveHit = true;
 		}
 	}
@@ -198,7 +198,7 @@ bool UCollectSurfacePathMechanic::RayToPathPoint(const FRay3d& Ray, FFrame3d& Po
 		FFrame3d PlaneHit(TargetPlane);
 		if (TargetPlane.RayPlaneIntersection(Ray.Origin, Ray.Direction, 2, PlaneHit.Origin))
 		{
-			double HitDistSqr = Ray.Project(PlaneHit.Origin);
+			double HitDistSqr = Ray.GetParameter(PlaneHit.Origin);
 			if (HitDistSqr < NearestHitDistSqr)
 			{
 				NearestHitFrame = PlaneHit;
diff --git a/Engine/Plugins/Runtime/MeshModelingToolset/Source/ModelingComponents/Private/Mechanics/CurveControlPointsMechanic.cpp b/Engine/Plugins/Runtime/MeshModelingToolset/Source/ModelingComponents/Private/Mechanics/CurveControlPointsMechanic.cpp
--- a/Engine/Plugins/Runtime/MeshModelingToolset/Source/ModelingComponents/Private/Mechanics/CurveControlPointsMechanic.cpp
+++ b/Engine/Plugins/Runtime/MeshModelingToolset/Source/ModelingComponents/Private/Mechanics/CurveControlPointsMechanic.cpp
@@ -561,14 +561,14 @@ bool UCurveControlPointsMechanic::HitTest(const FInputDeviceRay& ClickPos, FInpu
 	// Otherwise, see if we are in insert mode and hitting a segment
 	else if (bInsertPointToggle)
 	{
-		if (GeometrySet.FindNearestCurveToRay(ClickPos.WorldRay, Nearest, GeometrySetToleranceTest))
+		if (GeometrySet.FindNearestCurveToRay((FRay3d)ClickPos.WorldRay, Nearest, GeometrySetToleranceTest))
 		{
 			ResultOut = FInputRayHit(Nearest.RayParam);
 			return true;
 		}
 	}
 	// See if we hit a point for selection
-	else if (GeometrySet.FindNearestPointToRay(ClickPos.WorldRay, Nearest, GeometrySetToleranceTest))
+	else if (GeometrySet.FindNearestPointToRay((FRay3d)ClickPos.WorldRay, Nearest, GeometrySetToleranceTest))
 	{
 		ResultOut = FInputRayHit(Nearest.RayParam);
 		return true;
@@ -649,7 +649,7 @@ void UCurveControlPointsMechanic::OnClicked(const FInputDeviceRay& ClickPos)
 	else if (bInsertPointToggle)
 	{
 		// Adding on an existing edge takes priority to adding to the end.
-		if (GeometrySet.FindNearestCurveToRay(ClickPos.WorldRay, Nearest, GeometrySetToleranceTest))
+		if (GeometrySet.FindNearestCurveToRay((FRay3d)ClickPos.WorldRay, Nearest, GeometrySetToleranceTest))
 		{
 			ParentTool->GetToolManager()->BeginUndoTransaction(PointAdditionTransactionText);
 
@@ -704,7 +704,7 @@ void UCurveControlPointsMechanic::OnClicked(const FInputDeviceRay& ClickPos)
 		}
 	}
 	// Otherwise, check for plain old selection
-	else if (GeometrySet.FindNearestPointToRay(ClickPos.WorldRay, Nearest, GeometrySetToleranceTest))
+	else if (GeometrySet.FindNearestPointToRay((FRay3d)ClickPos.WorldRay, Nearest, GeometrySetToleranceTest))
 	{
 		ParentTool->GetToolManager()->BeginUndoTransaction(PointSelectionTransactionText);
 
@@ -907,7 +907,7 @@ bool UCurveControlPointsMechanic::OnUpdateHover(const FInputDeviceRay& DevicePos
 	FGeometrySet3::FNearest Nearest;
 
 	// In edit mode, point insertion on an edge has priority.
-	if (!bInteractiveInitializationMode && bInsertPointToggle && GeometrySet.FindNearestCurveToRay(DevicePos.WorldRay, Nearest, GeometrySetToleranceTest))
+	if (!bInteractiveInitializationMode && bInsertPointToggle && GeometrySet.FindNearestCurveToRay((FRay3d)DevicePos.WorldRay, Nearest, GeometrySetToleranceTest))
 	{
 		ClearHover();
 		FRenderablePoint RenderablePoint((FVector)Nearest.NearestGeoPoint, PreviewColor, PointsSize);
@@ -966,7 +966,7 @@ bool UCurveControlPointsMechanic::OnUpdateHover(const FInputDeviceRay& DevicePos
 	}
 
 	// Otherwise, see if we're hovering a point for selection
-	else if (!bInsertPointToggle && GeometrySet.FindNearestPointToRay(DevicePos.WorldRay, Nearest, GeometrySetToleranceTest))
+	else if (!bInsertPointToggle && GeometrySet.FindNearestPointToRay((FRay3d)DevicePos.WorldRay, Nearest, GeometrySetToleranceTest))
 	{
 		// Only need to update the hover if we changed the point
 		if (Nearest.ID != HoveredPointID)
diff --git a/Engine/Plugins/Runtime/MeshModelingToolset/Source/ModelingComponents/Private/Mechanics/DragAlignmentMechanic.cpp b/Engine/Plugins/Runtime/MeshModelingToolset/Source/ModelingComponents/Private/Mechanics/DragAlignmentMechanic.cpp
--- a/Engine/Plugins/Runtime/MeshModelingToolset/Source/ModelingComponents/Private/Mechanics/DragAlignmentMechanic.cpp
+++ b/Engine/Plugins/Runtime/MeshModelingToolset/Source/ModelingComponents/Private/Mechanics/DragAlignmentMechanic.cpp
@@ -199,7 +199,7 @@ void UDragAlignmentMechanic::InitializeDeformedMeshRayCast(TFunction<FDynamicMes
 	MeshRayCast = [this, GetSpatialIn, TargetTransform](const FRay& WorldRay, FHitResult& HitResult, bool bUseFilter)
 	{
 		// Transform ray into local space
-		UE::Geometry::FRay3d LocalRay(TargetTransform.InverseTransformPosition((FVector3d)WorldRay.Origin),
+		FRay3d LocalRay(TargetTransform.InverseTransformPosition((FVector3d)WorldRay.Origin),
 			TargetTransform.InverseTransformVector((FVector3d)WorldRay.Direction));
 		UE::Geometry::Normalize(LocalRay.Direction);
 
diff --git a/Engine/Plugins/Runtime/MeshModelingToolset/Source/ModelingComponents/Private/Mechanics/LatticeControlPointsMechanic.cpp b/Engine/Plugins/Runtime/MeshModelingToolset/Source/ModelingComponents/Private/Mechanics/LatticeControlPointsMechanic.cpp
--- a/Engine/Plugins/Runtime/MeshModelingToolset/Source/ModelingComponents/Private/Mechanics/LatticeControlPointsMechanic.cpp
+++ b/Engine/Plugins/Runtime/MeshModelingToolset/Source/ModelingComponents/Private/Mechanics/LatticeControlPointsMechanic.cpp
@@ -382,7 +382,7 @@ bool ULatticeControlPointsMechanic::HitTest(const FInputDeviceRay& ClickPos, FIn
 	FGeometrySet3::FNearest Nearest;
 
 	// See if we hit a point for selection
-	if (GeometrySet.FindNearestPointToRay(ClickPos.WorldRay, Nearest, GeometrySetToleranceTest))
+	if (GeometrySet.FindNearestPointToRay((FRay3d)ClickPos.WorldRay, Nearest, GeometrySetToleranceTest))
 	{
 		ResultOut = FInputRayHit(Nearest.RayParam);
 		return true;
@@ -411,7 +411,7 @@ void ULatticeControlPointsMechanic::OnClicked(const FInputDeviceRay& ClickPos)
 	const TSet<int32> PreClickSelection = SelectedPointIDs;
 	
 	FGeometrySet3::FNearest Nearest;
-	if (GeometrySet.FindNearestPointToRay(ClickPos.WorldRay, Nearest, GeometrySetToleranceTest))
+	if (GeometrySet.FindNearestPointToRay((FRay3d)ClickPos.WorldRay, Nearest, GeometrySetToleranceTest))
 	{
 		if (ShouldAddToSelectionFunc())
 		{
@@ -549,7 +549,7 @@ bool ULatticeControlPointsMechanic::OnUpdateHover(const FInputDeviceRay& DeviceP
 	FGeometrySet3::FNearest Nearest;
 
 	// see if we're hovering a point for selection
-	if (GeometrySet.FindNearestPointToRay(DevicePos.WorldRay, Nearest, GeometrySetToleranceTest))
+	if (GeometrySet.FindNearestPointToRay((FRay3d)DevicePos.WorldRay, Nearest, GeometrySetToleranceTest))
 	{
 		// Only need to update the hover if we changed the point
 		if (Nearest.ID != HoveredPointID)
diff --git a/Engine/Plugins/Runtime/MeshModelingToolset/Source/ModelingComponents/Private/Mechanics/PlaneDistanceFromHitMechanic.cpp b/Engine/Plugins/Runtime/MeshModelingToolset/Source/ModelingComponents/Private/Mechanics/PlaneDistanceFromHitMechanic.cpp
--- a/Engine/Plugins/Runtime/MeshModelingToolset/Source/ModelingComponents/Private/Mechanics/PlaneDistanceFromHitMechanic.cpp
+++ b/Engine/Plugins/Runtime/MeshModelingToolset/Source/ModelingComponents/Private/Mechanics/PlaneDistanceFromHitMechanic.cpp
@@ -52,7 +52,7 @@ void UPlaneDistanceFromHitMechanic::UpdateCurrentDistance(const FRay& WorldRay)
 	bCurrentHitIsWorldHit = false;
 
 	// cast ray at target object
-	UE::Geometry::FRay3d LocalRay = PreviewHeightFrame.ToFrame(WorldRay);
+	FRay3d LocalRay = PreviewHeightFrame.ToFrame((FRay3d)WorldRay);
 	int HitTID = PreviewHeightTargetAABB.FindNearestHitTriangle(LocalRay);
 	if (HitTID >= 0)
 	{
@@ -94,7 +94,7 @@ void UPlaneDistanceFromHitMechanic::UpdateCurrentDistance(const FRay& WorldRay)
 
 	if (bFoundHit == false && bFallbackToLineAxisPoint)
 	{
-		FDistLine3Ray3d Distance(FLine3d(PreviewHeightFrame.Origin, PreviewHeightFrame.Z()), UE::Geometry::FRay3d(WorldRay));
+		FDistLine3Ray3d Distance(FLine3d(PreviewHeightFrame.Origin, PreviewHeightFrame.Z()), FRay3d(WorldRay));
 		float WorldHitDist = WorldRay.GetParameter( (FVector)Distance.RayClosestPoint );
 		NearestHitFrameWorld = FFrame3d(Distance.RayClosestPoint, UE::Geometry::Normalized(Distance.RayClosestPoint - Distance.LineClosestPoint) );
 		FVector3d HitPosLocal = PreviewHeightFrame.ToFramePoint(Distance.RayClosestPoint);
diff --git a/Engine/Plugins/Runtime/MeshModelingToolset/Source/ModelingComponents/Private/Mechanics/SpaceCurveDeformationMechanic.cpp b/Engine/Plugins/Runtime/MeshModelingToolset/Source/ModelingComponents/Private/Mechanics/SpaceCurveDeformationMechanic.cpp
--- a/Engine/Plugins/Runtime/MeshModelingToolset/Source/ModelingComponents/Private/Mechanics/SpaceCurveDeformationMechanic.cpp
+++ b/Engine/Plugins/Runtime/MeshModelingToolset/Source/ModelingComponents/Private/Mechanics/SpaceCurveDeformationMechanic.cpp
@@ -416,7 +416,7 @@ void USpaceCurveDeformationMechanic::UpdateCurve(const TArray<FFrame3d>& NewPosi
 bool USpaceCurveDeformationMechanic::HitTest(const FInputDeviceRay& ClickPos, FInputRayHit& ResultOut)
 {
 	FGeometrySet3::FNearest Nearest;
-	if (GeometrySet.FindNearestPointToRay(ClickPos.WorldRay, Nearest, GeometrySetToleranceTest))
+	if (GeometrySet.FindNearestPointToRay((FRay3d)ClickPos.WorldRay, Nearest, GeometrySetToleranceTest))
 	{
 		ResultOut = FInputRayHit(Nearest.RayParam);
 		return true;
@@ -435,7 +435,7 @@ void USpaceCurveDeformationMechanic::OnClicked(const FInputDeviceRay& ClickPos)
 {
 	FGeometrySet3::FNearest Nearest;
 
-	if (GeometrySet.FindNearestPointToRay(ClickPos.WorldRay, Nearest, GeometrySetToleranceTest))
+	if (GeometrySet.FindNearestPointToRay((FRay3d)ClickPos.WorldRay, Nearest, GeometrySetToleranceTest))
 	{
 		ParentTool->GetToolManager()->BeginUndoTransaction(SpaceCurvePointSelectionTransactionText);
 
@@ -786,7 +786,7 @@ void USpaceCurveDeformationMechanic::ClearHover()
 bool USpaceCurveDeformationMechanic::OnUpdateHover(const FInputDeviceRay& DevicePos)
 {
 	FGeometrySet3::FNearest Nearest;
-	if (GeometrySet.FindNearestPointToRay(DevicePos.WorldRay, Nearest, GeometrySetToleranceTest))
+	if (GeometrySet.FindNearestPointToRay((FRay3d)DevicePos.WorldRay, Nearest, GeometrySetToleranceTest))
 	{
 		HoveredPointID = Nearest.ID;
 		return true;
diff --git a/Engine/Plugins/Runtime/MeshModelingToolset/Source/ModelingComponents/Private/Mechanics/SpatialCurveDistanceMechanic.cpp b/Engine/Plugins/Runtime/MeshModelingToolset/Source/ModelingComponents/Private/Mechanics/SpatialCurveDistanceMechanic.cpp
--- a/Engine/Plugins/Runtime/MeshModelingToolset/Source/ModelingComponents/Private/Mechanics/SpatialCurveDistanceMechanic.cpp
+++ b/Engine/Plugins/Runtime/MeshModelingToolset/Source/ModelingComponents/Private/Mechanics/SpatialCurveDistanceMechanic.cpp
@@ -37,7 +37,7 @@ void USpatialCurveDistanceMechanic::UpdateCurrentDistance(const FRay& WorldRay)
 	FVector3d RayNearest;
 	FVector3d CurveNearest;
 
-	UE::Geometry::FRay3d Ray(WorldRay);
+	FRay3d Ray(WorldRay);
 	int NumV = Curve.Num() - 1;
 	double MinDistSqr = TNumericLimits<double>::Max();
 	for (int k = 0; k < NumV; ++k)
diff --git a/Engine/Plugins/Runtime/MeshModelingToolset/Source/ModelingComponents/Private/Scene/MeshSceneAdapter.cpp b/Engine/Plugins/Runtime/MeshModelingToolset/Source/ModelingComponents/Private/Scene/MeshSceneAdapter.cpp
--- a/Engine/Plugins/Runtime/MeshModelingToolset/Source/ModelingComponents/Private/Scene/MeshSceneAdapter.cpp
+++ b/Engine/Plugins/Runtime/MeshModelingToolset/Source/ModelingComponents/Private/Scene/MeshSceneAdapter.cpp
@@ -431,7 +431,7 @@ public:
 			else
 			{
 				FVector3d WorldPos = LocalToWorldTransform.TransformPosition(UseRay.PointAt(RayHitT));
-				WorldHitResultOut.RayDistance = WorldRay.Project(WorldPos);
+				WorldHitResultOut.RayDistance = WorldRay.GetParameter(WorldPos);
 			}
 			WorldHitResultOut.HitMeshBaryCoords = HitBaryCoords;
 			return true;
@@ -762,7 +762,7 @@ public:
 			WorldHitResultOut.HitMeshTriIndex = HitTID;
 			WorldHitResultOut.HitMeshSpatialWrapper = this;
 			FVector3d WorldPos = LocalToWorldTransform.TransformPosition(LocalRay.PointAt(LocalHitT));
-			WorldHitResultOut.RayDistance = WorldRay.Project(WorldPos);
+			WorldHitResultOut.RayDistance = WorldRay.GetParameter(WorldPos);
 			WorldHitResultOut.HitMeshBaryCoords = HitBaryCoords;
 			return true;
 		}
diff --git a/Engine/Plugins/Runtime/MeshModelingToolset/Source/ModelingComponents/Private/Selection/GroupTopologySelector.cpp b/Engine/Plugins/Runtime/MeshModelingToolset/Source/ModelingComponents/Private/Selection/GroupTopologySelector.cpp
--- a/Engine/Plugins/Runtime/MeshModelingToolset/Source/ModelingComponents/Private/Selection/GroupTopologySelector.cpp
+++ b/Engine/Plugins/Runtime/MeshModelingToolset/Source/ModelingComponents/Private/Selection/GroupTopologySelector.cpp
@@ -148,7 +148,7 @@ bool FGroupTopologySelector::FindSelectedElement(const FSelectionSettings& Setti
 	// if we have both corner and edge, want to keep the one we are closer to
 	if (bHaveCornerHit && bHaveEdgeHit)
 	{
-		if (PointsWithinToleranceTest(CornerPosition, Ray.NearestPoint(CornerPosition), 0.75))
+		if (PointsWithinToleranceTest(CornerPosition, Ray.ClosestPoint(CornerPosition), 0.75))
 		{
 			bHaveEdgeHit = false;
 		}
@@ -162,7 +162,7 @@ bool FGroupTopologySelector::FindSelectedElement(const FSelectionSettings& Setti
 	if ((bHaveCornerHit || bHaveEdgeHit) && bHaveFaceHit)
 	{
 		FVector3d TestPos = (bHaveCornerHit) ? CornerPosition : EdgePosition;
-		if (!PointsWithinToleranceTest(TestPos, Ray.NearestPoint(TestPos), 0.15))
+		if (!PointsWithinToleranceTest(TestPos, Ray.ClosestPoint(TestPos), 0.15))
 		{
 			bHaveEdgeHit = bHaveCornerHit = false;
 		}
@@ -573,7 +573,7 @@ bool IsOccluded(const FGeometrySet3::FNearest& ClosestElement, const FVector3d&
 bool IsOccluded(const FVector3d& Point, const FVector3d& ViewOrigin, const FDynamicMeshAABBTree3* Spatial)
 {
 	// Shoot ray backwards to see if we hit something. 
-	UE::Geometry::FRay3d ToEyeRay(Point, UE::Geometry::Normalized(ViewOrigin - Point), true);
+	FRay3d ToEyeRay(Point, UE::Geometry::Normalized(ViewOrigin - Point), true);
 	ToEyeRay.Origin += (double)(100 * FMathf::ZeroTolerance) * ToEyeRay.Direction;
 	if (Spatial->FindNearestHitTriangle(ToEyeRay) >= 0)
 	{
diff --git a/Engine/Plugins/Runtime/MeshModelingToolset/Source/ModelingComponents/Private/Selection/MeshSelectionMechanic.cpp b/Engine/Plugins/Runtime/MeshModelingToolset/Source/ModelingComponents/Private/Selection/MeshSelectionMechanic.cpp
--- a/Engine/Plugins/Runtime/MeshModelingToolset/Source/ModelingComponents/Private/Selection/MeshSelectionMechanic.cpp
+++ b/Engine/Plugins/Runtime/MeshModelingToolset/Source/ModelingComponents/Private/Selection/MeshSelectionMechanic.cpp
@@ -889,9 +889,9 @@ void UMeshSelectionMechanic::OnClicked(const FInputDeviceRay& ClickPos)
 		int32 HitTid = IndexConstants::InvalidID;
 		for (int32 MeshIndex = 0; MeshIndex < MeshSpatials.Num(); ++MeshIndex)
 		{
-			FRay LocalRay(
-				MeshTransforms[MeshIndex].InverseTransformPosition(ClickPos.WorldRay.Origin),
-				MeshTransforms[MeshIndex].InverseTransformVector(ClickPos.WorldRay.Direction));
+			FRay3d LocalRay(
+				(FVector3d)MeshTransforms[MeshIndex].InverseTransformPosition(ClickPos.WorldRay.Origin),
+				(FVector3d)MeshTransforms[MeshIndex].InverseTransformVector(ClickPos.WorldRay.Direction));
 
 			double RayT = 0; 
 			if (MeshSpatials[MeshIndex]->FindNearestHitTriangle(LocalRay, RayT, HitTid))
@@ -932,7 +932,7 @@ void UMeshSelectionMechanic::OnClicked(const FInputDeviceRay& ClickPos)
 					GeometrySet.AddCurve(Eids[i], Polyline);
 				}
 				FGeometrySet3::FNearest Result;
-				if (GeometrySet.FindNearestCurveToRay(ClickPos.WorldRay, Result, 
+				if (GeometrySet.FindNearestCurveToRay((FRay3d)ClickPos.WorldRay, Result, 
 					[this](const FVector3d& Position1, const FVector3d& Position2) {
 						return ToolSceneQueriesUtil::PointSnapQuery(CameraState,
 							Position1, Position2,
@@ -955,7 +955,7 @@ void UMeshSelectionMechanic::OnClicked(const FInputDeviceRay& ClickPos)
 					GeometrySet.AddPoint(Vids[i], CurrentSelection.Mesh->GetTriVertex(HitTid, i));
 				}
 				FGeometrySet3::FNearest Result;
-				if (GeometrySet.FindNearestPointToRay(ClickPos.WorldRay, Result,
+				if (GeometrySet.FindNearestPointToRay((FRay3d)ClickPos.WorldRay, Result,
 					[this](const FVector3d& Position1, const FVector3d& Position2) {
 						return ToolSceneQueriesUtil::PointSnapQuery(CameraState,
 							Position1, Position2,
diff --git a/Engine/Plugins/Runtime/MeshModelingToolset/Source/ModelingComponents/Private/Selection/PolygonSelectionMechanic.cpp b/Engine/Plugins/Runtime/MeshModelingToolset/Source/ModelingComponents/Private/Selection/PolygonSelectionMechanic.cpp
--- a/Engine/Plugins/Runtime/MeshModelingToolset/Source/ModelingComponents/Private/Selection/PolygonSelectionMechanic.cpp
+++ b/Engine/Plugins/Runtime/MeshModelingToolset/Source/ModelingComponents/Private/Selection/PolygonSelectionMechanic.cpp
@@ -300,13 +300,13 @@ bool UPolygonSelectionMechanic::TopologyHitTest(const FRay& WorldRay, FHitResult
 	if (OutSelection.SelectedCornerIDs.Num() > 0)
 	{
 		OutHit.FaceIndex = OutSelection.GetASelectedCornerID();
-		OutHit.Distance = LocalRay.Project(LocalPosition);
+		OutHit.Distance = LocalRay.GetParameter(LocalPosition);
 		OutHit.ImpactPoint = (FVector)TargetTransform.TransformPosition(LocalRay.PointAt(OutHit.Distance));
 	}
 	else if (OutSelection.SelectedEdgeIDs.Num() > 0)
 	{
 		OutHit.FaceIndex = OutSelection.GetASelectedEdgeID();
-		OutHit.Distance = LocalRay.Project(LocalPosition);
+		OutHit.Distance = LocalRay.GetParameter(LocalPosition);
 		OutHit.ImpactPoint = (FVector)TargetTransform.TransformPosition(LocalRay.PointAt(OutHit.Distance));
 		OutHit.Item = EdgeSegmentId;
 	}
diff --git a/Engine/Plugins/Runtime/MeshModelingToolset/Source/ModelingComponents/Private/Snapping/RaySpatialSnapSolver.cpp b/Engine/Plugins/Runtime/MeshModelingToolset/Source/ModelingComponents/Private/Snapping/RaySpatialSnapSolver.cpp
--- a/Engine/Plugins/Runtime/MeshModelingToolset/Source/ModelingComponents/Private/Snapping/RaySpatialSnapSolver.cpp
+++ b/Engine/Plugins/Runtime/MeshModelingToolset/Source/ModelingComponents/Private/Snapping/RaySpatialSnapSolver.cpp
@@ -90,7 +90,7 @@ void FRaySpatialSnapSolver::UpdateSnappedPoint(const FRay3d& RayIn)
 
 	TFunction<FVector3d(const FVector3d&)> GetSnapFromPointFunc = [&RayIn](const FVector3d& Point)
 	{
-		return RayIn.NearestPoint(Point);
+		return RayIn.ClosestPoint(Point);
 	};
 
 	const FSnapTargetPoint* BestPointTarget = FindBestSnapInSet(TargetPoints, MinMetric, BestPriority, GetSnapFromPointFunc);
@@ -120,7 +120,7 @@ void FRaySpatialSnapSolver::UpdateSnappedPoint(const FRay3d& RayIn)
 		// used for grid+line snapping for example. We want to snap to the grid *after* we snap to the
 		// line, otherwise things snap to weird unexpected places.
 		FVector3d UseSnapPosition = (BestSnapTarget->bHaveConstrainedPosition) ? BestSnapTarget->ConstrainedPosition : BestSnapTarget->Position;
-		SetActiveSnapData(*BestSnapTarget, RayIn.NearestPoint(BestSnapTarget->Position), UseSnapPosition, MinMetric);
+		SetActiveSnapData(*BestSnapTarget, RayIn.ClosestPoint(BestSnapTarget->Position), UseSnapPosition, MinMetric);
 	}
 	else
 	{
diff --git a/Engine/Plugins/Runtime/MeshModelingToolset/Source/ModelingComponents/Public/BaseTools/BaseCreateFromSelectedTool.h b/Engine/Plugins/Runtime/MeshModelingToolset/Source/ModelingComponents/Public/BaseTools/BaseCreateFromSelectedTool.h
--- a/Engine/Plugins/Runtime/MeshModelingToolset/Source/ModelingComponents/Public/BaseTools/BaseCreateFromSelectedTool.h
+++ b/Engine/Plugins/Runtime/MeshModelingToolset/Source/ModelingComponents/Public/BaseTools/BaseCreateFromSelectedTool.h
@@ -101,7 +101,6 @@ class MODELINGCOMPONENTS_API UBaseCreateFromSelectedTool : public UMultiSelectio
 protected:
 	using FTransform3d = UE::Geometry::FTransform3d;
 	using FFrame3d = UE::Geometry::FFrame3d;
-	using FRay3d = UE::Geometry::FRay3d;
 public:
 	UBaseCreateFromSelectedTool() = default;
 
diff --git a/Engine/Plugins/Runtime/MeshModelingToolset/Source/ModelingComponents/Public/BaseTools/BaseMeshProcessingTool.h b/Engine/Plugins/Runtime/MeshModelingToolset/Source/ModelingComponents/Public/BaseTools/BaseMeshProcessingTool.h
--- a/Engine/Plugins/Runtime/MeshModelingToolset/Source/ModelingComponents/Public/BaseTools/BaseMeshProcessingTool.h
+++ b/Engine/Plugins/Runtime/MeshModelingToolset/Source/ModelingComponents/Public/BaseTools/BaseMeshProcessingTool.h
@@ -79,7 +79,6 @@ class MODELINGCOMPONENTS_API UBaseMeshProcessingTool : public USingleSelectionTo
 protected:
 	using FTransform3d = UE::Geometry::FTransform3d;
 	using FFrame3d = UE::Geometry::FFrame3d;
-	using FRay3d = UE::Geometry::FRay3d;
 public:
 	UBaseMeshProcessingTool() = default;
 
diff --git a/Engine/Plugins/Runtime/MeshModelingToolset/Source/ModelingComponents/Public/Mechanics/CollectSurfacePathMechanic.h b/Engine/Plugins/Runtime/MeshModelingToolset/Source/ModelingComponents/Public/Mechanics/CollectSurfacePathMechanic.h
--- a/Engine/Plugins/Runtime/MeshModelingToolset/Source/ModelingComponents/Public/Mechanics/CollectSurfacePathMechanic.h
+++ b/Engine/Plugins/Runtime/MeshModelingToolset/Source/ModelingComponents/Public/Mechanics/CollectSurfacePathMechanic.h
@@ -30,7 +30,6 @@ class MODELINGCOMPONENTS_API UCollectSurfacePathMechanic : public UInteractionMe
 	GENERATED_BODY()
 public:
 	using FFrame3d = UE::Geometry::FFrame3d;
-	using FRay3d = UE::Geometry::FRay3d;
 
 	TUniqueFunction<bool()> IsDoneFunc = nullptr;
 
diff --git a/Engine/Plugins/Runtime/MeshModelingToolset/Source/ModelingComponents/Public/PreviewMesh.h b/Engine/Plugins/Runtime/MeshModelingToolset/Source/ModelingComponents/Public/PreviewMesh.h
--- a/Engine/Plugins/Runtime/MeshModelingToolset/Source/ModelingComponents/Public/PreviewMesh.h
+++ b/Engine/Plugins/Runtime/MeshModelingToolset/Source/ModelingComponents/Public/PreviewMesh.h
@@ -57,7 +57,6 @@ UCLASS(Transient)
 class MODELINGCOMPONENTS_API UPreviewMesh : public UObject, public IMeshVertexCommandChangeTarget, public IMeshCommandChangeTarget, public IMeshReplacementCommandChangeTarget
 {
 	GENERATED_BODY()
-	using FRay3d = UE::Geometry::FRay3d;
 public:
 	UPreviewMesh();
 	virtual ~UPreviewMesh();
diff --git a/Engine/Plugins/Runtime/MeshModelingToolset/Source/ModelingComponents/Public/Selection/GroupTopologySelector.h b/Engine/Plugins/Runtime/MeshModelingToolset/Source/ModelingComponents/Public/Selection/GroupTopologySelector.h
--- a/Engine/Plugins/Runtime/MeshModelingToolset/Source/ModelingComponents/Public/Selection/GroupTopologySelector.h
+++ b/Engine/Plugins/Runtime/MeshModelingToolset/Source/ModelingComponents/Public/Selection/GroupTopologySelector.h
@@ -5,7 +5,6 @@
 #include "GroupTopology.h"
 #include "Spatial/GeometrySet3.h"
 #include "DynamicMesh/DynamicMeshAABBTree3.h"
-#include "RayTypes.h"
 
 struct FCameraRectangle;
 class FToolDataVisualizer;
@@ -35,7 +34,6 @@ using UE::Geometry::FGeometrySet3;
  */
 class MODELINGCOMPONENTS_API FGroupTopologySelector
 {
-	using FRay3d = UE::Geometry::FRay3d;
 	using FTransform3d = UE::Geometry::FTransform3d;
 public:
 
diff --git a/Engine/Plugins/Runtime/MeshModelingToolset/Source/ModelingComponents/Public/Selection/PolygonSelectionMechanic.h b/Engine/Plugins/Runtime/MeshModelingToolset/Source/ModelingComponents/Public/Selection/PolygonSelectionMechanic.h
--- a/Engine/Plugins/Runtime/MeshModelingToolset/Source/ModelingComponents/Public/Selection/PolygonSelectionMechanic.h
+++ b/Engine/Plugins/Runtime/MeshModelingToolset/Source/ModelingComponents/Public/Selection/PolygonSelectionMechanic.h
@@ -84,7 +84,6 @@ class MODELINGCOMPONENTS_API UPolygonSelectionMechanic : public UInteractionMech
 	using FFrame3d = UE::Geometry::FFrame3d;
 	using FTransform3d = UE::Geometry::FTransform3d;	
 	using FAxisAlignedBox3d = UE::Geometry::FAxisAlignedBox3d;
-	using FRay3d = UE::Geometry::FRay3d;
 public:
 
 	virtual ~UPolygonSelectionMechanic();
diff --git a/Engine/Source/Runtime/Core/Public/Math/Ray.h b/Engine/Source/Runtime/Core/Public/Math/Ray.h
--- a/Engine/Source/Runtime/Core/Public/Math/Ray.h
+++ b/Engine/Source/Runtime/Core/Public/Math/Ray.h
@@ -99,6 +99,17 @@ public:
 		}
 	}
 
+	/**
+	 * Find minimum distance from query point to ray
+	 *
+	 * @param Point query Point
+	 * @return distance to Ray
+	 */
+	T Dist(const TVector<T>& Point) const
+	{
+		return FMath::Sqrt(DistSquared(Point));
+	}
+
 	/**
 	 * Find closest point on ray to query point
 	 * @param Point query point
diff --git a/Engine/Source/Runtime/GeometryCore/Private/Spatial/GeometrySet3.cpp b/Engine/Source/Runtime/GeometryCore/Private/Spatial/GeometrySet3.cpp
--- a/Engine/Source/Runtime/GeometryCore/Private/Spatial/GeometrySet3.cpp
+++ b/Engine/Source/Runtime/GeometryCore/Private/Spatial/GeometrySet3.cpp
@@ -105,8 +105,8 @@ bool FGeometrySet3::FindNearestPointToRay(const FRay3d& Ray, FNearest& ResultOut
 	ParallelFor(NumPoints, [&](int pi)
 	{
 		const FPoint& Point = Points[pi];
-		double RayT = Ray.Project(Point.Position);
-		FVector3d RayPt = Ray.NearestPoint(Point.Position);
+		double RayT = Ray.GetParameter(Point.Position);
+		FVector3d RayPt = Ray.ClosestPoint(Point.Position);
 		if (PointWithinToleranceTest(RayPt, Point.Position))
 		{
 			Critical.Lock();
@@ -150,8 +150,8 @@ bool FGeometrySet3::CollectPointsNearRay(const FRay3d& Ray, TArray<FNearest>& Re
 	ParallelFor(NumPoints, [&](int pi)
 	{
 		const FPoint& Point = Points[pi];
-		double RayT = Ray.Project(Point.Position);
-		FVector3d RayPt = Ray.NearestPoint(Point.Position);
+		double RayT = Ray.GetParameter(Point.Position);
+		FVector3d RayPt = Ray.ClosestPoint(Point.Position);
 		if (PointWithinToleranceTest(RayPt, Point.Position))
 		{
 			WithinTolerance[pi] = true;
diff --git a/Engine/Source/Runtime/GeometryCore/Public/Curve/CurveUtil.h b/Engine/Source/Runtime/GeometryCore/Public/Curve/CurveUtil.h
--- a/Engine/Source/Runtime/GeometryCore/Public/Curve/CurveUtil.h
+++ b/Engine/Source/Runtime/GeometryCore/Public/Curve/CurveUtil.h
@@ -5,7 +5,6 @@
 #include "VectorTypes.h"
 #include "SegmentTypes.h"
 #include "LineTypes.h"
-#include "RayTypes.h"
 #include "BoxTypes.h"
 #include "SegmentTypes.h"
 
diff --git a/Engine/Source/Runtime/GeometryCore/Public/Distance/DistLine3Ray3.h b/Engine/Source/Runtime/GeometryCore/Public/Distance/DistLine3Ray3.h
--- a/Engine/Source/Runtime/GeometryCore/Public/Distance/DistLine3Ray3.h
+++ b/Engine/Source/Runtime/GeometryCore/Public/Distance/DistLine3Ray3.h
@@ -6,7 +6,6 @@
 
 #include "VectorTypes.h"
 #include "LineTypes.h"
-#include "RayTypes.h"
 
 namespace UE
 {
@@ -24,7 +23,7 @@ class TDistLine3Ray3
 public:
 	// Input
 	TLine3<Real> Line;
-	TRay3<Real> Ray;
+	TRay<Real> Ray;
 
 	// Results
 	Real DistanceSquared = -1.0;
@@ -35,7 +34,7 @@ public:
 	Real RayParameter;
 
 
-	TDistLine3Ray3(const TLine3<Real>& LineIn, const TRay3<Real>& RayIn)
+	TDistLine3Ray3(const TLine3<Real>& LineIn, const TRay<Real>& RayIn)
 	{
 		Line = LineIn;
 		Ray = RayIn;
diff --git a/Engine/Source/Runtime/GeometryCore/Public/Distance/DistRay3Segment3.h b/Engine/Source/Runtime/GeometryCore/Public/Distance/DistRay3Segment3.h
--- a/Engine/Source/Runtime/GeometryCore/Public/Distance/DistRay3Segment3.h
+++ b/Engine/Source/Runtime/GeometryCore/Public/Distance/DistRay3Segment3.h
@@ -6,7 +6,6 @@
 
 #include "VectorTypes.h"
 #include "SegmentTypes.h"
-#include "RayTypes.h"
 
 namespace UE
 {
@@ -23,7 +22,7 @@ class TDistRay3Segment3
 {
 public:
 	// Input
-	TRay3<Real> Ray;
+	TRay<Real> Ray;
 	TSegment3<Real> Segment;
 
 	// Results
@@ -35,7 +34,7 @@ public:
 	Real SegmentParameter;
 
 
-	TDistRay3Segment3(const TRay3<Real>& RayIn, const TSegment3<Real>& SegmentIn)
+	TDistRay3Segment3(const TRay<Real>& RayIn, const TSegment3<Real>& SegmentIn)
 	{
 		Ray = RayIn;
 		Segment = SegmentIn;
@@ -223,7 +222,7 @@ public:
 
 
 
-	static double SquaredDistance(const TRay3<Real>& Ray, const TSegment3<Real>& Segment,
+	static double SquaredDistance(const TRay<Real>& Ray, const TSegment3<Real>& Segment,
 		Real& RayParam, Real& SegParam)
 	{
 		TVector<Real> diff = Ray.Origin - Segment.Center;
diff --git a/Engine/Source/Runtime/GeometryCore/Public/FrameTypes.h b/Engine/Source/Runtime/GeometryCore/Public/FrameTypes.h
--- a/Engine/Source/Runtime/GeometryCore/Public/FrameTypes.h
+++ b/Engine/Source/Runtime/GeometryCore/Public/FrameTypes.h
@@ -5,7 +5,6 @@
 
 #include "VectorTypes.h"
 #include "VectorUtil.h"
-#include "RayTypes.h"
 #include "Quaternion.h"
 #include "TransformTypes.h"
 
@@ -232,14 +231,14 @@ struct TFrame3
 
 
 	/** @return input Ray transformed into local coordinate system of Frame */
-	TRay3<RealType> ToFrame(const TRay3<RealType>& Ray) const
+	TRay<RealType> ToFrame(const TRay<RealType>& Ray) const
 	{
-		return TRay3<RealType>(ToFramePoint(Ray.Origin), ToFrameVector(Ray.Direction));
+		return TRay<RealType>(ToFramePoint(Ray.Origin), ToFrameVector(Ray.Direction));
 	}
 	/** @return input Ray transformed from local coordinate system of Frame into "World" coordinate system */
-	TRay3<RealType> FromFrame(const TRay3<RealType>& Ray) const
+	TRay<RealType> FromFrame(const TRay<RealType>& Ray) const
 	{
-		return TRay3<RealType>(FromFramePoint(Ray.Origin), FromFrameVector(Ray.Direction));
+		return TRay<RealType>(FromFramePoint(Ray.Origin), FromFrameVector(Ray.Direction));
 	}
 
 
diff --git a/Engine/Source/Runtime/GeometryCore/Public/Intersection/IntrRay3AxisAlignedBox3.h b/Engine/Source/Runtime/GeometryCore/Public/Intersection/IntrRay3AxisAlignedBox3.h
--- a/Engine/Source/Runtime/GeometryCore/Public/Intersection/IntrRay3AxisAlignedBox3.h
+++ b/Engine/Source/Runtime/GeometryCore/Public/Intersection/IntrRay3AxisAlignedBox3.h
@@ -4,7 +4,6 @@
 
 #pragma once
 
-#include "RayTypes.h"
 #include "BoxTypes.h"
 #include "VectorUtil.h"
 
@@ -34,7 +33,7 @@ public:
 	 * @param ExpandExtents box is expanded by this amount in each direction, useful for dealing with float imprecision
 	 * @return true if ray intersects box
 	 */
-	static bool TestIntersection(const TRay3<RealType>& Ray, const TAxisAlignedBox3<RealType>& Box, RealType ExpandExtents = 0)
+	static bool TestIntersection(const TRay<RealType>& Ray, const TAxisAlignedBox3<RealType>& Box, RealType ExpandExtents = 0)
 	{
 		TVector<RealType> WdU = TVector<RealType>::Zero();
 		TVector<RealType> AWdU = TVector<RealType>::Zero();
@@ -108,7 +107,7 @@ public:
 	 * @param RayParamOut ray intersect T-value, or TNumericLimits::Max()
 	 * @return true if ray intersects box
 	 */
-	static bool FindIntersection(const TRay3<RealType>& Ray, const TAxisAlignedBox3<RealType>& Box, RealType& RayParamOut)
+	static bool FindIntersection(const TRay<RealType>& Ray, const TAxisAlignedBox3<RealType>& Box, RealType& RayParamOut)
 	{
 		RealType RayParam0 = 0.0;
 		RealType RayParam1 = TNumericLimits<RealType>::Max();
diff --git a/Engine/Source/Runtime/GeometryCore/Public/Intersection/IntrRay3Triangle3.h b/Engine/Source/Runtime/GeometryCore/Public/Intersection/IntrRay3Triangle3.h
--- a/Engine/Source/Runtime/GeometryCore/Public/Intersection/IntrRay3Triangle3.h
+++ b/Engine/Source/Runtime/GeometryCore/Public/Intersection/IntrRay3Triangle3.h
@@ -5,7 +5,6 @@
 #pragma once
 
 #include "VectorTypes.h"
-#include "RayTypes.h"
 #include "TriangleTypes.h"
 #include "VectorUtil.h"
 
@@ -24,7 +23,7 @@ class TIntrRay3Triangle3
 {
 public:
 	// Input
-	TRay3<Real> Ray;
+	TRay<Real> Ray;
 	TTriangle3<Real> Triangle;
 
 	// Output
@@ -33,7 +32,7 @@ public:
 	EIntersectionType IntersectionType;
 
 
-	TIntrRay3Triangle3(const TRay3<Real>& RayIn, const TTriangle3<Real>& TriangleIn)
+	TIntrRay3Triangle3(const TRay<Real>& RayIn, const TTriangle3<Real>& TriangleIn)
 	{
 		Ray = RayIn;
 		Triangle = TriangleIn;
diff --git a/Engine/Source/Runtime/GeometryCore/Public/Polyline3.h b/Engine/Source/Runtime/GeometryCore/Public/Polyline3.h
--- a/Engine/Source/Runtime/GeometryCore/Public/Polyline3.h
+++ b/Engine/Source/Runtime/GeometryCore/Public/Polyline3.h
@@ -5,7 +5,6 @@
 #include "VectorTypes.h"
 #include "SegmentTypes.h"
 #include "LineTypes.h"
-#include "RayTypes.h"
 #include "BoxTypes.h"
 
 namespace UE
diff --git a/Engine/Source/Runtime/GeometryCore/Public/RayTypes.h b/Engine/Source/Runtime/GeometryCore/Public/RayTypes.h
deleted file mode 100644
diff --git a/Engine/Source/Runtime/GeometryCore/Public/Spatial/GeometrySet3.h b/Engine/Source/Runtime/GeometryCore/Public/Spatial/GeometrySet3.h
--- a/Engine/Source/Runtime/GeometryCore/Public/Spatial/GeometrySet3.h
+++ b/Engine/Source/Runtime/GeometryCore/Public/Spatial/GeometrySet3.h
@@ -4,7 +4,6 @@
 
 #include "Async/ParallelFor.h"
 #include "BoxTypes.h"
-#include "RayTypes.h"
 #include "Polyline3.h"
 
 namespace UE
diff --git a/Engine/Source/Runtime/GeometryCore/Public/Spatial/SpatialInterfaces.h b/Engine/Source/Runtime/GeometryCore/Public/Spatial/SpatialInterfaces.h
--- a/Engine/Source/Runtime/GeometryCore/Public/Spatial/SpatialInterfaces.h
+++ b/Engine/Source/Runtime/GeometryCore/Public/Spatial/SpatialInterfaces.h
@@ -5,7 +5,6 @@
 #pragma once
 
 #include "VectorTypes.h"
-#include "RayTypes.h"
 
 namespace UE
 {
