InteractiveToolsFramework: Remove GetTargetRequirements from UInteractiveToolBuilder.

#rnx

diff --git a/Engine/Plugins/Experimental/MeshLODToolset/Source/MeshLODToolset/Public/Tools/GenerateStaticMeshLODAssetTool.h b/Engine/Plugins/Experimental/MeshLODToolset/Source/MeshLODToolset/Public/Tools/GenerateStaticMeshLODAssetTool.h
--- a/Engine/Plugins/Experimental/MeshLODToolset/Source/MeshLODToolset/Public/Tools/GenerateStaticMeshLODAssetTool.h
+++ b/Engine/Plugins/Experimental/MeshLODToolset/Source/MeshLODToolset/Public/Tools/GenerateStaticMeshLODAssetTool.h
@@ -49,7 +49,7 @@ public:
 	virtual UInteractiveTool* BuildTool(const FToolBuilderState& SceneState) const override;
 
 protected:
-	virtual const FToolTargetTypeRequirements& GetTargetRequirements() const override;
+	virtual const FToolTargetTypeRequirements& GetTargetRequirements() const;
 };
 
 
diff --git a/Engine/Plugins/Experimental/MeshLODToolset/Source/MeshLODToolset/Public/Tools/LODManagerTool.h b/Engine/Plugins/Experimental/MeshLODToolset/Source/MeshLODToolset/Public/Tools/LODManagerTool.h
--- a/Engine/Plugins/Experimental/MeshLODToolset/Source/MeshLODToolset/Public/Tools/LODManagerTool.h
+++ b/Engine/Plugins/Experimental/MeshLODToolset/Source/MeshLODToolset/Public/Tools/LODManagerTool.h
@@ -27,7 +27,7 @@ public:
 	virtual UInteractiveTool* BuildTool(const FToolBuilderState& SceneState) const override;
 
 protected:
-	virtual const FToolTargetTypeRequirements& GetTargetRequirements() const override;
+	virtual const FToolTargetTypeRequirements& GetTargetRequirements() const;
 };
 
 
diff --git a/Engine/Plugins/Experimental/MeshModelingToolsetExp/Source/MeshModelingToolsEditorOnly/Public/AddPivotActorTool.h b/Engine/Plugins/Experimental/MeshModelingToolsetExp/Source/MeshModelingToolsEditorOnly/Public/AddPivotActorTool.h
--- a/Engine/Plugins/Experimental/MeshModelingToolsetExp/Source/MeshModelingToolsEditorOnly/Public/AddPivotActorTool.h
+++ b/Engine/Plugins/Experimental/MeshModelingToolsetExp/Source/MeshModelingToolsEditorOnly/Public/AddPivotActorTool.h
@@ -24,7 +24,7 @@ public:
 	virtual UInteractiveTool* BuildTool(const FToolBuilderState& SceneState) const override;
 
 protected:
-	virtual const FToolTargetTypeRequirements& GetTargetRequirements() const override;
+	virtual const FToolTargetTypeRequirements& GetTargetRequirements() const;
 };
 
 /** 
diff --git a/Engine/Plugins/Experimental/MeshModelingToolsetExp/Source/MeshModelingToolsEditorOnly/Public/AttributeEditorTool.h b/Engine/Plugins/Experimental/MeshModelingToolsetExp/Source/MeshModelingToolsEditorOnly/Public/AttributeEditorTool.h
--- a/Engine/Plugins/Experimental/MeshModelingToolsetExp/Source/MeshModelingToolsEditorOnly/Public/AttributeEditorTool.h
+++ b/Engine/Plugins/Experimental/MeshModelingToolsetExp/Source/MeshModelingToolsEditorOnly/Public/AttributeEditorTool.h
@@ -22,7 +22,7 @@ public:
 	virtual UInteractiveTool* BuildTool(const FToolBuilderState& SceneState) const override;
 
 protected:
-	virtual const FToolTargetTypeRequirements& GetTargetRequirements() const override;
+	virtual const FToolTargetTypeRequirements& GetTargetRequirements() const;
 };
 
 
diff --git a/Engine/Plugins/Experimental/MeshModelingToolsetExp/Source/MeshModelingToolsEditorOnly/Public/MergeMeshesTool.h b/Engine/Plugins/Experimental/MeshModelingToolsetExp/Source/MeshModelingToolsEditorOnly/Public/MergeMeshesTool.h
--- a/Engine/Plugins/Experimental/MeshModelingToolsetExp/Source/MeshModelingToolsEditorOnly/Public/MergeMeshesTool.h
+++ b/Engine/Plugins/Experimental/MeshModelingToolsetExp/Source/MeshModelingToolsEditorOnly/Public/MergeMeshesTool.h
@@ -24,7 +24,7 @@ public:
 	virtual UInteractiveTool* BuildTool(const FToolBuilderState& SceneState) const override;
 
 protected:
-	virtual const FToolTargetTypeRequirements& GetTargetRequirements() const override;
+	virtual const FToolTargetTypeRequirements& GetTargetRequirements() const;
 };
 
 
diff --git a/Engine/Plugins/Experimental/MeshModelingToolsetExp/Source/MeshModelingToolsEditorOnly/Public/VoxelCSGMeshesTool.h b/Engine/Plugins/Experimental/MeshModelingToolsetExp/Source/MeshModelingToolsEditorOnly/Public/VoxelCSGMeshesTool.h
--- a/Engine/Plugins/Experimental/MeshModelingToolsetExp/Source/MeshModelingToolsEditorOnly/Public/VoxelCSGMeshesTool.h
+++ b/Engine/Plugins/Experimental/MeshModelingToolsetExp/Source/MeshModelingToolsEditorOnly/Public/VoxelCSGMeshesTool.h
@@ -25,7 +25,7 @@ public:
 	virtual UInteractiveTool* BuildTool(const FToolBuilderState& SceneState) const override;
 
 protected:
-	virtual const FToolTargetTypeRequirements& GetTargetRequirements() const override;
+	virtual const FToolTargetTypeRequirements& GetTargetRequirements() const;
 };
 
 
diff --git a/Engine/Plugins/Experimental/MeshModelingToolsetExp/Source/MeshModelingToolsExp/Private/SkinWeightsBindingTool.cpp b/Engine/Plugins/Experimental/MeshModelingToolsetExp/Source/MeshModelingToolsExp/Private/SkinWeightsBindingTool.cpp
--- a/Engine/Plugins/Experimental/MeshModelingToolsetExp/Source/MeshModelingToolsExp/Private/SkinWeightsBindingTool.cpp
+++ b/Engine/Plugins/Experimental/MeshModelingToolsetExp/Source/MeshModelingToolsExp/Private/SkinWeightsBindingTool.cpp
@@ -720,7 +720,7 @@ private:
 
 bool USkinWeightsBindingToolBuilder::CanBuildTool(const FToolBuilderState& SceneState) const
 {
-	return SceneState.TargetManager->CountSelectedAndTargetable(SceneState, GetTargetRequirements()) == 1;
+	return SceneState.TargetManager->CountSelectedAndTargetable(SceneState, FToolTargetTypeRequirements()) == 1;
 }
 
 
@@ -728,7 +728,7 @@ UInteractiveTool* USkinWeightsBindingToolBuilder::BuildTool(const FToolBuilderSt
 {
 	USkinWeightsBindingTool* NewTool = NewObject<USkinWeightsBindingTool>(SceneState.ToolManager);
 
-	TArray<TObjectPtr<UToolTarget>> Targets = SceneState.TargetManager->BuildAllSelectedTargetable(SceneState, GetTargetRequirements());
+	TArray<TObjectPtr<UToolTarget>> Targets = SceneState.TargetManager->BuildAllSelectedTargetable(SceneState, FToolTargetTypeRequirements());
 	NewTool->SetTargets(MoveTemp(Targets));
 	NewTool->SetWorld(SceneState.World);
 
diff --git a/Engine/Plugins/Experimental/MeshModelingToolsetExp/Source/MeshModelingToolsExp/Public/AlignObjectsTool.h b/Engine/Plugins/Experimental/MeshModelingToolsetExp/Source/MeshModelingToolsExp/Public/AlignObjectsTool.h
--- a/Engine/Plugins/Experimental/MeshModelingToolsetExp/Source/MeshModelingToolsExp/Public/AlignObjectsTool.h
+++ b/Engine/Plugins/Experimental/MeshModelingToolsetExp/Source/MeshModelingToolsExp/Public/AlignObjectsTool.h
@@ -26,7 +26,7 @@ public:
 	virtual UInteractiveTool* BuildTool(const FToolBuilderState& SceneState) const override;
 
 protected:
-	virtual const FToolTargetTypeRequirements& GetTargetRequirements() const override;
+	virtual const FToolTargetTypeRequirements& GetTargetRequirements() const;
 };
 
 
diff --git a/Engine/Plugins/Experimental/MeshModelingToolsetExp/Source/MeshModelingToolsExp/Public/BakeMeshAttributeMapsTool.h b/Engine/Plugins/Experimental/MeshModelingToolsetExp/Source/MeshModelingToolsExp/Public/BakeMeshAttributeMapsTool.h
--- a/Engine/Plugins/Experimental/MeshModelingToolsetExp/Source/MeshModelingToolsExp/Public/BakeMeshAttributeMapsTool.h
+++ b/Engine/Plugins/Experimental/MeshModelingToolsetExp/Source/MeshModelingToolsExp/Public/BakeMeshAttributeMapsTool.h
@@ -38,7 +38,7 @@ public:
 	virtual UInteractiveTool* BuildTool(const FToolBuilderState& SceneState) const override;
 
 protected:
-	virtual const FToolTargetTypeRequirements& GetTargetRequirements() const override;
+	virtual const FToolTargetTypeRequirements& GetTargetRequirements() const;
 };
 
 
diff --git a/Engine/Plugins/Experimental/MeshModelingToolsetExp/Source/MeshModelingToolsExp/Public/BakeMeshAttributeVertexTool.h b/Engine/Plugins/Experimental/MeshModelingToolsetExp/Source/MeshModelingToolsExp/Public/BakeMeshAttributeVertexTool.h
--- a/Engine/Plugins/Experimental/MeshModelingToolsetExp/Source/MeshModelingToolsExp/Public/BakeMeshAttributeVertexTool.h
+++ b/Engine/Plugins/Experimental/MeshModelingToolsetExp/Source/MeshModelingToolsExp/Public/BakeMeshAttributeVertexTool.h
@@ -31,7 +31,7 @@ public:
 	virtual UInteractiveTool* BuildTool(const FToolBuilderState& SceneState) const override;
 
 protected:
-	virtual const FToolTargetTypeRequirements& GetTargetRequirements() const override;
+	virtual const FToolTargetTypeRequirements& GetTargetRequirements() const;
 };
 
 
diff --git a/Engine/Plugins/Experimental/MeshModelingToolsetExp/Source/MeshModelingToolsExp/Public/BakeMultiMeshAttributeMapsTool.h b/Engine/Plugins/Experimental/MeshModelingToolsetExp/Source/MeshModelingToolsExp/Public/BakeMultiMeshAttributeMapsTool.h
--- a/Engine/Plugins/Experimental/MeshModelingToolsetExp/Source/MeshModelingToolsExp/Public/BakeMultiMeshAttributeMapsTool.h
+++ b/Engine/Plugins/Experimental/MeshModelingToolsetExp/Source/MeshModelingToolsExp/Public/BakeMultiMeshAttributeMapsTool.h
@@ -28,7 +28,7 @@ class MESHMODELINGTOOLSEXP_API UBakeMultiMeshAttributeMapsToolBuilder : public U
 	virtual UInteractiveTool* BuildTool(const FToolBuilderState& SceneState) const override;
 
 	protected:
-	virtual const FToolTargetTypeRequirements& GetTargetRequirements() const override;
+	virtual const FToolTargetTypeRequirements& GetTargetRequirements() const;
 };
 
 
diff --git a/Engine/Plugins/Experimental/MeshModelingToolsetExp/Source/MeshModelingToolsExp/Public/BakeTransformTool.h b/Engine/Plugins/Experimental/MeshModelingToolsetExp/Source/MeshModelingToolsExp/Public/BakeTransformTool.h
--- a/Engine/Plugins/Experimental/MeshModelingToolsetExp/Source/MeshModelingToolsExp/Public/BakeTransformTool.h
+++ b/Engine/Plugins/Experimental/MeshModelingToolsetExp/Source/MeshModelingToolsExp/Public/BakeTransformTool.h
@@ -30,7 +30,7 @@ public:
 	virtual UInteractiveTool* BuildTool(const FToolBuilderState& SceneState) const override;
 
 protected:
-	virtual const FToolTargetTypeRequirements& GetTargetRequirements() const override;
+	virtual const FToolTargetTypeRequirements& GetTargetRequirements() const;
 };
 
 
diff --git a/Engine/Plugins/Experimental/MeshModelingToolsetExp/Source/MeshModelingToolsExp/Public/ConvertMeshesTool.h b/Engine/Plugins/Experimental/MeshModelingToolsetExp/Source/MeshModelingToolsExp/Public/ConvertMeshesTool.h
--- a/Engine/Plugins/Experimental/MeshModelingToolsetExp/Source/MeshModelingToolsExp/Public/ConvertMeshesTool.h
+++ b/Engine/Plugins/Experimental/MeshModelingToolsetExp/Source/MeshModelingToolsExp/Public/ConvertMeshesTool.h
@@ -20,7 +20,7 @@ public:
 	virtual UInteractiveTool* BuildTool(const FToolBuilderState& SceneState) const override;
 
 protected:
-	virtual const FToolTargetTypeRequirements& GetTargetRequirements() const override;
+	virtual const FToolTargetTypeRequirements& GetTargetRequirements() const;
 };
 
 /**
diff --git a/Engine/Plugins/Experimental/MeshModelingToolsetExp/Source/MeshModelingToolsExp/Public/CubeGridTool.h b/Engine/Plugins/Experimental/MeshModelingToolsetExp/Source/MeshModelingToolsExp/Public/CubeGridTool.h
--- a/Engine/Plugins/Experimental/MeshModelingToolsetExp/Source/MeshModelingToolsExp/Public/CubeGridTool.h
+++ b/Engine/Plugins/Experimental/MeshModelingToolsetExp/Source/MeshModelingToolsExp/Public/CubeGridTool.h
@@ -54,7 +54,7 @@ public:
 
 protected:
 
-	virtual const FToolTargetTypeRequirements& GetTargetRequirements() const override;
+	virtual const FToolTargetTypeRequirements& GetTargetRequirements() const;
 };
 
 
diff --git a/Engine/Plugins/Experimental/MeshModelingToolsetExp/Source/MeshModelingToolsExp/Public/EditNormalsTool.h b/Engine/Plugins/Experimental/MeshModelingToolsetExp/Source/MeshModelingToolsExp/Public/EditNormalsTool.h
--- a/Engine/Plugins/Experimental/MeshModelingToolsetExp/Source/MeshModelingToolsExp/Public/EditNormalsTool.h
+++ b/Engine/Plugins/Experimental/MeshModelingToolsetExp/Source/MeshModelingToolsExp/Public/EditNormalsTool.h
@@ -30,7 +30,7 @@ public:
 	virtual UInteractiveTool* BuildTool(const FToolBuilderState& SceneState) const override;
 
 protected:
-	virtual const FToolTargetTypeRequirements& GetTargetRequirements() const override;
+	virtual const FToolTargetTypeRequirements& GetTargetRequirements() const;
 };
 
 
diff --git a/Engine/Plugins/Experimental/MeshModelingToolsetExp/Source/MeshModelingToolsExp/Public/EditPivotTool.h b/Engine/Plugins/Experimental/MeshModelingToolsetExp/Source/MeshModelingToolsExp/Public/EditPivotTool.h
--- a/Engine/Plugins/Experimental/MeshModelingToolsetExp/Source/MeshModelingToolsExp/Public/EditPivotTool.h
+++ b/Engine/Plugins/Experimental/MeshModelingToolsetExp/Source/MeshModelingToolsExp/Public/EditPivotTool.h
@@ -30,7 +30,7 @@ public:
 	virtual UInteractiveTool* BuildTool(const FToolBuilderState& SceneState) const override;
 
 protected:
-	virtual const FToolTargetTypeRequirements& GetTargetRequirements() const override;
+	virtual const FToolTargetTypeRequirements& GetTargetRequirements() const;
 };
 
 
diff --git a/Engine/Plugins/Experimental/MeshModelingToolsetExp/Source/MeshModelingToolsExp/Public/MirrorTool.h b/Engine/Plugins/Experimental/MeshModelingToolsetExp/Source/MeshModelingToolsExp/Public/MirrorTool.h
--- a/Engine/Plugins/Experimental/MeshModelingToolsetExp/Source/MeshModelingToolsExp/Public/MirrorTool.h
+++ b/Engine/Plugins/Experimental/MeshModelingToolsetExp/Source/MeshModelingToolsExp/Public/MirrorTool.h
@@ -24,7 +24,7 @@ public:
 	virtual UInteractiveTool* BuildTool(const FToolBuilderState& SceneState) const override;
 
 protected:
-	virtual const FToolTargetTypeRequirements& GetTargetRequirements() const override;
+	virtual const FToolTargetTypeRequirements& GetTargetRequirements() const;
 };
 
 UENUM()
diff --git a/Engine/Plugins/Experimental/MeshModelingToolsetExp/Source/MeshModelingToolsExp/Public/Physics/PhysicsInspectorTool.h b/Engine/Plugins/Experimental/MeshModelingToolsetExp/Source/MeshModelingToolsExp/Public/Physics/PhysicsInspectorTool.h
--- a/Engine/Plugins/Experimental/MeshModelingToolsetExp/Source/MeshModelingToolsExp/Public/Physics/PhysicsInspectorTool.h
+++ b/Engine/Plugins/Experimental/MeshModelingToolsetExp/Source/MeshModelingToolsExp/Public/Physics/PhysicsInspectorTool.h
@@ -20,7 +20,7 @@ public:
 	virtual UInteractiveTool* BuildTool(const FToolBuilderState& SceneState) const override;
 
 protected:
-	virtual const FToolTargetTypeRequirements& GetTargetRequirements() const override;
+	virtual const FToolTargetTypeRequirements& GetTargetRequirements() const;
 };
 
 
diff --git a/Engine/Plugins/Experimental/MeshModelingToolsetExp/Source/MeshModelingToolsExp/Public/Physics/SetCollisionGeometryTool.h b/Engine/Plugins/Experimental/MeshModelingToolsetExp/Source/MeshModelingToolsExp/Public/Physics/SetCollisionGeometryTool.h
--- a/Engine/Plugins/Experimental/MeshModelingToolsetExp/Source/MeshModelingToolsExp/Public/Physics/SetCollisionGeometryTool.h
+++ b/Engine/Plugins/Experimental/MeshModelingToolsetExp/Source/MeshModelingToolsExp/Public/Physics/SetCollisionGeometryTool.h
@@ -28,7 +28,7 @@ public:
 	virtual UInteractiveTool* BuildTool(const FToolBuilderState& SceneState) const override;
 
 protected:
-	virtual const FToolTargetTypeRequirements& GetTargetRequirements() const override;
+	virtual const FToolTargetTypeRequirements& GetTargetRequirements() const;
 };
 
 
diff --git a/Engine/Plugins/Experimental/MeshModelingToolsetExp/Source/MeshModelingToolsExp/Public/PlaneCutTool.h b/Engine/Plugins/Experimental/MeshModelingToolsetExp/Source/MeshModelingToolsExp/Public/PlaneCutTool.h
--- a/Engine/Plugins/Experimental/MeshModelingToolsetExp/Source/MeshModelingToolsExp/Public/PlaneCutTool.h
+++ b/Engine/Plugins/Experimental/MeshModelingToolsetExp/Source/MeshModelingToolsExp/Public/PlaneCutTool.h
@@ -27,7 +27,7 @@ public:
 	virtual UInteractiveTool* BuildTool(const FToolBuilderState& SceneState) const override;
 
 protected:
-	virtual const FToolTargetTypeRequirements& GetTargetRequirements() const override;
+	virtual const FToolTargetTypeRequirements& GetTargetRequirements() const;
 };
 
 
diff --git a/Engine/Plugins/Experimental/MeshModelingToolsetExp/Source/MeshModelingToolsExp/Public/ProjectToTargetTool.h b/Engine/Plugins/Experimental/MeshModelingToolsetExp/Source/MeshModelingToolsExp/Public/ProjectToTargetTool.h
--- a/Engine/Plugins/Experimental/MeshModelingToolsetExp/Source/MeshModelingToolsExp/Public/ProjectToTargetTool.h
+++ b/Engine/Plugins/Experimental/MeshModelingToolsetExp/Source/MeshModelingToolsExp/Public/ProjectToTargetTool.h
@@ -18,7 +18,7 @@ public:
 	virtual UInteractiveTool* BuildTool(const FToolBuilderState& SceneState) const override;
 
 protected:
-	virtual const FToolTargetTypeRequirements& GetTargetRequirements() const override;
+	virtual const FToolTargetTypeRequirements& GetTargetRequirements() const;
 };
 
 
diff --git a/Engine/Plugins/Experimental/MeshModelingToolsetExp/Source/MeshModelingToolsExp/Public/RemeshMeshTool.h b/Engine/Plugins/Experimental/MeshModelingToolsetExp/Source/MeshModelingToolsExp/Public/RemeshMeshTool.h
--- a/Engine/Plugins/Experimental/MeshModelingToolsetExp/Source/MeshModelingToolsExp/Public/RemeshMeshTool.h
+++ b/Engine/Plugins/Experimental/MeshModelingToolsetExp/Source/MeshModelingToolsExp/Public/RemeshMeshTool.h
@@ -33,7 +33,7 @@ public:
 	virtual UInteractiveTool* BuildTool(const FToolBuilderState& SceneState) const override;
 
 protected:
-	virtual const FToolTargetTypeRequirements& GetTargetRequirements() const override;
+	virtual const FToolTargetTypeRequirements& GetTargetRequirements() const;
 };
 
 /**
diff --git a/Engine/Plugins/Experimental/MeshModelingToolsetExp/Source/MeshModelingToolsExp/Public/RemoveOccludedTrianglesTool.h b/Engine/Plugins/Experimental/MeshModelingToolsetExp/Source/MeshModelingToolsExp/Public/RemoveOccludedTrianglesTool.h
--- a/Engine/Plugins/Experimental/MeshModelingToolsetExp/Source/MeshModelingToolsExp/Public/RemoveOccludedTrianglesTool.h
+++ b/Engine/Plugins/Experimental/MeshModelingToolsetExp/Source/MeshModelingToolsExp/Public/RemoveOccludedTrianglesTool.h
@@ -30,7 +30,7 @@ public:
 	virtual UInteractiveTool* BuildTool(const FToolBuilderState& SceneState) const override;
 
 protected:
-	virtual const FToolTargetTypeRequirements& GetTargetRequirements() const override;
+	virtual const FToolTargetTypeRequirements& GetTargetRequirements() const;
 };
 
 
diff --git a/Engine/Plugins/Experimental/MeshModelingToolsetExp/Source/MeshModelingToolsExp/Public/SplitMeshesTool.h b/Engine/Plugins/Experimental/MeshModelingToolsetExp/Source/MeshModelingToolsExp/Public/SplitMeshesTool.h
--- a/Engine/Plugins/Experimental/MeshModelingToolsetExp/Source/MeshModelingToolsExp/Public/SplitMeshesTool.h
+++ b/Engine/Plugins/Experimental/MeshModelingToolsetExp/Source/MeshModelingToolsExp/Public/SplitMeshesTool.h
@@ -20,7 +20,7 @@ public:
 	virtual UInteractiveTool* BuildTool(const FToolBuilderState& SceneState) const override;
 
 protected:
-	virtual const FToolTargetTypeRequirements& GetTargetRequirements() const override;
+	virtual const FToolTargetTypeRequirements& GetTargetRequirements() const;
 };
 
 
diff --git a/Engine/Plugins/Experimental/MeshModelingToolsetExp/Source/MeshModelingToolsExp/Public/TransferMeshTool.h b/Engine/Plugins/Experimental/MeshModelingToolsetExp/Source/MeshModelingToolsExp/Public/TransferMeshTool.h
--- a/Engine/Plugins/Experimental/MeshModelingToolsetExp/Source/MeshModelingToolsExp/Public/TransferMeshTool.h
+++ b/Engine/Plugins/Experimental/MeshModelingToolsetExp/Source/MeshModelingToolsExp/Public/TransferMeshTool.h
@@ -20,7 +20,7 @@ public:
 	virtual UInteractiveTool* BuildTool(const FToolBuilderState& SceneState) const override;
 
 protected:
-	virtual const FToolTargetTypeRequirements& GetTargetRequirements() const override;
+	virtual const FToolTargetTypeRequirements& GetTargetRequirements() const;
 };
 
 /**
diff --git a/Engine/Plugins/Experimental/MeshModelingToolsetExp/Source/MeshModelingToolsExp/Public/TransformMeshesTool.h b/Engine/Plugins/Experimental/MeshModelingToolsetExp/Source/MeshModelingToolsExp/Public/TransformMeshesTool.h
--- a/Engine/Plugins/Experimental/MeshModelingToolsetExp/Source/MeshModelingToolsExp/Public/TransformMeshesTool.h
+++ b/Engine/Plugins/Experimental/MeshModelingToolsetExp/Source/MeshModelingToolsExp/Public/TransformMeshesTool.h
@@ -30,7 +30,7 @@ public:
 	virtual UInteractiveTool* BuildTool(const FToolBuilderState& SceneState) const override;
 
 protected:
-	virtual const FToolTargetTypeRequirements& GetTargetRequirements() const override;
+	virtual const FToolTargetTypeRequirements& GetTargetRequirements() const;
 };
 
 
diff --git a/Engine/Plugins/Experimental/MeshModelingToolsetExp/Source/MeshModelingToolsExp/Public/UVLayoutTool.h b/Engine/Plugins/Experimental/MeshModelingToolsetExp/Source/MeshModelingToolsExp/Public/UVLayoutTool.h
--- a/Engine/Plugins/Experimental/MeshModelingToolsetExp/Source/MeshModelingToolsExp/Public/UVLayoutTool.h
+++ b/Engine/Plugins/Experimental/MeshModelingToolsetExp/Source/MeshModelingToolsExp/Public/UVLayoutTool.h
@@ -38,7 +38,7 @@ public:
 	virtual UInteractiveTool* BuildTool(const FToolBuilderState& SceneState) const override;
 
 protected:
-	virtual const FToolTargetTypeRequirements& GetTargetRequirements() const override;
+	virtual const FToolTargetTypeRequirements& GetTargetRequirements() const;
 };
 
 
diff --git a/Engine/Plugins/Experimental/UVEditor/Source/UVEditorTools/Public/UVEditorLayerEditTool.h b/Engine/Plugins/Experimental/UVEditor/Source/UVEditorTools/Public/UVEditorLayerEditTool.h
--- a/Engine/Plugins/Experimental/UVEditor/Source/UVEditorTools/Public/UVEditorLayerEditTool.h
+++ b/Engine/Plugins/Experimental/UVEditor/Source/UVEditorTools/Public/UVEditorLayerEditTool.h
@@ -27,7 +27,7 @@ public:
 	const TArray<TObjectPtr<UUVEditorToolMeshInput>>* Targets = nullptr;
 
 protected:
-	virtual const FToolTargetTypeRequirements& GetTargetRequirements() const override;
+	virtual const FToolTargetTypeRequirements& GetTargetRequirements() const;
 };
 
 UENUM()
diff --git a/Engine/Plugins/Experimental/UVEditor/Source/UVEditorTools/Public/UVEditorLayoutTool.h b/Engine/Plugins/Experimental/UVEditor/Source/UVEditorTools/Public/UVEditorLayoutTool.h
--- a/Engine/Plugins/Experimental/UVEditor/Source/UVEditorTools/Public/UVEditorLayoutTool.h
+++ b/Engine/Plugins/Experimental/UVEditor/Source/UVEditorTools/Public/UVEditorLayoutTool.h
@@ -27,7 +27,7 @@ public:
 	const TArray<TObjectPtr<UUVEditorToolMeshInput>>* Targets = nullptr;
 
 protected:
-	virtual const FToolTargetTypeRequirements& GetTargetRequirements() const override;
+	virtual const FToolTargetTypeRequirements& GetTargetRequirements() const;
 };
 
 /**
diff --git a/Engine/Plugins/Experimental/UVEditor/Source/UVEditorTools/Public/UVEditorRecomputeUVsTool.h b/Engine/Plugins/Experimental/UVEditor/Source/UVEditorTools/Public/UVEditorRecomputeUVsTool.h
--- a/Engine/Plugins/Experimental/UVEditor/Source/UVEditorTools/Public/UVEditorRecomputeUVsTool.h
+++ b/Engine/Plugins/Experimental/UVEditor/Source/UVEditorTools/Public/UVEditorRecomputeUVsTool.h
@@ -38,7 +38,7 @@ public:
 	const TArray<TObjectPtr<UUVEditorToolMeshInput>>* Targets = nullptr;
 
 protected:
-	virtual const FToolTargetTypeRequirements& GetTargetRequirements() const override;
+	virtual const FToolTargetTypeRequirements& GetTargetRequirements() const;
 };
 
 
diff --git a/Engine/Plugins/Experimental/UVEditor/Source/UVEditorTools/Public/UVSelectTool.h b/Engine/Plugins/Experimental/UVEditor/Source/UVEditorTools/Public/UVSelectTool.h
--- a/Engine/Plugins/Experimental/UVEditor/Source/UVEditorTools/Public/UVSelectTool.h
+++ b/Engine/Plugins/Experimental/UVEditor/Source/UVEditorTools/Public/UVSelectTool.h
@@ -47,7 +47,7 @@ public:
 	const TArray<TObjectPtr<UUVEditorToolMeshInput>>* Targets = nullptr;
 
 protected:
-	virtual const FToolTargetTypeRequirements& GetTargetRequirements() const override;
+	virtual const FToolTargetTypeRequirements& GetTargetRequirements() const;
 };
 
 UENUM()
diff --git a/Engine/Plugins/Experimental/UVEditor/Source/UVEditorToolsEditorOnly/Public/UVEditorParameterizeMeshTool.h b/Engine/Plugins/Experimental/UVEditor/Source/UVEditorToolsEditorOnly/Public/UVEditorParameterizeMeshTool.h
--- a/Engine/Plugins/Experimental/UVEditor/Source/UVEditorToolsEditorOnly/Public/UVEditorParameterizeMeshTool.h
+++ b/Engine/Plugins/Experimental/UVEditor/Source/UVEditorToolsEditorOnly/Public/UVEditorParameterizeMeshTool.h
@@ -38,7 +38,7 @@ public:
 	const TArray<TObjectPtr<UUVEditorToolMeshInput>>* Targets = nullptr;
 
 protected:
-	virtual const FToolTargetTypeRequirements& GetTargetRequirements() const override;
+	virtual const FToolTargetTypeRequirements& GetTargetRequirements() const;
 };
 
 
diff --git a/Engine/Plugins/Runtime/MeshModelingToolset/Source/MeshModelingTools/Public/CombineMeshesTool.h b/Engine/Plugins/Runtime/MeshModelingToolset/Source/MeshModelingTools/Public/CombineMeshesTool.h
--- a/Engine/Plugins/Runtime/MeshModelingToolset/Source/MeshModelingTools/Public/CombineMeshesTool.h
+++ b/Engine/Plugins/Runtime/MeshModelingToolset/Source/MeshModelingTools/Public/CombineMeshesTool.h
@@ -30,7 +30,7 @@ public:
 	virtual UInteractiveTool* BuildTool(const FToolBuilderState& SceneState) const override;
 
 protected:
-	virtual const FToolTargetTypeRequirements& GetTargetRequirements() const override;
+	virtual const FToolTargetTypeRequirements& GetTargetRequirements() const;
 };
 
 
diff --git a/Engine/Plugins/Runtime/MeshModelingToolset/Source/ModelingComponents/Public/BaseTools/BaseCreateFromSelectedTool.h b/Engine/Plugins/Runtime/MeshModelingToolset/Source/ModelingComponents/Public/BaseTools/BaseCreateFromSelectedTool.h
--- a/Engine/Plugins/Runtime/MeshModelingToolset/Source/ModelingComponents/Public/BaseTools/BaseCreateFromSelectedTool.h
+++ b/Engine/Plugins/Runtime/MeshModelingToolset/Source/ModelingComponents/Public/BaseTools/BaseCreateFromSelectedTool.h
@@ -33,7 +33,7 @@ public:
 	virtual UBaseCreateFromSelectedTool* MakeNewToolInstance(UObject* Outer) const { check(false); return nullptr; }
 
 protected:
-	virtual const FToolTargetTypeRequirements& GetTargetRequirements() const override;
+	virtual const FToolTargetTypeRequirements& GetTargetRequirements() const;
 };
 
 
diff --git a/Engine/Plugins/Runtime/MeshModelingToolset/Source/ModelingComponents/Public/BaseTools/BaseMeshProcessingTool.h b/Engine/Plugins/Runtime/MeshModelingToolset/Source/ModelingComponents/Public/BaseTools/BaseMeshProcessingTool.h
--- a/Engine/Plugins/Runtime/MeshModelingToolset/Source/ModelingComponents/Public/BaseTools/BaseMeshProcessingTool.h
+++ b/Engine/Plugins/Runtime/MeshModelingToolset/Source/ModelingComponents/Public/BaseTools/BaseMeshProcessingTool.h
@@ -36,7 +36,7 @@ public:
 	virtual UBaseMeshProcessingTool* MakeNewToolInstance(UObject* Outer) const { check(false); return nullptr; }
 
 protected:
-	virtual const FToolTargetTypeRequirements& GetTargetRequirements() const override;
+	virtual const FToolTargetTypeRequirements& GetTargetRequirements() const;
 };
 
 
diff --git a/Engine/Plugins/Runtime/MeshModelingToolset/Source/ModelingComponents/Public/BaseTools/SingleSelectionMeshEditingTool.h b/Engine/Plugins/Runtime/MeshModelingToolset/Source/ModelingComponents/Public/BaseTools/SingleSelectionMeshEditingTool.h
--- a/Engine/Plugins/Runtime/MeshModelingToolset/Source/ModelingComponents/Public/BaseTools/SingleSelectionMeshEditingTool.h
+++ b/Engine/Plugins/Runtime/MeshModelingToolset/Source/ModelingComponents/Public/BaseTools/SingleSelectionMeshEditingTool.h
@@ -36,7 +36,7 @@ public:
 	virtual bool WantsInputSelectionIfAvailable() const { return false; }
 
 protected:
-	virtual const FToolTargetTypeRequirements& GetTargetRequirements() const override;
+	virtual const FToolTargetTypeRequirements& GetTargetRequirements() const;
 };
 
 
diff --git a/Engine/Source/Runtime/InteractiveToolsFramework/Public/BaseTools/MeshSurfacePointTool.h b/Engine/Source/Runtime/InteractiveToolsFramework/Public/BaseTools/MeshSurfacePointTool.h
--- a/Engine/Source/Runtime/InteractiveToolsFramework/Public/BaseTools/MeshSurfacePointTool.h
+++ b/Engine/Source/Runtime/InteractiveToolsFramework/Public/BaseTools/MeshSurfacePointTool.h
@@ -46,7 +46,7 @@ public:
 	virtual void InitializeNewTool(UMeshSurfacePointTool* Tool, const FToolBuilderState& SceneState) const;
 
 protected:
-	virtual const FToolTargetTypeRequirements& GetTargetRequirements() const override;
+	virtual const FToolTargetTypeRequirements& GetTargetRequirements() const;
 };
 
 
diff --git a/Engine/Source/Runtime/InteractiveToolsFramework/Public/InteractiveToolBuilder.h b/Engine/Source/Runtime/InteractiveToolsFramework/Public/InteractiveToolBuilder.h
--- a/Engine/Source/Runtime/InteractiveToolsFramework/Public/InteractiveToolBuilder.h
+++ b/Engine/Source/Runtime/InteractiveToolsFramework/Public/InteractiveToolBuilder.h
@@ -67,15 +67,4 @@ public:
 	{
 		return;
 	}
-
-protected:
-	/**
-	 * Gives the target requirements of the associated tool. Usually, it is the tool builder
-	 * will use this function in CanBuildTool and BuildTool to find and create any necessary targets.
-	 */
-	virtual const FToolTargetTypeRequirements& GetTargetRequirements() const
-	{
-		static FToolTargetTypeRequirements TypeRequirements; // Default initialized to no requirements.
-		return TypeRequirements;
-	}
 };
\ No newline at end of file
