Mesh Attribute Editor: update the visible attribute lists when undoing/redoing adding and deleting attributes

#rnx

diff --git a/Engine/Plugins/Experimental/MeshModelingToolsetExp/Source/MeshModelingToolsEditorOnly/Private/AttributeEditorTool.cpp b/Engine/Plugins/Experimental/MeshModelingToolsetExp/Source/MeshModelingToolsEditorOnly/Private/AttributeEditorTool.cpp
--- a/Engine/Plugins/Experimental/MeshModelingToolsetExp/Source/MeshModelingToolsEditorOnly/Private/AttributeEditorTool.cpp
+++ b/Engine/Plugins/Experimental/MeshModelingToolsetExp/Source/MeshModelingToolsEditorOnly/Private/AttributeEditorTool.cpp
@@ -364,6 +364,11 @@ void UAttributeEditorTool::InitializeAttributeLists()
 		}
 	}
 
+	if (!ModifyAttributeProps->AttributeNamesList.Contains(ModifyAttributeProps->Attribute))
+	{
+		ModifyAttributeProps->Attribute.Reset();
+	}
+
 	bAttributeListsValid = true;
 }
 
@@ -492,6 +497,9 @@ void UAttributeEditorTool::ClearUVs()
 			}
 		});
 	}
+
+	EmitAttributesChange();
+
 	GetToolManager()->EndUndoTransaction();
 
 	// update attrib lists
@@ -536,6 +544,9 @@ void UAttributeEditorTool::DeleteSelectedUVSet()
 			}
 		});
 	}
+
+	EmitAttributesChange();
+
 	GetToolManager()->EndUndoTransaction();
 
 	bAttributeListsValid = false;
@@ -567,6 +578,9 @@ void UAttributeEditorTool::AddUVSet()
 			}
 		});
 	}
+
+	EmitAttributesChange();
+
 	GetToolManager()->EndUndoTransaction();
 
 	bAttributeListsValid = false;
@@ -614,6 +628,9 @@ void UAttributeEditorTool::DuplicateSelectedUVSet()
 			}
 		});
 	}
+
+	EmitAttributesChange();
+
 	GetToolManager()->EndUndoTransaction();
 
 	bAttributeListsValid = false;
@@ -647,6 +664,9 @@ void UAttributeEditorTool::AddNewAttribute(EAttributeEditorElementType ElemType,
 
 	GetToolManager()->BeginUndoTransaction(LOCTEXT("NewAttributeTransactionMessage", "Add Attribute"));
 	UE::ToolTarget::CommitMeshDescriptionUpdate(Targets[0], &NewMesh);
+
+	EmitAttributesChange();
+
 	GetToolManager()->EndUndoTransaction();
 
 	bAttributeListsValid = false;
@@ -731,6 +751,9 @@ void UAttributeEditorTool::DeleteAttribute()
 
 	GetToolManager()->BeginUndoTransaction(LOCTEXT("RemoveAttributeTransactionMessage", "Remove Attribute"));
 	UE::ToolTarget::CommitMeshDescriptionUpdate(Targets[0], &NewMesh);
+
+	EmitAttributesChange();
+
 	GetToolManager()->EndUndoTransaction();
 
 	bAttributeListsValid = false;
@@ -761,6 +784,9 @@ void UAttributeEditorTool::SetLightmapUVsEnabled(bool bEnabled)
 			StaticMesh->PostEditChange();
 		}
 	}
+
+	EmitAttributesChange();
+
 	GetToolManager()->EndUndoTransaction();
 
 	// update attrib lists
@@ -790,6 +816,8 @@ void UAttributeEditorTool::ResetLightmapUVsChannels()
 		}
 	}
 
+	EmitAttributesChange();
+
 	GetToolManager()->EndUndoTransaction();
 
 	// update attrib lists
@@ -811,10 +839,29 @@ void UAttributeEditorTool::UpdateAutoGeneratedLightmapUVChannel(IPrimitiveCompon
 	}
 }
 
+void UAttributeEditorTool::EmitAttributesChange()
+{
+	TUniquePtr<FAttributeEditor_AttributeListsChange> AttributesChange = MakeUnique<FAttributeEditor_AttributeListsChange>();
+	GetToolManager()->EmitObjectChange(this, MoveTemp(AttributesChange), LOCTEXT("AttributesChange", "Attributes Change"));
+}
 
+void FAttributeEditor_AttributeListsChange::Apply(UObject* Object)
+{
+	// We just need the properties to update from the newly-changed static mesh component
+	UAttributeEditorTool* Tool = CastChecked<UAttributeEditorTool>(Object);
+	Tool->bAttributeListsValid = false;
+}
 
+void FAttributeEditor_AttributeListsChange::Revert(UObject* Object)
+{
+	UAttributeEditorTool* Tool = CastChecked<UAttributeEditorTool>(Object);
+	Tool->bAttributeListsValid = false;
+}
 
-
+FString FAttributeEditor_AttributeListsChange::ToString() const
+{
+	return FString(TEXT("AttributeLists Change"));
+}
 
 
 #undef LOCTEXT_NAMESPACE
diff --git a/Engine/Plugins/Experimental/MeshModelingToolsetExp/Source/MeshModelingToolsEditorOnly/Public/AttributeEditorTool.h b/Engine/Plugins/Experimental/MeshModelingToolsetExp/Source/MeshModelingToolsEditorOnly/Public/AttributeEditorTool.h
--- a/Engine/Plugins/Experimental/MeshModelingToolsetExp/Source/MeshModelingToolsEditorOnly/Public/AttributeEditorTool.h
+++ b/Engine/Plugins/Experimental/MeshModelingToolsetExp/Source/MeshModelingToolsEditorOnly/Public/AttributeEditorTool.h
@@ -6,6 +6,7 @@
 #include "UObject/NoExportTypes.h"
 #include "MultiSelectionTool.h"
 #include "InteractiveToolBuilder.h"
+#include "InteractiveToolChange.h"
 #include "AttributeEditorTool.generated.h"
 
 
@@ -403,4 +404,16 @@ protected:
 	void ResetLightmapUVsChannels();
 
 	void AddNewAttribute(EAttributeEditorElementType ElemType, EAttributeEditorAttribType AttribType, FName AttributeName);
+
+	friend class FAttributeEditor_AttributeListsChange;
+	void EmitAttributesChange();
+};
+
+
+class MESHMODELINGTOOLSEDITORONLY_API FAttributeEditor_AttributeListsChange : public FToolCommandChange
+{
+public:
+	virtual void Apply(UObject* Object) override;
+	virtual void Revert(UObject* Object) override;
+	virtual FString ToString() const override;
 };
