fix updated pivot not being initially visible after accepting BakeRS result w/ "Center Pivot" enabled

#rnx

diff --git a/Engine/Plugins/Experimental/MeshModelingToolsetExp/Source/MeshModelingToolsExp/Private/BakeTransformTool.cpp b/Engine/Plugins/Experimental/MeshModelingToolsetExp/Source/MeshModelingToolsExp/Private/BakeTransformTool.cpp
--- a/Engine/Plugins/Experimental/MeshModelingToolsetExp/Source/MeshModelingToolsExp/Private/BakeTransformTool.cpp
+++ b/Engine/Plugins/Experimental/MeshModelingToolsetExp/Source/MeshModelingToolsExp/Private/BakeTransformTool.cpp
@@ -262,6 +262,22 @@ void UBakeTransformTool::UpdateAssets()
 		}
 	}
 
+	if (BasicProperties->bRecenterPivot)
+	{
+		// hack to ensure user sees the updated pivot immediately: request re-select of the original selection
+		FSelectedOjectsChangeList NewSelection;
+		NewSelection.ModificationType = ESelectedObjectsModificationType::Replace;
+		for (int OrigMeshIdx = 0; OrigMeshIdx < Targets.Num(); OrigMeshIdx++)
+		{
+			AActor* OwnerActor = UE::ToolTarget::GetTargetActor(Targets[OrigMeshIdx]);
+			if (OwnerActor)
+			{
+				NewSelection.Actors.Add(OwnerActor);
+			}
+		}
+		GetToolManager()->RequestSelectionChange(NewSelection);
+	}
+
 	GetToolManager()->EndUndoTransaction();
 }
 
