InteractiveToolsFramework: Remove unnecessary target interface functions from MultiSelectionTool.

ModelingTools: Add UMultiSelectionMeshEditingTool to hold common interface requirements and functions that pertain to modeling mode MultiSelectionTools.

#rnx

diff --git a/Engine/Plugins/Experimental/MeshLODToolset/Source/MeshLODToolset/Private/Tools/GenerateStaticMeshLODAssetTool.cpp b/Engine/Plugins/Experimental/MeshLODToolset/Source/MeshLODToolset/Private/Tools/GenerateStaticMeshLODAssetTool.cpp
--- a/Engine/Plugins/Experimental/MeshLODToolset/Source/MeshLODToolset/Private/Tools/GenerateStaticMeshLODAssetTool.cpp
+++ b/Engine/Plugins/Experimental/MeshLODToolset/Source/MeshLODToolset/Private/Tools/GenerateStaticMeshLODAssetTool.cpp
@@ -220,15 +220,10 @@ bool UGenerateStaticMeshLODAssetToolBuilder::CanBuildTool(const FToolBuilderStat
 	return (SceneState.TargetManager->CountSelectedAndTargetable(SceneState, GetTargetRequirements()) == 1);
 }
 
-UInteractiveTool* UGenerateStaticMeshLODAssetToolBuilder::BuildTool(const FToolBuilderState& SceneState) const
+UMultiSelectionMeshEditingTool* UGenerateStaticMeshLODAssetToolBuilder::CreateNewTool(const FToolBuilderState& SceneState) const
 {
 	UGenerateStaticMeshLODAssetTool* NewTool = NewObject<UGenerateStaticMeshLODAssetTool>(SceneState.ToolManager);
 	NewTool->SetUseAssetEditorMode(bUseAssetEditorMode);
-
-	TArray<TObjectPtr<UToolTarget>> Targets = SceneState.TargetManager->BuildAllSelectedTargetable(SceneState, GetTargetRequirements());
-	NewTool->SetTargets(MoveTemp(Targets));
-	NewTool->SetWorld(SceneState.World);
-
 	return NewTool;
 }
 
@@ -249,12 +244,6 @@ void UGenerateStaticMeshLODAssetToolPresetProperties::PostAction(EGenerateLODAss
  */
 
 
-void UGenerateStaticMeshLODAssetTool::SetWorld(UWorld* World)
-{
-	this->TargetWorld = World;
-}
-
-
 void UGenerateStaticMeshLODAssetTool::SetUseAssetEditorMode(bool bEnable)
 {
 	bIsInAssetEditorMode = bEnable;
diff --git a/Engine/Plugins/Experimental/MeshLODToolset/Source/MeshLODToolset/Private/Tools/LODManagerTool.cpp b/Engine/Plugins/Experimental/MeshLODToolset/Source/MeshLODToolset/Private/Tools/LODManagerTool.cpp
--- a/Engine/Plugins/Experimental/MeshLODToolset/Source/MeshLODToolset/Private/Tools/LODManagerTool.cpp
+++ b/Engine/Plugins/Experimental/MeshLODToolset/Source/MeshLODToolset/Private/Tools/LODManagerTool.cpp
@@ -45,20 +45,9 @@ const FToolTargetTypeRequirements& ULODManagerToolBuilder::GetTargetRequirements
 	return TypeRequirements;
 }
 
-bool ULODManagerToolBuilder::CanBuildTool(const FToolBuilderState& SceneState) const
+UMultiSelectionMeshEditingTool* ULODManagerToolBuilder::CreateNewTool(const FToolBuilderState& SceneState) const
 {
-	return SceneState.TargetManager->CountSelectedAndTargetable(SceneState, GetTargetRequirements()) > 0;
-}
-
-UInteractiveTool* ULODManagerToolBuilder::BuildTool(const FToolBuilderState& SceneState) const
-{
-	ULODManagerTool* NewTool = NewObject<ULODManagerTool>(SceneState.ToolManager);
-
-	TArray<TObjectPtr<UToolTarget>> Targets = SceneState.TargetManager->BuildAllSelectedTargetable(SceneState, GetTargetRequirements());
-	NewTool->SetTargets(MoveTemp(Targets));
-	NewTool->SetWorld(SceneState.World);
-
-	return NewTool;
+	return NewObject<ULODManagerTool>(SceneState.ToolManager);
 }
 
 
@@ -79,10 +68,6 @@ ULODManagerTool::ULODManagerTool()
 {
 }
 
-void ULODManagerTool::SetWorld(UWorld* World)
-{
-	this->TargetWorld = World;
-}
 
 void ULODManagerTool::Setup()
 {
@@ -146,7 +131,7 @@ void ULODManagerTool::Shutdown(EToolShutdownType ShutdownType)
 	}
 	for (int32 ComponentIdx = 0; ComponentIdx < Targets.Num(); ComponentIdx++)
 	{
-		TargetComponentInterface(ComponentIdx)->SetOwnerVisibility(true);
+		UE::ToolTarget::ShowSourceObject(Targets[ComponentIdx]);
 	}
 }
 
@@ -494,8 +479,7 @@ void ULODManagerTool::DeleteHiResSourceModel()
 	GetToolManager()->BeginUndoTransaction(LOCTEXT("DeleteHiResSourceModel", "Delete HiRes Source"));
 	for (int32 ComponentIdx = 0; ComponentIdx < Targets.Num(); ComponentIdx++)
 	{
-		IPrimitiveComponentBackedTarget* TargetComponent = TargetComponentInterface(ComponentIdx);
-		UStaticMeshComponent* StaticMeshComponent = Cast<UStaticMeshComponent>(TargetComponent->GetOwnerComponent());
+		UStaticMeshComponent* StaticMeshComponent = Cast<UStaticMeshComponent>(UE::ToolTarget::GetTargetComponent(Targets[ComponentIdx]));
 		if (StaticMeshComponent == nullptr || StaticMeshComponent->GetStaticMesh() == nullptr)
 		{
 			continue;
@@ -527,8 +511,7 @@ void ULODManagerTool::MoveHiResToLOD0()
 	GetToolManager()->BeginUndoTransaction(LOCTEXT("RestoreLOD0", "Move HiRes to LOD0"));
 	for (int32 ComponentIdx = 0; ComponentIdx < Targets.Num(); ComponentIdx++)
 	{
-		IPrimitiveComponentBackedTarget* TargetComponent = TargetComponentInterface(ComponentIdx);
-		UStaticMeshComponent* StaticMeshComponent = Cast<UStaticMeshComponent>(TargetComponent->GetOwnerComponent());
+		UStaticMeshComponent* StaticMeshComponent = Cast<UStaticMeshComponent>(UE::ToolTarget::GetTargetComponent(Targets[ComponentIdx]));
 		if (StaticMeshComponent == nullptr || StaticMeshComponent->GetStaticMesh() == nullptr)
 		{
 			continue;
@@ -569,8 +552,7 @@ void ULODManagerTool::RemoveUnreferencedMaterials()
 	GetToolManager()->BeginUndoTransaction(LOCTEXT("RemoveUnreferencedMaterials", "Remove Unreferenced Materials"));
 	for (int32 ComponentIdx = 0; ComponentIdx < Targets.Num(); ComponentIdx++)
 	{
-		IPrimitiveComponentBackedTarget* TargetComponent = TargetComponentInterface(ComponentIdx);
-		UStaticMeshComponent* StaticMeshComponent = Cast<UStaticMeshComponent>(TargetComponent->GetOwnerComponent());
+		UStaticMeshComponent* StaticMeshComponent = Cast<UStaticMeshComponent>(UE::ToolTarget::GetTargetComponent(Targets[ComponentIdx]));
 		if (StaticMeshComponent == nullptr || StaticMeshComponent->GetStaticMesh() == nullptr)
 		{
 			continue;
diff --git a/Engine/Plugins/Experimental/MeshLODToolset/Source/MeshLODToolset/Public/Tools/GenerateStaticMeshLODAssetTool.h b/Engine/Plugins/Experimental/MeshLODToolset/Source/MeshLODToolset/Public/Tools/GenerateStaticMeshLODAssetTool.h
--- a/Engine/Plugins/Experimental/MeshLODToolset/Source/MeshLODToolset/Public/Tools/GenerateStaticMeshLODAssetTool.h
+++ b/Engine/Plugins/Experimental/MeshLODToolset/Source/MeshLODToolset/Public/Tools/GenerateStaticMeshLODAssetTool.h
@@ -4,7 +4,7 @@
 
 #include "CoreMinimal.h"
 #include "UObject/NoExportTypes.h"
-#include "MultiSelectionTool.h"
+#include "BaseTools/MultiSelectionMeshEditingTool.h"
 #include "InteractiveToolBuilder.h"
 #include "ModelingOperators.h" //IDynamicMeshOperatorFactory
 #include "DynamicMesh/DynamicMesh3.h"
@@ -38,7 +38,7 @@ enum class EGenerateLODAssetOutputMode : uint8
  * Tool builder
  */
 UCLASS()
-class MESHLODTOOLSET_API UGenerateStaticMeshLODAssetToolBuilder : public UInteractiveToolBuilder
+class MESHLODTOOLSET_API UGenerateStaticMeshLODAssetToolBuilder : public UMultiSelectionMeshEditingToolBuilder
 {
 	GENERATED_BODY()
 
@@ -46,10 +46,10 @@ public:
 	bool bUseAssetEditorMode = false;
 
 	virtual bool CanBuildTool(const FToolBuilderState& SceneState) const override;
-	virtual UInteractiveTool* BuildTool(const FToolBuilderState& SceneState) const override;
+	virtual UMultiSelectionMeshEditingTool* CreateNewTool(const FToolBuilderState& SceneState) const override;
 
 protected:
-	virtual const FToolTargetTypeRequirements& GetTargetRequirements() const;
+	virtual const FToolTargetTypeRequirements& GetTargetRequirements() const override;
 };
 
 
@@ -260,7 +260,7 @@ public:
  * Simple tool to combine multiple meshes into a single mesh asset
  */
 UCLASS()
-class MESHLODTOOLSET_API UGenerateStaticMeshLODAssetTool : public UMultiSelectionTool
+class MESHLODTOOLSET_API UGenerateStaticMeshLODAssetTool : public UMultiSelectionMeshEditingTool
 {
 	GENERATED_BODY()
 
@@ -272,8 +272,6 @@ public:
 	virtual void Setup() override;
 	virtual void Shutdown(EToolShutdownType ShutdownType) override;
 
-	virtual void SetWorld(UWorld* World);
-
 	virtual void OnTick(float DeltaTime);
 
 	virtual bool HasCancel() const override { return true; }
@@ -328,7 +326,6 @@ protected:
 
 protected:
 	bool bIsInAssetEditorMode = false;
-	UWorld* TargetWorld;
 
 	UPROPERTY()
 	TObjectPtr<UGenerateStaticMeshLODProcess> GenerateProcess;
diff --git a/Engine/Plugins/Experimental/MeshLODToolset/Source/MeshLODToolset/Public/Tools/LODManagerTool.h b/Engine/Plugins/Experimental/MeshLODToolset/Source/MeshLODToolset/Public/Tools/LODManagerTool.h
--- a/Engine/Plugins/Experimental/MeshLODToolset/Source/MeshLODToolset/Public/Tools/LODManagerTool.h
+++ b/Engine/Plugins/Experimental/MeshLODToolset/Source/MeshLODToolset/Public/Tools/LODManagerTool.h
@@ -3,7 +3,7 @@
 #pragma once
 
 #include "CoreMinimal.h"
-#include "MultiSelectionTool.h"
+#include "BaseTools/MultiSelectionMeshEditingTool.h"
 #include "InteractiveToolBuilder.h"
 #include "Engine/StaticMesh.h"
 #include "PreviewMesh.h"
@@ -13,21 +13,21 @@
 
 // predeclarations
 class UMaterialInterface;
+class ULODManagerTool;
 
 /**
  *
  */
 UCLASS()
-class MESHLODTOOLSET_API ULODManagerToolBuilder : public UInteractiveToolBuilder
+class MESHLODTOOLSET_API ULODManagerToolBuilder : public UMultiSelectionMeshEditingToolBuilder
 {
 	GENERATED_BODY()
 
 public:
-	virtual bool CanBuildTool(const FToolBuilderState& SceneState) const override;
-	virtual UInteractiveTool* BuildTool(const FToolBuilderState& SceneState) const override;
+	virtual UMultiSelectionMeshEditingTool* CreateNewTool(const FToolBuilderState& SceneState) const override;
 
 protected:
-	virtual const FToolTargetTypeRequirements& GetTargetRequirements() const;
+	virtual const FToolTargetTypeRequirements& GetTargetRequirements() const override;
 };
 
 
@@ -167,7 +167,7 @@ public:
  * Mesh Attribute Editor Tool
  */
 UCLASS()
-class MESHLODTOOLSET_API ULODManagerTool : public UMultiSelectionTool
+class MESHLODTOOLSET_API ULODManagerTool : public UMultiSelectionMeshEditingTool
 {
 	GENERATED_BODY()
 
@@ -178,8 +178,6 @@ public:
 	virtual void Shutdown(EToolShutdownType ShutdownType) override;
 	virtual void OnTick(float DeltaTime) override;
 
-	virtual void SetWorld(UWorld* World);
-
 	virtual bool HasCancel() const override { return false; }
 	virtual bool HasAccept() const override { return false; }
 	virtual bool CanAccept() const override { return false; }
@@ -213,8 +211,6 @@ public:
 	void RemoveUnreferencedMaterials();
 
 protected:
-	UWorld* TargetWorld;
-
 	UStaticMesh* GetSingleStaticMesh();
 
 	ELODManagerToolActions PendingAction = ELODManagerToolActions::NoAction;
diff --git a/Engine/Plugins/Experimental/MeshModelingToolsetExp/Source/MeshModelingToolsEditorOnly/Private/AddPivotActorTool.cpp b/Engine/Plugins/Experimental/MeshModelingToolsetExp/Source/MeshModelingToolsEditorOnly/Private/AddPivotActorTool.cpp
--- a/Engine/Plugins/Experimental/MeshModelingToolsetExp/Source/MeshModelingToolsEditorOnly/Private/AddPivotActorTool.cpp
+++ b/Engine/Plugins/Experimental/MeshModelingToolsetExp/Source/MeshModelingToolsEditorOnly/Private/AddPivotActorTool.cpp
@@ -77,21 +77,13 @@ bool UAddPivotActorToolBuilder::CanBuildTool(const FToolBuilderState& SceneState
 	return bAllActorsMovable && ParentActors.Num() == 1;
 }
 
-UInteractiveTool* UAddPivotActorToolBuilder::BuildTool(const FToolBuilderState& SceneState) const
+UMultiSelectionMeshEditingTool* UAddPivotActorToolBuilder::CreateNewTool(const FToolBuilderState& SceneState) const
 {
 	UAddPivotActorTool* NewTool = NewObject<UAddPivotActorTool>(SceneState.ToolManager);
-
 	if (SceneState.SelectedActors.Num() == 1 && ExactCast<AActor>(SceneState.SelectedActors[0]))
 	{
 		NewTool->SetPivotRepositionMode(SceneState.SelectedActors[0]);
 	}
-	else
-	{
-		NewTool->SetTargets(SceneState.TargetManager->BuildAllSelectedTargetable(
-			SceneState, GetTargetRequirements()));
-		NewTool->SetWorld(SceneState.World);
-	}
-
 	return NewTool;
 }
 
diff --git a/Engine/Plugins/Experimental/MeshModelingToolsetExp/Source/MeshModelingToolsEditorOnly/Private/AttributeEditorTool.cpp b/Engine/Plugins/Experimental/MeshModelingToolsetExp/Source/MeshModelingToolsEditorOnly/Private/AttributeEditorTool.cpp
--- a/Engine/Plugins/Experimental/MeshModelingToolsetExp/Source/MeshModelingToolsEditorOnly/Private/AttributeEditorTool.cpp
+++ b/Engine/Plugins/Experimental/MeshModelingToolsetExp/Source/MeshModelingToolsEditorOnly/Private/AttributeEditorTool.cpp
@@ -33,31 +33,9 @@ using namespace UE::Geometry;
  * ToolBuilder
  */
 
-
-const FToolTargetTypeRequirements& UAttributeEditorToolBuilder::GetTargetRequirements() const
-{
-	static FToolTargetTypeRequirements TypeRequirements({
-		UMeshDescriptionCommitter::StaticClass(),
-		UMeshDescriptionProvider::StaticClass(),
-		UPrimitiveComponentBackedTarget::StaticClass()
-		});
-	return TypeRequirements;
-}
-
-bool UAttributeEditorToolBuilder::CanBuildTool(const FToolBuilderState& SceneState) const
-{
-	return SceneState.TargetManager->CountSelectedAndTargetable(SceneState, GetTargetRequirements()) > 0;
-}
-
-UInteractiveTool* UAttributeEditorToolBuilder::BuildTool(const FToolBuilderState& SceneState) const
+UMultiSelectionMeshEditingTool* UAttributeEditorToolBuilder::CreateNewTool(const FToolBuilderState& SceneState) const
 {
-	UAttributeEditorTool* NewTool = NewObject<UAttributeEditorTool>(SceneState.ToolManager);
-
-	TArray<TObjectPtr<UToolTarget>> Targets = SceneState.TargetManager->BuildAllSelectedTargetable(SceneState, GetTargetRequirements());
-	NewTool->SetTargets(MoveTemp(Targets));
-	NewTool->SetWorld(SceneState.World);
-
-	return NewTool;
+	return NewObject<UAttributeEditorTool>(SceneState.ToolManager);
 }
 
 
@@ -98,10 +76,6 @@ UAttributeEditorTool::UAttributeEditorTool()
 {
 }
 
-void UAttributeEditorTool::SetWorld(UWorld* World)
-{
-	this->TargetWorld = World;
-}
 
 void UAttributeEditorTool::Setup()
 {
@@ -449,7 +423,7 @@ void UAttributeEditorTool::ClearNormals()
 
 	for (int32 ComponentIdx = 0; ComponentIdx < Targets.Num(); ComponentIdx++)
 	{
-		TargetMeshCommitterInterface(ComponentIdx)->CommitMeshDescription([&](const IMeshDescriptionCommitter::FCommitterParams& CommitParams)
+		Cast<IMeshDescriptionCommitter>(Targets[ComponentIdx])->CommitMeshDescription([&](const IMeshDescriptionCommitter::FCommitterParams& CommitParams)
 		{
 			TEdgeAttributesRef<bool> EdgeHardnesses = CommitParams.MeshDescriptionOut->EdgeAttributes().GetAttributesRef<bool>(MeshAttribute::Edge::IsHard);
 			if (EdgeHardnesses.IsValid())
@@ -474,7 +448,7 @@ void UAttributeEditorTool::ClearUVs()
 	GetToolManager()->BeginUndoTransaction(LOCTEXT("ClearUVsTransactionMessage", "Clear Selected UVs"));
 	for (int32 ComponentIdx = 0; ComponentIdx < Targets.Num(); ComponentIdx++)
 	{
-		TargetMeshCommitterInterface(ComponentIdx)->CommitMeshDescription([&](const IMeshDescriptionCommitter::FCommitterParams& CommitParams)
+		Cast<IMeshDescriptionCommitter>(Targets[ComponentIdx])->CommitMeshDescription([&](const IMeshDescriptionCommitter::FCommitterParams& CommitParams)
 		{
 			TVertexInstanceAttributesRef<FVector2f> InstanceUVs =
 				CommitParams.MeshDescriptionOut->VertexInstanceAttributes().GetAttributesRef<FVector2f>(MeshAttribute::VertexInstance::TextureCoordinate);
@@ -493,7 +467,7 @@ void UAttributeEditorTool::ClearUVs()
 
 			if (bHaveAutoGeneratedLightmapUVSet)
 			{
-				UpdateAutoGeneratedLightmapUVChannel(TargetComponentInterface(ComponentIdx), InstanceUVs.GetNumChannels());
+				UpdateAutoGeneratedLightmapUVChannel(UE::ToolTarget::GetTargetComponent(Targets[ComponentIdx]), InstanceUVs.GetNumChannels());
 			}
 		});
 	}
@@ -525,7 +499,7 @@ void UAttributeEditorTool::DeleteSelectedUVSet()
 	GetToolManager()->BeginUndoTransaction(LOCTEXT("ClearUVsTransactionMessage", "Clear Selected UVs"));
 	for (int32 ComponentIdx = 0; ComponentIdx < Targets.Num(); ComponentIdx++)
 	{
-		TargetMeshCommitterInterface(ComponentIdx)->CommitMeshDescription([&](const IMeshDescriptionCommitter::FCommitterParams& CommitParams)
+		Cast<IMeshDescriptionCommitter>(Targets[ComponentIdx])->CommitMeshDescription([&](const IMeshDescriptionCommitter::FCommitterParams& CommitParams)
 		{
 			TVertexInstanceAttributesRef<FVector2f> InstanceUVs =
 				CommitParams.MeshDescriptionOut->VertexInstanceAttributes().GetAttributesRef<FVector2f>(MeshAttribute::VertexInstance::TextureCoordinate);
@@ -540,7 +514,7 @@ void UAttributeEditorTool::DeleteSelectedUVSet()
 
 			if (bHaveAutoGeneratedLightmapUVSet)
 			{
-				UpdateAutoGeneratedLightmapUVChannel(TargetComponentInterface(ComponentIdx), InstanceUVs.GetNumChannels());
+				UpdateAutoGeneratedLightmapUVChannel(UE::ToolTarget::GetTargetComponent(Targets[ComponentIdx]), InstanceUVs.GetNumChannels());
 			}
 		});
 	}
@@ -558,7 +532,7 @@ void UAttributeEditorTool::AddUVSet()
 	GetToolManager()->BeginUndoTransaction(LOCTEXT("AddUVSetMessage", "Add UV Set"));
 	for (int32 ComponentIdx = 0; ComponentIdx < Targets.Num(); ComponentIdx++)
 	{
-		TargetMeshCommitterInterface(ComponentIdx)->CommitMeshDescription([&](const IMeshDescriptionCommitter::FCommitterParams& CommitParams)
+		Cast<IMeshDescriptionCommitter>(Targets[ComponentIdx])->CommitMeshDescription([&](const IMeshDescriptionCommitter::FCommitterParams& CommitParams)
 		{
 			TVertexInstanceAttributesRef<FVector2f> InstanceUVs =
 				CommitParams.MeshDescriptionOut->VertexInstanceAttributes().GetAttributesRef<FVector2f>(MeshAttribute::VertexInstance::TextureCoordinate);
@@ -573,7 +547,7 @@ void UAttributeEditorTool::AddUVSet()
 
 				if (bHaveAutoGeneratedLightmapUVSet)
 				{
-					UpdateAutoGeneratedLightmapUVChannel(TargetComponentInterface(ComponentIdx), InstanceUVs.GetNumChannels());
+					UpdateAutoGeneratedLightmapUVChannel(UE::ToolTarget::GetTargetComponent(Targets[ComponentIdx]), InstanceUVs.GetNumChannels());
 				}
 			}
 		});
@@ -600,7 +574,7 @@ void UAttributeEditorTool::DuplicateSelectedUVSet()
 	GetToolManager()->BeginUndoTransaction(LOCTEXT("DuplicateUVSetMessage", "Duplicate UV Set"));
 	for (int32 ComponentIdx = 0; ComponentIdx < Targets.Num(); ComponentIdx++)
 	{
-		TargetMeshCommitterInterface(ComponentIdx)->CommitMeshDescription([&](const IMeshDescriptionCommitter::FCommitterParams& CommitParams)
+		Cast<IMeshDescriptionCommitter>(Targets[ComponentIdx])->CommitMeshDescription([&](const IMeshDescriptionCommitter::FCommitterParams& CommitParams)
 		{
 			TVertexInstanceAttributesRef<FVector2f> InstanceUVs =
 				CommitParams.MeshDescriptionOut->VertexInstanceAttributes().GetAttributesRef<FVector2f>(MeshAttribute::VertexInstance::TextureCoordinate);
@@ -620,7 +594,7 @@ void UAttributeEditorTool::DuplicateSelectedUVSet()
 
 				if (bHaveAutoGeneratedLightmapUVSet)
 				{
-					UpdateAutoGeneratedLightmapUVChannel(TargetComponentInterface(ComponentIdx), InstanceUVs.GetNumChannels());
+					UpdateAutoGeneratedLightmapUVChannel(UE::ToolTarget::GetTargetComponent(Targets[ComponentIdx]), InstanceUVs.GetNumChannels());
 				}
 
 				GetToolManager()->DisplayMessage(FText::Format(LOCTEXT("Copied UV Set", "Copied UV{0} to UV{1}"), 
@@ -773,7 +747,7 @@ void UAttributeEditorTool::SetLightmapUVsEnabled(bool bEnabled)
 	}
 	for (int32 ComponentIdx = 0; ComponentIdx < Targets.Num(); ComponentIdx++)
 	{
-		UStaticMeshComponent* StaticMeshComponent = Cast<UStaticMeshComponent>(TargetComponentInterface(ComponentIdx)->GetOwnerComponent());
+		UStaticMeshComponent* StaticMeshComponent = Cast<UStaticMeshComponent>(UE::ToolTarget::GetTargetComponent(Targets[ComponentIdx]));
 		if (StaticMeshComponent && StaticMeshComponent->GetStaticMesh() != nullptr)
 		{
 			UStaticMesh* StaticMesh = StaticMeshComponent->GetStaticMesh();
@@ -801,10 +775,10 @@ void UAttributeEditorTool::ResetLightmapUVsChannels()
 	for (int32 ComponentIdx = 0; ComponentIdx < Targets.Num(); ComponentIdx++)
 	{
 		TVertexInstanceAttributesConstRef<FVector2f> InstanceUVs = 
-			TargetMeshProviderInterface(ComponentIdx)->GetMeshDescription()->VertexInstanceAttributes().GetAttributesRef<FVector2f>(MeshAttribute::VertexInstance::TextureCoordinate);
+			Cast<IMeshDescriptionProvider>(Targets[ComponentIdx])->GetMeshDescription()->VertexInstanceAttributes().GetAttributesRef<FVector2f>(MeshAttribute::VertexInstance::TextureCoordinate);
 		int32 SetChannel = FMath::Max(InstanceUVs.GetNumChannels(), 1);
 
-		UStaticMeshComponent* StaticMeshComponent = Cast<UStaticMeshComponent>(TargetComponentInterface(ComponentIdx)->GetOwnerComponent());
+		UStaticMeshComponent* StaticMeshComponent = Cast<UStaticMeshComponent>(UE::ToolTarget::GetTargetComponent(Targets[ComponentIdx]));
 		if (StaticMeshComponent && StaticMeshComponent->GetStaticMesh() != nullptr)
 		{
 			UStaticMesh* StaticMesh = StaticMeshComponent->GetStaticMesh();
@@ -826,9 +800,9 @@ void UAttributeEditorTool::ResetLightmapUVsChannels()
 
 
 
-void UAttributeEditorTool::UpdateAutoGeneratedLightmapUVChannel(IPrimitiveComponentBackedTarget* Target, int32 NewMaxUVChannels)
+void UAttributeEditorTool::UpdateAutoGeneratedLightmapUVChannel(UPrimitiveComponent* TargetComponent, int32 NewMaxUVChannels)
 {
-	UStaticMeshComponent* StaticMeshComponent = Cast<UStaticMeshComponent>(Target->GetOwnerComponent());
+	UStaticMeshComponent* StaticMeshComponent = Cast<UStaticMeshComponent>(TargetComponent);
 	if (StaticMeshComponent && StaticMeshComponent->GetStaticMesh() != nullptr)
 	{
 		UStaticMesh* StaticMesh = StaticMeshComponent->GetStaticMesh();
diff --git a/Engine/Plugins/Experimental/MeshModelingToolsetExp/Source/MeshModelingToolsEditorOnly/Private/MergeMeshesTool.cpp b/Engine/Plugins/Experimental/MeshModelingToolsetExp/Source/MeshModelingToolsEditorOnly/Private/MergeMeshesTool.cpp
--- a/Engine/Plugins/Experimental/MeshModelingToolsetExp/Source/MeshModelingToolsEditorOnly/Private/MergeMeshesTool.cpp
+++ b/Engine/Plugins/Experimental/MeshModelingToolsetExp/Source/MeshModelingToolsEditorOnly/Private/MergeMeshesTool.cpp
@@ -5,6 +5,7 @@
 #include "ToolBuilderUtil.h"
 #include "ToolSetupUtil.h"
 #include "ModelingObjectsCreationAPI.h"
+#include "ModelingToolTargetUtil.h"
 #include "Selection/ToolSelectionUtil.h"
 
 #include "DynamicMesh/DynamicMesh3.h"
@@ -36,20 +37,14 @@ const FToolTargetTypeRequirements& UMergeMeshesToolBuilder::GetTargetRequirement
 
 bool UMergeMeshesToolBuilder::CanBuildTool(const FToolBuilderState& SceneState) const
 {
-	const int32 MinRequiredComponents = 1;
+	constexpr int32 MinRequiredComponents = 1;
 	const bool bHasComponents = SceneState.TargetManager->CountSelectedAndTargetable(SceneState, GetTargetRequirements()) >= MinRequiredComponents;
 	return ( bHasComponents );
 }
 
-UInteractiveTool* UMergeMeshesToolBuilder::BuildTool(const FToolBuilderState& SceneState) const
+UMultiSelectionMeshEditingTool* UMergeMeshesToolBuilder::CreateNewTool(const FToolBuilderState& SceneState) const
 {
-	UMergeMeshesTool* NewTool = NewObject<UMergeMeshesTool>(SceneState.ToolManager);
-
-	TArray<TObjectPtr<UToolTarget>> Targets = SceneState.TargetManager->BuildAllSelectedTargetable(SceneState, GetTargetRequirements());
-	NewTool->SetTargets(MoveTemp(Targets));
-	NewTool->SetWorld(SceneState.World);
-
-	return NewTool;
+	return NewObject<UMergeMeshesTool>(SceneState.ToolManager);
 }
 
 
@@ -63,11 +58,6 @@ UMergeMeshesTool::UMergeMeshesTool()
 {
 }
 
-void UMergeMeshesTool::SetWorld(UWorld* World)
-{
-	this->TargetWorld = World;
-}
-
 
 void UMergeMeshesTool::Setup()
 {
@@ -87,7 +77,7 @@ void UMergeMeshesTool::Setup()
 	// Hide the source meshes
 	for (int32 ComponentIdx = 0; ComponentIdx < Targets.Num(); ComponentIdx++)
 	{
-		TargetComponentInterface(ComponentIdx)->SetOwnerVisibility(false);
+		UE::ToolTarget::HideSourceObject(Targets[ComponentIdx]);
 	}
 
 	// save transformed version of input meshes (maybe this could happen in the Operator?)
@@ -128,7 +118,7 @@ void UMergeMeshesTool::Shutdown(EToolShutdownType ShutdownType)
 	// Restore (unhide) the source meshes
 	for (int32 ComponentIdx = 0; ComponentIdx < Targets.Num(); ComponentIdx++)
 	{
-		TargetComponentInterface(ComponentIdx)->SetOwnerVisibility(true);
+		UE::ToolTarget::ShowSourceObject(Targets[ComponentIdx]);
 	}
 	if (ShutdownType == EToolShutdownType::Accept)
 	{
@@ -140,7 +130,7 @@ void UMergeMeshesTool::Shutdown(EToolShutdownType ShutdownType)
 		TArray<AActor*> Actors;
 		for (int32 ComponentIdx = 0; ComponentIdx < Targets.Num(); ComponentIdx++)
 		{
-			Actors.Add(TargetComponentInterface(ComponentIdx)->GetOwnerActor());
+			Actors.Add(UE::ToolTarget::GetTargetActor(Targets[ComponentIdx]));
 		}
 		HandleSourcesProperties->ApplyMethod(Actors, GetToolManager());
 
@@ -189,8 +179,8 @@ void UMergeMeshesTool::CacheInputMeshes()
 	for (int32 ComponentIdx = 0; ComponentIdx < Targets.Num(); ComponentIdx++)
 	{
 		UE::Geometry::FVoxelMergeMeshesOp::FInputMesh InputMesh;
-		InputMesh.Mesh = TargetMeshProviderInterface(ComponentIdx)->GetMeshDescription();
-		InputMesh.Transform = TargetComponentInterface(ComponentIdx)->GetWorldTransform();
+		InputMesh.Mesh = UE::ToolTarget::GetMeshDescription(Targets[ComponentIdx]);
+		InputMesh.Transform = (FTransform) UE::ToolTarget::GetLocalToWorldTransform(Targets[ComponentIdx]);
 		InputMeshes.Add(InputMesh);
 	}
 }
diff --git a/Engine/Plugins/Experimental/MeshModelingToolsetExp/Source/MeshModelingToolsEditorOnly/Private/VoxelCSGMeshesTool.cpp b/Engine/Plugins/Experimental/MeshModelingToolsetExp/Source/MeshModelingToolsEditorOnly/Private/VoxelCSGMeshesTool.cpp
--- a/Engine/Plugins/Experimental/MeshModelingToolsetExp/Source/MeshModelingToolsEditorOnly/Private/VoxelCSGMeshesTool.cpp
+++ b/Engine/Plugins/Experimental/MeshModelingToolsetExp/Source/MeshModelingToolsEditorOnly/Private/VoxelCSGMeshesTool.cpp
@@ -6,6 +6,7 @@
 #include "ToolSetupUtil.h"
 #include "ModelingObjectsCreationAPI.h"
 #include "Selection/ToolSelectionUtil.h"
+#include "ModelingToolTargetUtil.h"
 
 #include "DynamicMesh/DynamicMesh3.h"
 #include "MeshDescriptionToDynamicMesh.h"
@@ -25,30 +26,14 @@ using namespace UE::Geometry;
 /*
  * ToolBuilder
  */
-const FToolTargetTypeRequirements& UVoxelCSGMeshesToolBuilder::GetTargetRequirements() const
-{
-	static FToolTargetTypeRequirements TypeRequirements({
-		UMeshDescriptionCommitter::StaticClass(),
-		UMeshDescriptionProvider::StaticClass(),
-		UPrimitiveComponentBackedTarget::StaticClass()
-		});
-	return TypeRequirements;
-}
-
 bool UVoxelCSGMeshesToolBuilder::CanBuildTool(const FToolBuilderState& SceneState) const
 {
 	return SceneState.TargetManager->CountSelectedAndTargetable(SceneState, GetTargetRequirements()) == 2;
 }
 
-UInteractiveTool* UVoxelCSGMeshesToolBuilder::BuildTool(const FToolBuilderState& SceneState) const
+UMultiSelectionMeshEditingTool* UVoxelCSGMeshesToolBuilder::CreateNewTool(const FToolBuilderState& SceneState) const
 {
-	UVoxelCSGMeshesTool* NewTool = NewObject<UVoxelCSGMeshesTool>(SceneState.ToolManager);
-
-	TArray<TObjectPtr<UToolTarget>> Targets = SceneState.TargetManager->BuildAllSelectedTargetable(SceneState, GetTargetRequirements());
-	NewTool->SetTargets(MoveTemp(Targets));
-	NewTool->SetWorld(SceneState.World);
-
-	return NewTool;
+	return NewObject<UVoxelCSGMeshesTool>(SceneState.ToolManager);
 }
 
 /*
@@ -59,11 +44,6 @@ UVoxelCSGMeshesTool::UVoxelCSGMeshesTool()
 {
 }
 
-void UVoxelCSGMeshesTool::SetWorld(UWorld* World)
-{
-	this->TargetWorld = World;
-}
-
 void UVoxelCSGMeshesTool::Setup()
 {
 	UInteractiveTool::Setup();
@@ -83,7 +63,7 @@ void UVoxelCSGMeshesTool::Setup()
 	// Hide the source meshes
 	for (int32 ComponentIdx = 0; ComponentIdx < Targets.Num(); ComponentIdx++)
 	{
-		TargetComponentInterface(ComponentIdx)->SetOwnerVisibility(false);
+		UE::ToolTarget::HideSourceObject(Targets[ComponentIdx]);
 	}
 
 	// save transformed version of input meshes (maybe this could happen in the Operator?)
@@ -123,7 +103,7 @@ void UVoxelCSGMeshesTool::Shutdown(EToolShutdownType ShutdownType)
 	// Restore (unhide) the source meshes
 	for (int32 ComponentIdx = 0; ComponentIdx < Targets.Num(); ComponentIdx++)
 	{
-		TargetComponentInterface(ComponentIdx)->SetOwnerVisibility(true);
+		UE::ToolTarget::ShowSourceObject(Targets[ComponentIdx]);
 	}
 	if (ShutdownType == EToolShutdownType::Accept)
 	{
@@ -135,7 +115,7 @@ void UVoxelCSGMeshesTool::Shutdown(EToolShutdownType ShutdownType)
 		TArray<AActor*> Actors;
 		for (int32 ComponentIdx = 0; ComponentIdx < Targets.Num(); ComponentIdx++)
 		{
-			Actors.Add(TargetComponentInterface(ComponentIdx)->GetOwnerActor());
+			Actors.Add(UE::ToolTarget::GetTargetActor(Targets[ComponentIdx]));
 		}
 		HandleSourcesProperties->ApplyMethod(Actors, GetToolManager());
 
@@ -184,8 +164,8 @@ void UVoxelCSGMeshesTool::CacheInputMeshes()
 	for (int32 ComponentIdx = 0; ComponentIdx < Targets.Num(); ComponentIdx++)
 	{
 		UE::Geometry::FVoxelBooleanMeshesOp::FInputMesh InputMesh;
-		InputMesh.Mesh = TargetMeshProviderInterface(ComponentIdx)->GetMeshDescription();
-		InputMesh.Transform = TargetComponentInterface(ComponentIdx)->GetWorldTransform();
+		InputMesh.Mesh = UE::ToolTarget::GetMeshDescription(Targets[ComponentIdx]);
+		InputMesh.Transform = (FTransform) UE::ToolTarget::GetLocalToWorldTransform(Targets[ComponentIdx]);
 		InputMeshes.Add(InputMesh);
 	}
 }
diff --git a/Engine/Plugins/Experimental/MeshModelingToolsetExp/Source/MeshModelingToolsEditorOnly/Public/AddPivotActorTool.h b/Engine/Plugins/Experimental/MeshModelingToolsetExp/Source/MeshModelingToolsEditorOnly/Public/AddPivotActorTool.h
--- a/Engine/Plugins/Experimental/MeshModelingToolsetExp/Source/MeshModelingToolsEditorOnly/Public/AddPivotActorTool.h
+++ b/Engine/Plugins/Experimental/MeshModelingToolsetExp/Source/MeshModelingToolsEditorOnly/Public/AddPivotActorTool.h
@@ -5,7 +5,7 @@
 #include "CoreMinimal.h"
 
 #include "InteractiveToolBuilder.h"
-#include "MultiSelectionTool.h"
+#include "BaseTools/MultiSelectionMeshEditingTool.h"
 #include "GameFramework/Actor.h"
 
 #include "AddPivotActorTool.generated.h"
@@ -16,15 +16,15 @@ class UTransformProxy;
 
 
 UCLASS()
-class MESHMODELINGTOOLSEDITORONLY_API UAddPivotActorToolBuilder : public UInteractiveToolBuilder
+class MESHMODELINGTOOLSEDITORONLY_API UAddPivotActorToolBuilder : public UMultiSelectionMeshEditingToolBuilder
 {
 	GENERATED_BODY()
 public:
 	virtual bool CanBuildTool(const FToolBuilderState& SceneState) const override;
-	virtual UInteractiveTool* BuildTool(const FToolBuilderState& SceneState) const override;
+	virtual UMultiSelectionMeshEditingTool* CreateNewTool(const FToolBuilderState& SceneState) const override;
 
 protected:
-	virtual const FToolTargetTypeRequirements& GetTargetRequirements() const;
+	virtual const FToolTargetTypeRequirements& GetTargetRequirements() const override;
 };
 
 /** 
@@ -33,12 +33,11 @@ protected:
  * animation.
  */
 UCLASS()
-class MESHMODELINGTOOLSEDITORONLY_API UAddPivotActorTool : public UMultiSelectionTool
+class MESHMODELINGTOOLSEDITORONLY_API UAddPivotActorTool : public UMultiSelectionMeshEditingTool
 {
 	GENERATED_BODY()
 public:
 
-	virtual void SetWorld(UWorld* World) { TargetWorld = World; }
 	virtual void SetPivotRepositionMode(AActor* PivotActor) { ExistingPivotActor = PivotActor; }
 
 	virtual void Setup() override;
@@ -61,7 +60,6 @@ protected:
 	UPROPERTY()
 	TObjectPtr<UTransformProxy> TransformProxy = nullptr;
 
-	TObjectPtr<UWorld> TargetWorld;
 	TWeakObjectPtr<AActor> ExistingPivotActor = nullptr;
 
 	FTransform ExistingPivotOriginalTransform;
diff --git a/Engine/Plugins/Experimental/MeshModelingToolsetExp/Source/MeshModelingToolsEditorOnly/Public/AttributeEditorTool.h b/Engine/Plugins/Experimental/MeshModelingToolsetExp/Source/MeshModelingToolsEditorOnly/Public/AttributeEditorTool.h
--- a/Engine/Plugins/Experimental/MeshModelingToolsetExp/Source/MeshModelingToolsEditorOnly/Public/AttributeEditorTool.h
+++ b/Engine/Plugins/Experimental/MeshModelingToolsetExp/Source/MeshModelingToolsEditorOnly/Public/AttributeEditorTool.h
@@ -4,26 +4,26 @@
 
 #include "CoreMinimal.h"
 #include "UObject/NoExportTypes.h"
-#include "MultiSelectionTool.h"
+#include "BaseTools/MultiSelectionMeshEditingTool.h"
 #include "InteractiveToolBuilder.h"
 #include "InteractiveToolChange.h"
 #include "AttributeEditorTool.generated.h"
 
 
+class IPrimitiveComponentBackedTarget;
+class UAttributeEditorTool;
+
+
 /**
  *
  */
 UCLASS()
-class MESHMODELINGTOOLSEDITORONLY_API UAttributeEditorToolBuilder : public UInteractiveToolBuilder
+class MESHMODELINGTOOLSEDITORONLY_API UAttributeEditorToolBuilder : public UMultiSelectionMeshEditingToolBuilder
 {
 	GENERATED_BODY()
 
 public:
-	virtual bool CanBuildTool(const FToolBuilderState& SceneState) const override;
-	virtual UInteractiveTool* BuildTool(const FToolBuilderState& SceneState) const override;
-
-protected:
-	virtual const FToolTargetTypeRequirements& GetTargetRequirements() const;
+	virtual UMultiSelectionMeshEditingTool* CreateNewTool(const FToolBuilderState& SceneState) const override;
 };
 
 
@@ -327,7 +327,7 @@ public:
  * Mesh Attribute Editor Tool
  */
 UCLASS()
-class MESHMODELINGTOOLSEDITORONLY_API UAttributeEditorTool : public UMultiSelectionTool
+class MESHMODELINGTOOLSEDITORONLY_API UAttributeEditorTool : public UMultiSelectionMeshEditingTool
 {
 	GENERATED_BODY()
 
@@ -339,8 +339,6 @@ public:
 	virtual void Shutdown(EToolShutdownType ShutdownType) override;
 	virtual void OnTick(float DeltaTime) override;
 
-	virtual void SetWorld(UWorld* World);
-
 	virtual bool HasCancel() const override { return false; }
 	virtual bool HasAccept() const override { return false; }
 	virtual bool CanAccept() const override { return false; }
@@ -372,8 +370,6 @@ protected:
 
 
 protected:
-	UWorld* TargetWorld;
-
 	bool bTargetIsStaticMesh = true;
 
 	TArray<FAttributeEditorAttribInfo> VertexAttributes;
@@ -387,7 +383,7 @@ protected:
 	bool bHaveAutoGeneratedLightmapUVSet = false;
 	void InitializeAttributeLists();
 
-	void UpdateAutoGeneratedLightmapUVChannel(IPrimitiveComponentBackedTarget* Target, int32 NewMaxUVChannels);
+	void UpdateAutoGeneratedLightmapUVChannel(UPrimitiveComponent* TargetComponent, int32 NewMaxUVChannels);
 
 	EAttributeEditorToolActions PendingAction = EAttributeEditorToolActions::NoAction;
 	void ClearNormals();
diff --git a/Engine/Plugins/Experimental/MeshModelingToolsetExp/Source/MeshModelingToolsEditorOnly/Public/MergeMeshesTool.h b/Engine/Plugins/Experimental/MeshModelingToolsetExp/Source/MeshModelingToolsEditorOnly/Public/MergeMeshesTool.h
--- a/Engine/Plugins/Experimental/MeshModelingToolsetExp/Source/MeshModelingToolsEditorOnly/Public/MergeMeshesTool.h
+++ b/Engine/Plugins/Experimental/MeshModelingToolsetExp/Source/MeshModelingToolsEditorOnly/Public/MergeMeshesTool.h
@@ -3,7 +3,7 @@
 #pragma once
 
 #include "CoreMinimal.h"
-#include "MultiSelectionTool.h"
+#include "BaseTools/MultiSelectionMeshEditingTool.h"
 #include "InteractiveToolBuilder.h"
 #include "MeshOpPreviewHelpers.h"
 #include "Properties/MeshStatisticsProperties.h"
@@ -15,16 +15,16 @@
  *
  */
 UCLASS()
-class MESHMODELINGTOOLSEDITORONLY_API UMergeMeshesToolBuilder : public UInteractiveToolBuilder
+class MESHMODELINGTOOLSEDITORONLY_API UMergeMeshesToolBuilder : public UMultiSelectionMeshEditingToolBuilder
 {
 	GENERATED_BODY()
 
 public:
 	virtual bool CanBuildTool(const FToolBuilderState& SceneState) const override;
-	virtual UInteractiveTool* BuildTool(const FToolBuilderState& SceneState) const override;
+	virtual UMultiSelectionMeshEditingTool* CreateNewTool(const FToolBuilderState& SceneState) const override;
 
 protected:
-	virtual const FToolTargetTypeRequirements& GetTargetRequirements() const;
+	virtual const FToolTargetTypeRequirements& GetTargetRequirements() const override;
 };
 
 
@@ -59,15 +59,13 @@ public:
  *
  */
 UCLASS()
-class MESHMODELINGTOOLSEDITORONLY_API UMergeMeshesTool : public UMultiSelectionTool, public UE::Geometry::IDynamicMeshOperatorFactory
+class MESHMODELINGTOOLSEDITORONLY_API UMergeMeshesTool : public UMultiSelectionMeshEditingTool, public UE::Geometry::IDynamicMeshOperatorFactory
 {
 	GENERATED_BODY()
 
 public:
 	UMergeMeshesTool();
 
-	virtual void SetWorld(UWorld* World);
-
 	virtual void Setup() override;
 	virtual void Shutdown(EToolShutdownType ShutdownType) override;
 
@@ -97,8 +95,6 @@ protected:
 	TObjectPtr<UMeshOpPreviewWithBackgroundCompute> Preview;
 
 protected:
-	UWorld* TargetWorld;
-
 	TArray<UE::Geometry::FVoxelMergeMeshesOp::FInputMesh> InputMeshes;
 	/** stash copies of the transforms and pointers to the meshes for consumption by merge Op*/
 	void CacheInputMeshes();
diff --git a/Engine/Plugins/Experimental/MeshModelingToolsetExp/Source/MeshModelingToolsEditorOnly/Public/VoxelCSGMeshesTool.h b/Engine/Plugins/Experimental/MeshModelingToolsetExp/Source/MeshModelingToolsEditorOnly/Public/VoxelCSGMeshesTool.h
--- a/Engine/Plugins/Experimental/MeshModelingToolsetExp/Source/MeshModelingToolsEditorOnly/Public/VoxelCSGMeshesTool.h
+++ b/Engine/Plugins/Experimental/MeshModelingToolsetExp/Source/MeshModelingToolsEditorOnly/Public/VoxelCSGMeshesTool.h
@@ -5,7 +5,7 @@
 #include "CoreMinimal.h"
 #include "InteractiveToolBuilder.h"
 #include "MeshOpPreviewHelpers.h"
-#include "MultiSelectionTool.h"
+#include "BaseTools/MultiSelectionMeshEditingTool.h"
 #include "Properties/MeshStatisticsProperties.h"
 #include "PropertySets/OnAcceptProperties.h"
 #include "CompositionOps/VoxelBooleanMeshesOp.h"
@@ -16,16 +16,13 @@
  *
  */
 UCLASS()
-class MESHMODELINGTOOLSEDITORONLY_API UVoxelCSGMeshesToolBuilder : public UInteractiveToolBuilder
+class MESHMODELINGTOOLSEDITORONLY_API UVoxelCSGMeshesToolBuilder : public UMultiSelectionMeshEditingToolBuilder
 {
 	GENERATED_BODY()
 
 public:
 	virtual bool CanBuildTool(const FToolBuilderState& SceneState) const override;
-	virtual UInteractiveTool* BuildTool(const FToolBuilderState& SceneState) const override;
-
-protected:
-	virtual const FToolTargetTypeRequirements& GetTargetRequirements() const;
+	virtual UMultiSelectionMeshEditingTool* CreateNewTool(const FToolBuilderState& SceneState) const override;
 };
 
 
@@ -84,15 +81,13 @@ public:
  *
  */
 UCLASS()
-class MESHMODELINGTOOLSEDITORONLY_API UVoxelCSGMeshesTool : public UMultiSelectionTool, public UE::Geometry::IDynamicMeshOperatorFactory
+class MESHMODELINGTOOLSEDITORONLY_API UVoxelCSGMeshesTool : public UMultiSelectionMeshEditingTool, public UE::Geometry::IDynamicMeshOperatorFactory
 {
 	GENERATED_BODY()
 
 public:
 	UVoxelCSGMeshesTool();
 
-	virtual void SetWorld(UWorld* World);
-
 	virtual void Setup() override;
 	virtual void Shutdown(EToolShutdownType ShutdownType) override;
 
@@ -123,8 +118,6 @@ protected:
 	TObjectPtr<UMeshOpPreviewWithBackgroundCompute> Preview;
 
 protected:
-	UWorld* TargetWorld;
-
 	TArray<UE::Geometry::FVoxelBooleanMeshesOp::FInputMesh> InputMeshes;
 	/** stash copies of the transforms and pointers to the meshes for consumption by the CSG Op*/
 	void CacheInputMeshes();
diff --git a/Engine/Plugins/Experimental/MeshModelingToolsetExp/Source/MeshModelingToolsExp/Private/AlignObjectsTool.cpp b/Engine/Plugins/Experimental/MeshModelingToolsetExp/Source/MeshModelingToolsExp/Private/AlignObjectsTool.cpp
--- a/Engine/Plugins/Experimental/MeshModelingToolsetExp/Source/MeshModelingToolsExp/Private/AlignObjectsTool.cpp
+++ b/Engine/Plugins/Experimental/MeshModelingToolsetExp/Source/MeshModelingToolsExp/Private/AlignObjectsTool.cpp
@@ -2,7 +2,6 @@
 
 #include "AlignObjectsTool.h"
 #include "InteractiveToolManager.h"
-#include "InteractiveGizmoManager.h"
 #include "ToolBuilderUtil.h"
 #include "ToolSetupUtil.h"
 #include "DynamicMesh/DynamicMesh3.h"
@@ -11,6 +10,7 @@
 #include "BaseBehaviors/ClickDragBehavior.h"
 #include "TargetInterfaces/PrimitiveComponentBackedTarget.h"
 #include "ToolTargetManager.h"
+#include "ModelingToolTargetUtil.h"
 
 #include "Components/PrimitiveComponent.h"
 #include "CollisionQueryParams.h"
@@ -37,15 +37,9 @@ bool UAlignObjectsToolBuilder::CanBuildTool(const FToolBuilderState& SceneState)
 	return SceneState.TargetManager->CountSelectedAndTargetable(SceneState, GetTargetRequirements()) >= 2;
 }
 
-UInteractiveTool* UAlignObjectsToolBuilder::BuildTool(const FToolBuilderState& SceneState) const
+UMultiSelectionMeshEditingTool* UAlignObjectsToolBuilder::CreateNewTool(const FToolBuilderState& SceneState) const
 {
-	UAlignObjectsTool* NewTool = NewObject<UAlignObjectsTool>(SceneState.ToolManager);
-
-	TArray<TObjectPtr<UToolTarget>> Targets = SceneState.TargetManager->BuildAllSelectedTargetable(SceneState, GetTargetRequirements());
-	NewTool->SetTargets(MoveTemp(Targets));
-	NewTool->SetWorld(SceneState.World, SceneState.GizmoManager);
-
-	return NewTool;
+	return NewObject<UAlignObjectsTool>(SceneState.ToolManager);
 }
 
 /*
@@ -56,12 +50,6 @@ UAlignObjectsTool::UAlignObjectsTool()
 {
 }
 
-void UAlignObjectsTool::SetWorld(UWorld* World, UInteractiveGizmoManager* GizmoManagerIn)
-{
-	this->TargetWorld = World;
-	this->GizmoManager = GizmoManagerIn;
-}
-
 
 void UAlignObjectsTool::Setup()
 {
@@ -147,13 +135,10 @@ void UAlignObjectsTool::Precompute()
 
 	for (TObjectPtr<UToolTarget>& Target : Targets)
 	{
-		IPrimitiveComponentBackedTarget* Component = Cast<IPrimitiveComponentBackedTarget>(Target);
-		check(Component);
-
 		FAlignInfo AlignInfo;
-		AlignInfo.Component = Component->GetOwnerComponent();
-		AlignInfo.SavedTransform = Component->GetWorldTransform();
-		AlignInfo.WorldTransform = FTransform3d(Component->GetWorldTransform());
+		AlignInfo.Component = UE::ToolTarget::GetTargetComponent(Target);
+		AlignInfo.SavedTransform = (FTransform) UE::ToolTarget::GetLocalToWorldTransform(Target);
+		AlignInfo.WorldTransform = UE::ToolTarget::GetLocalToWorldTransform(Target);
 		AlignInfo.WorldBounds = FAxisAlignedBox3d(AlignInfo.Component->Bounds.GetBox());
 		AlignInfo.WorldPivot = AlignInfo.WorldTransform.TransformPosition(FVector3d::Zero());
 		ComponentInfo.Add(AlignInfo);
diff --git a/Engine/Plugins/Experimental/MeshModelingToolsetExp/Source/MeshModelingToolsExp/Private/BakeMeshAttributeMapsTool.cpp b/Engine/Plugins/Experimental/MeshModelingToolsetExp/Source/MeshModelingToolsExp/Private/BakeMeshAttributeMapsTool.cpp
--- a/Engine/Plugins/Experimental/MeshModelingToolsetExp/Source/MeshModelingToolsExp/Private/BakeMeshAttributeMapsTool.cpp
+++ b/Engine/Plugins/Experimental/MeshModelingToolsetExp/Source/MeshModelingToolsExp/Private/BakeMeshAttributeMapsTool.cpp
@@ -73,14 +73,9 @@ bool UBakeMeshAttributeMapsToolBuilder::CanBuildTool(const FToolBuilderState& Sc
 	return false;
 }
 
-UInteractiveTool* UBakeMeshAttributeMapsToolBuilder::BuildTool(const FToolBuilderState& SceneState) const
+UMultiSelectionMeshEditingTool* UBakeMeshAttributeMapsToolBuilder::CreateNewTool(const FToolBuilderState& SceneState) const
 {
-	UBakeMeshAttributeMapsTool* NewTool = NewObject<UBakeMeshAttributeMapsTool>(SceneState.ToolManager);
-
-	TArray<TObjectPtr<UToolTarget>> Targets = SceneState.TargetManager->BuildAllSelectedTargetable(SceneState, GetTargetRequirements());
-	NewTool->SetTargets(MoveTemp(Targets));
-	NewTool->SetWorld(SceneState.World);
-	return NewTool;
+	return NewObject<UBakeMeshAttributeMapsTool>(SceneState.ToolManager);
 }
 
 
diff --git a/Engine/Plugins/Experimental/MeshModelingToolsetExp/Source/MeshModelingToolsExp/Private/BakeMeshAttributeTool.cpp b/Engine/Plugins/Experimental/MeshModelingToolsetExp/Source/MeshModelingToolsExp/Private/BakeMeshAttributeTool.cpp
--- a/Engine/Plugins/Experimental/MeshModelingToolsetExp/Source/MeshModelingToolsExp/Private/BakeMeshAttributeTool.cpp
+++ b/Engine/Plugins/Experimental/MeshModelingToolsetExp/Source/MeshModelingToolsExp/Private/BakeMeshAttributeTool.cpp
@@ -20,12 +20,6 @@ void UBakeMeshAttributeTool::Setup()
 }
 
 
-void UBakeMeshAttributeTool::SetWorld(UWorld* World)
-{
-	TargetWorld = World;
-}
-
-
 int UBakeMeshAttributeTool::SelectColorTextureToBake(const TArray<UTexture*>& Textures)
 {
 	TArray<int> TextureVotes;
diff --git a/Engine/Plugins/Experimental/MeshModelingToolsetExp/Source/MeshModelingToolsExp/Private/BakeMeshAttributeVertexTool.cpp b/Engine/Plugins/Experimental/MeshModelingToolsetExp/Source/MeshModelingToolsExp/Private/BakeMeshAttributeVertexTool.cpp
--- a/Engine/Plugins/Experimental/MeshModelingToolsetExp/Source/MeshModelingToolsExp/Private/BakeMeshAttributeVertexTool.cpp
+++ b/Engine/Plugins/Experimental/MeshModelingToolsetExp/Source/MeshModelingToolsExp/Private/BakeMeshAttributeVertexTool.cpp
@@ -37,31 +37,15 @@ using namespace UE::Geometry;
  * ToolBuilder
  */
 
-const FToolTargetTypeRequirements& UBakeMeshAttributeVertexToolBuilder::GetTargetRequirements() const
-{
-	static FToolTargetTypeRequirements TypeRequirements({
-		UMeshDescriptionProvider::StaticClass(),
-		UMeshDescriptionCommitter::StaticClass(),
-		UPrimitiveComponentBackedTarget::StaticClass(),
-		UMaterialProvider::StaticClass()
-	});
-	return TypeRequirements;
-}
-
 bool UBakeMeshAttributeVertexToolBuilder::CanBuildTool(const FToolBuilderState& SceneState) const
 {
 	const int32 NumTargets = SceneState.TargetManager->CountSelectedAndTargetable(SceneState, GetTargetRequirements());
 	return (NumTargets == 1 || NumTargets == 2);
 }
 
-UInteractiveTool* UBakeMeshAttributeVertexToolBuilder::BuildTool(const FToolBuilderState& SceneState) const
+UMultiSelectionMeshEditingTool* UBakeMeshAttributeVertexToolBuilder::CreateNewTool(const FToolBuilderState& SceneState) const
 {
-	UBakeMeshAttributeVertexTool* NewTool = NewObject<UBakeMeshAttributeVertexTool>(SceneState.ToolManager);
-	TArray<TObjectPtr<UToolTarget>> Targets = SceneState.TargetManager->BuildAllSelectedTargetable(
-		SceneState, GetTargetRequirements());
-	NewTool->SetTargets(MoveTemp(Targets));
-	NewTool->SetWorld(SceneState.World);
-	return NewTool;
+	return NewObject<UBakeMeshAttributeVertexTool>(SceneState.ToolManager);
 }
 
 
@@ -468,9 +452,9 @@ TUniquePtr<UE::Geometry::TGenericDataOperator<FMeshVertexBaker>> UBakeMeshAttrib
 
 void UBakeMeshAttributeVertexTool::UpdateDetailMesh()
 {
-	IPrimitiveComponentBackedTarget* TargetComponent = TargetComponentInterface(0);
-	IPrimitiveComponentBackedTarget* DetailComponent = TargetComponentInterface(bIsBakeToSelf ? 0 : 1);
-	IMeshDescriptionProvider* DetailMeshProvider = TargetMeshProviderInterface(bIsBakeToSelf ? 0 : 1);
+	IPrimitiveComponentBackedTarget* TargetComponent = Cast<IPrimitiveComponentBackedTarget>(Targets[0]);
+	IPrimitiveComponentBackedTarget* DetailComponent = Cast<IPrimitiveComponentBackedTarget>(Targets[bIsBakeToSelf ? 0 : 1]);
+	IMeshDescriptionProvider* DetailMeshProvider = Cast<IMeshDescriptionProvider>(Targets[bIsBakeToSelf ? 0 : 1]);
 
 	DetailMesh = MakeShared<FDynamicMesh3, ESPMode::ThreadSafe>();
 	FMeshDescriptionToDynamicMesh Converter;
diff --git a/Engine/Plugins/Experimental/MeshModelingToolsetExp/Source/MeshModelingToolsExp/Private/BakeMultiMeshAttributeMapsTool.cpp b/Engine/Plugins/Experimental/MeshModelingToolsetExp/Source/MeshModelingToolsExp/Private/BakeMultiMeshAttributeMapsTool.cpp
--- a/Engine/Plugins/Experimental/MeshModelingToolsetExp/Source/MeshModelingToolsExp/Private/BakeMultiMeshAttributeMapsTool.cpp
+++ b/Engine/Plugins/Experimental/MeshModelingToolsetExp/Source/MeshModelingToolsExp/Private/BakeMultiMeshAttributeMapsTool.cpp
@@ -53,13 +53,9 @@ bool UBakeMultiMeshAttributeMapsToolBuilder::CanBuildTool(const FToolBuilderStat
 	return (NumTargets > 1);
 }
 
-UInteractiveTool* UBakeMultiMeshAttributeMapsToolBuilder::BuildTool(const FToolBuilderState& SceneState) const
+UMultiSelectionMeshEditingTool* UBakeMultiMeshAttributeMapsToolBuilder::CreateNewTool(const FToolBuilderState& SceneState) const
 {
-	UBakeMultiMeshAttributeMapsTool* NewTool = NewObject<UBakeMultiMeshAttributeMapsTool>(SceneState.ToolManager);
-	TArray<TObjectPtr<UToolTarget>> Targets = SceneState.TargetManager->BuildAllSelectedTargetable(SceneState, GetTargetRequirements());
-	NewTool->SetTargets(MoveTemp(Targets));
-	NewTool->SetWorld(SceneState.World);
-	return NewTool;
+	return NewObject<UBakeMultiMeshAttributeMapsTool>(SceneState.ToolManager);
 }
 
 
diff --git a/Engine/Plugins/Experimental/MeshModelingToolsetExp/Source/MeshModelingToolsExp/Private/BakeTransformTool.cpp b/Engine/Plugins/Experimental/MeshModelingToolsetExp/Source/MeshModelingToolsExp/Private/BakeTransformTool.cpp
--- a/Engine/Plugins/Experimental/MeshModelingToolsetExp/Source/MeshModelingToolsExp/Private/BakeTransformTool.cpp
+++ b/Engine/Plugins/Experimental/MeshModelingToolsetExp/Source/MeshModelingToolsExp/Private/BakeTransformTool.cpp
@@ -32,31 +32,9 @@ using namespace UE::Geometry;
  * ToolBuilder
  */
 
-
-const FToolTargetTypeRequirements& UBakeTransformToolBuilder::GetTargetRequirements() const
-{
-	static FToolTargetTypeRequirements TypeRequirements({
-		UMeshDescriptionCommitter::StaticClass(),
-		UMeshDescriptionProvider::StaticClass(),
-		UPrimitiveComponentBackedTarget::StaticClass()
-		});
-	return TypeRequirements;
-}
-
-bool UBakeTransformToolBuilder::CanBuildTool(const FToolBuilderState& SceneState) const
-{
-	return SceneState.TargetManager->CountSelectedAndTargetable(SceneState, GetTargetRequirements()) > 0;
-}
-
-UInteractiveTool* UBakeTransformToolBuilder::BuildTool(const FToolBuilderState& SceneState) const
+UMultiSelectionMeshEditingTool* UBakeTransformToolBuilder::CreateNewTool(const FToolBuilderState& SceneState) const
 {
-	UBakeTransformTool* NewTool = NewObject<UBakeTransformTool>(SceneState.ToolManager);
-
-	TArray<TObjectPtr<UToolTarget>> Targets = SceneState.TargetManager->BuildAllSelectedTargetable(SceneState, GetTargetRequirements());
-	NewTool->SetTargets(MoveTemp(Targets));
-	NewTool->SetWorld(SceneState.World);
-
-	return NewTool;
+	return NewObject<UBakeTransformTool>(SceneState.ToolManager);
 }
 
 
@@ -71,15 +49,10 @@ UBakeTransformToolProperties::UBakeTransformToolProperties()
 }
 
 
-
 UBakeTransformTool::UBakeTransformTool()
 {
 }
 
-void UBakeTransformTool::SetWorld(UWorld* World)
-{
-	this->TargetWorld = World;
-}
 
 void UBakeTransformTool::Setup()
 {
diff --git a/Engine/Plugins/Experimental/MeshModelingToolsetExp/Source/MeshModelingToolsExp/Private/ConvertMeshesTool.cpp b/Engine/Plugins/Experimental/MeshModelingToolsetExp/Source/MeshModelingToolsExp/Private/ConvertMeshesTool.cpp
--- a/Engine/Plugins/Experimental/MeshModelingToolsetExp/Source/MeshModelingToolsExp/Private/ConvertMeshesTool.cpp
+++ b/Engine/Plugins/Experimental/MeshModelingToolsetExp/Source/MeshModelingToolsExp/Private/ConvertMeshesTool.cpp
@@ -10,6 +10,9 @@
 #include "ModelingObjectsCreationAPI.h"
 #include "Selection/ToolSelectionUtil.h"
 
+#include "TargetInterfaces/MeshDescriptionProvider.h"
+#include "TargetInterfaces/PrimitiveComponentBackedTarget.h"
+
 using namespace UE::Geometry;
 
 #define LOCTEXT_NAMESPACE "UConvertMeshesTool"
@@ -27,18 +30,9 @@ const FToolTargetTypeRequirements& UConvertMeshesToolBuilder::GetTargetRequireme
 	return TypeRequirements;
 }
 
-bool UConvertMeshesToolBuilder::CanBuildTool(const FToolBuilderState& SceneState) const
-{
-	return SceneState.TargetManager->CountSelectedAndTargetable(SceneState, GetTargetRequirements()) > 0;
-}
-
-UInteractiveTool* UConvertMeshesToolBuilder::BuildTool(const FToolBuilderState& SceneState) const
+UMultiSelectionMeshEditingTool* UConvertMeshesToolBuilder::CreateNewTool(const FToolBuilderState& SceneState) const
 {
-	UConvertMeshesTool* NewTool = NewObject<UConvertMeshesTool>(SceneState.ToolManager);
-	TArray<TObjectPtr<UToolTarget>> Targets = SceneState.TargetManager->BuildAllSelectedTargetable(SceneState, GetTargetRequirements());
-	NewTool->SetTargets(MoveTemp(Targets));
-	NewTool->SetWorld(SceneState.World);
-	return NewTool;
+	return NewObject<UConvertMeshesTool>(SceneState.ToolManager);
 }
 
 
@@ -48,11 +42,6 @@ UInteractiveTool* UConvertMeshesToolBuilder::BuildTool(const FToolBuilderState&
  */
 
 
-void UConvertMeshesTool::SetWorld(UWorld* World)
-{
-	this->TargetWorld = World;
-}
-
 void UConvertMeshesTool::Setup()
 {
 	UInteractiveTool::Setup();
diff --git a/Engine/Plugins/Experimental/MeshModelingToolsetExp/Source/MeshModelingToolsExp/Private/EditNormalsTool.cpp b/Engine/Plugins/Experimental/MeshModelingToolsetExp/Source/MeshModelingToolsExp/Private/EditNormalsTool.cpp
--- a/Engine/Plugins/Experimental/MeshModelingToolsetExp/Source/MeshModelingToolsExp/Private/EditNormalsTool.cpp
+++ b/Engine/Plugins/Experimental/MeshModelingToolsetExp/Source/MeshModelingToolsExp/Private/EditNormalsTool.cpp
@@ -5,6 +5,7 @@
 #include "ToolBuilderUtil.h"
 
 #include "ToolSetupUtil.h"
+#include "ModelingToolTargetUtil.h"
 
 #include "DynamicMesh/DynamicMesh3.h"
 #include "BaseBehaviors/MultiClickSequenceInputBehavior.h"
@@ -33,30 +34,9 @@ using namespace UE::Geometry;
  */
 
 
-const FToolTargetTypeRequirements& UEditNormalsToolBuilder::GetTargetRequirements() const
+UMultiSelectionMeshEditingTool* UEditNormalsToolBuilder::CreateNewTool(const FToolBuilderState& SceneState) const
 {
-	static FToolTargetTypeRequirements TypeRequirements({
-		UMeshDescriptionCommitter::StaticClass(),
-		UMeshDescriptionProvider::StaticClass(),
-		UPrimitiveComponentBackedTarget::StaticClass()
-		});
-	return TypeRequirements;
-}
-
-bool UEditNormalsToolBuilder::CanBuildTool(const FToolBuilderState& SceneState) const
-{
-	return SceneState.TargetManager->CountSelectedAndTargetable(SceneState, GetTargetRequirements()) > 0;
-}
-
-UInteractiveTool* UEditNormalsToolBuilder::BuildTool(const FToolBuilderState& SceneState) const
-{
-	UEditNormalsTool* NewTool = NewObject<UEditNormalsTool>(SceneState.ToolManager);
-
-	TArray<TObjectPtr<UToolTarget>> Targets = SceneState.TargetManager->BuildAllSelectedTargetable(SceneState, GetTargetRequirements());
-	NewTool->SetTargets(MoveTemp(Targets));
-	NewTool->SetWorld(SceneState.World);
-
-	return NewTool;
+	return NewObject<UEditNormalsTool>(SceneState.ToolManager);
 }
 
 
@@ -86,10 +66,6 @@ UEditNormalsTool::UEditNormalsTool()
 {
 }
 
-void UEditNormalsTool::SetWorld(UWorld* World)
-{
-	this->TargetWorld = World;
-}
 
 void UEditNormalsTool::Setup()
 {
@@ -98,8 +74,7 @@ void UEditNormalsTool::Setup()
 	// hide input StaticMeshComponent
 	for (int32 ComponentIdx = 0, NumTargets = Targets.Num(); ComponentIdx < NumTargets; ComponentIdx++)
 	{
-		IPrimitiveComponentBackedTarget* TargetComponent = TargetComponentInterface(ComponentIdx);
-		TargetComponent->SetOwnerVisibility(false);
+		UE::ToolTarget::HideSourceObject(Targets[ComponentIdx]);
 	}
 
 	BasicProperties = NewObject<UEditNormalsToolProperties>(this, TEXT("Mesh Normals Settings"));
@@ -150,15 +125,14 @@ void UEditNormalsTool::UpdateNumPreviews()
 			OpFactory->ComponentIndex = PreviewIdx;
 			OriginalDynamicMeshes[PreviewIdx] = MakeShared<FDynamicMesh3, ESPMode::ThreadSafe>();
 			FMeshDescriptionToDynamicMesh Converter;
-			Converter.Convert(TargetMeshProviderInterface(PreviewIdx)->GetMeshDescription(), *OriginalDynamicMeshes[PreviewIdx]);
+			Converter.Convert(UE::ToolTarget::GetMeshDescription(Targets[PreviewIdx]), *OriginalDynamicMeshes[PreviewIdx]);
 
 			UMeshOpPreviewWithBackgroundCompute* Preview = Previews.Add_GetRef(NewObject<UMeshOpPreviewWithBackgroundCompute>(OpFactory, "Preview"));
 			Preview->Setup(this->TargetWorld, OpFactory);
 			ToolSetupUtil::ApplyRenderingConfigurationToPreview(Preview->PreviewMesh, Targets[PreviewIdx]);
 			Preview->PreviewMesh->SetTangentsMode(EDynamicMeshComponentTangentsMode::AutoCalculated);
 
-			FComponentMaterialSet MaterialSet;
-			TargetMaterialInterface(PreviewIdx)->GetMaterialSet(MaterialSet);
+			const FComponentMaterialSet MaterialSet = UE::ToolTarget::GetMaterialSet(Targets[PreviewIdx]);
 			Preview->ConfigureMaterials(MaterialSet.Materials,
 				ToolSetupUtil::GetDefaultWorkingMaterial(GetToolManager())
 			);
@@ -177,8 +151,7 @@ void UEditNormalsTool::Shutdown(EToolShutdownType ShutdownType)
 	// Restore (unhide) the source meshes
 	for (int32 ComponentIdx = 0, NumTargets = Targets.Num(); ComponentIdx < NumTargets; ComponentIdx++)
 	{
-		IPrimitiveComponentBackedTarget* TargetComponent = TargetComponentInterface(ComponentIdx);
-		TargetComponent->SetOwnerVisibility(true);
+		UE::ToolTarget::ShowSourceObject(Targets[ComponentIdx]);
 	}
 
 	TArray<FDynamicMeshOpResult> Results;
@@ -203,7 +176,7 @@ TUniquePtr<FDynamicMeshOperator> UEditNormalsOperatorFactory::MakeNewOperator()
 	NormalsOp->NormalCalculationMethod = Tool->BasicProperties->NormalCalculationMethod;
 	NormalsOp->NormalSplitThreshold = Tool->BasicProperties->SharpEdgeAngleThreshold;
 
-	FTransform LocalToWorld = Tool->TargetComponentInterface(ComponentIndex)->GetWorldTransform();
+	const FTransform LocalToWorld = (FTransform) UE::ToolTarget::GetLocalToWorldTransform(Tool->Targets[ComponentIndex]);
 	NormalsOp->OriginalMesh = Tool->OriginalDynamicMeshes[ComponentIndex];
 	
 	NormalsOp->SetTransform(LocalToWorld);
@@ -276,7 +249,7 @@ void UEditNormalsTool::GenerateAsset(const TArray<FDynamicMeshOpResult>& Results
 	for (int32 ComponentIdx = 0; ComponentIdx < Targets.Num(); ComponentIdx++)
 	{
 		// disable auto-generated normals StaticMesh build setting
-		UStaticMeshComponent* StaticMeshComponent = Cast<UStaticMeshComponent>(TargetComponentInterface(ComponentIdx)->GetOwnerComponent());
+		UStaticMeshComponent* StaticMeshComponent = Cast<UStaticMeshComponent>(UE::ToolTarget::GetTargetComponent(Targets[ComponentIdx]));
 		if (StaticMeshComponent != nullptr)
 		{
 			UStaticMesh* StaticMesh = StaticMeshComponent->GetStaticMesh();
@@ -290,7 +263,7 @@ void UEditNormalsTool::GenerateAsset(const TArray<FDynamicMeshOpResult>& Results
 		}
 
 		check(Results[ComponentIdx].Mesh.Get() != nullptr);
-		TargetMeshCommitterInterface(ComponentIdx)->CommitMeshDescription([&Results, &ComponentIdx, this](const IMeshDescriptionCommitter::FCommitterParams& CommitParams)
+		Cast<IMeshDescriptionCommitter>(Targets[ComponentIdx])->CommitMeshDescription([&Results, &ComponentIdx, this](const IMeshDescriptionCommitter::FCommitterParams& CommitParams)
 		{
 			FDynamicMeshToMeshDescription Converter;
 			
diff --git a/Engine/Plugins/Experimental/MeshModelingToolsetExp/Source/MeshModelingToolsExp/Private/EditPivotTool.cpp b/Engine/Plugins/Experimental/MeshModelingToolsetExp/Source/MeshModelingToolsExp/Private/EditPivotTool.cpp
--- a/Engine/Plugins/Experimental/MeshModelingToolsetExp/Source/MeshModelingToolsExp/Private/EditPivotTool.cpp
+++ b/Engine/Plugins/Experimental/MeshModelingToolsetExp/Source/MeshModelingToolsExp/Private/EditPivotTool.cpp
@@ -35,31 +35,9 @@ using namespace UE::Geometry;
  * ToolBuilder
  */
 
-
-const FToolTargetTypeRequirements& UEditPivotToolBuilder::GetTargetRequirements() const
-{
-	static FToolTargetTypeRequirements TypeRequirements({
-		UMeshDescriptionCommitter::StaticClass(),
-		UMeshDescriptionProvider::StaticClass(),
-		UPrimitiveComponentBackedTarget::StaticClass()
-	});
-	return TypeRequirements;
-}
-
-bool UEditPivotToolBuilder::CanBuildTool(const FToolBuilderState& SceneState) const
-{
-	return SceneState.TargetManager->CountSelectedAndTargetable(SceneState, GetTargetRequirements()) >= 1;
-}
-
-UInteractiveTool* UEditPivotToolBuilder::BuildTool(const FToolBuilderState& SceneState) const
+UMultiSelectionMeshEditingTool* UEditPivotToolBuilder::CreateNewTool(const FToolBuilderState& SceneState) const
 {
-	UEditPivotTool* NewTool = NewObject<UEditPivotTool>(SceneState.ToolManager);
-
-	TArray<TObjectPtr<UToolTarget>> Targets = SceneState.TargetManager->BuildAllSelectedTargetable(SceneState, GetTargetRequirements());
-	NewTool->SetTargets(MoveTemp(Targets));
-	NewTool->SetWorld(SceneState.World, SceneState.GizmoManager);
-
-	return NewTool;
+	return NewObject<UEditPivotTool>(SceneState.ToolManager);
 }
 
 
@@ -84,12 +62,6 @@ UEditPivotTool::UEditPivotTool()
 {
 }
 
-void UEditPivotTool::SetWorld(UWorld* World, UInteractiveGizmoManager* GizmoManagerIn)
-{
-	this->TargetWorld = World;
-	this->GizmoManager = GizmoManagerIn;
-}
-
 
 void UEditPivotTool::Setup()
 {
@@ -154,7 +126,7 @@ void UEditPivotTool::Shutdown(EToolShutdownType ShutdownType)
 
 	FFrame3d CurPivotFrame(ActiveGizmos[0].TransformProxy->GetTransform());
 
-	GizmoManager->DestroyAllGizmosByOwner(this);
+	GetToolManager()->GetPairedGizmoManager()->DestroyAllGizmosByOwner(this);
 
 	if (ShutdownType == EToolShutdownType::Accept)
 	{
@@ -324,7 +296,7 @@ void UEditPivotTool::SetActiveGizmos_Single(bool bLocalRotations)
 	{
 		Transformable.TransformProxy->AddComponent(UE::ToolTarget::GetTargetComponent(Targets[ComponentIdx]));
 	}
-	Transformable.TransformGizmo = UE::TransformGizmoUtil::CreateCustomTransformGizmo(GizmoManager,
+	Transformable.TransformGizmo = UE::TransformGizmoUtil::CreateCustomTransformGizmo(GetToolManager()->GetPairedGizmoManager(),
 		ETransformGizmoSubElements::StandardTranslateRotate, this
 	);
 	Transformable.TransformGizmo->SetActiveTarget(Transformable.TransformProxy, GetToolManager());
@@ -338,7 +310,7 @@ void UEditPivotTool::SetActiveGizmos_Single(bool bLocalRotations)
 
 void UEditPivotTool::ResetActiveGizmos()
 {
-	GizmoManager->DestroyAllGizmosByOwner(this);
+	GetToolManager()->GetPairedGizmoManager()->DestroyAllGizmosByOwner(this);
 	ActiveGizmos.Reset();
 }
 
diff --git a/Engine/Plugins/Experimental/MeshModelingToolsetExp/Source/MeshModelingToolsExp/Private/MirrorTool.cpp b/Engine/Plugins/Experimental/MeshModelingToolsetExp/Source/MeshModelingToolsExp/Private/MirrorTool.cpp
--- a/Engine/Plugins/Experimental/MeshModelingToolsetExp/Source/MeshModelingToolsExp/Private/MirrorTool.cpp
+++ b/Engine/Plugins/Experimental/MeshModelingToolsetExp/Source/MeshModelingToolsExp/Private/MirrorTool.cpp
@@ -10,6 +10,7 @@
 #include "DynamicMeshToMeshDescription.h"
 #include "InteractiveToolManager.h"
 #include "MeshDescriptionToDynamicMesh.h"
+#include "ModelingToolTargetUtil.h"
 #include "Misc/MessageDialog.h"
 #include "ToolBuilderUtil.h"
 #include "ToolSetupUtil.h"
@@ -25,33 +26,11 @@ using namespace UE::Geometry;
 #define LOCTEXT_NAMESPACE "UMirrorTool"
 
 
-//Tool builder functions
+// Tool builder functions
 
-const FToolTargetTypeRequirements& UMirrorToolBuilder::GetTargetRequirements() const
+UMultiSelectionMeshEditingTool* UMirrorToolBuilder::CreateNewTool(const FToolBuilderState& SceneState) const
 {
-	static FToolTargetTypeRequirements TypeRequirements({
-		UMaterialProvider::StaticClass(),
-		UMeshDescriptionCommitter::StaticClass(),
-		UMeshDescriptionProvider::StaticClass(),
-		UPrimitiveComponentBackedTarget::StaticClass()
-		});
-	return TypeRequirements;
-}
-
-bool UMirrorToolBuilder::CanBuildTool(const FToolBuilderState& SceneState) const
-{
-	return SceneState.TargetManager->CountSelectedAndTargetable(SceneState, GetTargetRequirements()) > 0;
-}
-
-UInteractiveTool* UMirrorToolBuilder::BuildTool(const FToolBuilderState& SceneState) const
-{
-	UMirrorTool* NewTool = NewObject<UMirrorTool>(SceneState.ToolManager);
-
-	TArray<TObjectPtr<UToolTarget>> Targets = SceneState.TargetManager->BuildAllSelectedTargetable(SceneState, GetTargetRequirements());
-	NewTool->SetTargets(MoveTemp(Targets));
-	NewTool->SetWorld(SceneState.World);
-
-	return NewTool;
+	return NewObject<UMirrorTool>(SceneState.ToolManager);
 }
 
 
@@ -69,7 +48,7 @@ TUniquePtr<FDynamicMeshOperator> UMirrorOperatorFactory::MakeNewOperator()
 	MirrorOp->PlaneTolerance = MirrorTool->Settings->PlaneTolerance;
 	MirrorOp->bAllowBowtieVertexCreation = MirrorTool->Settings->bAllowBowtieVertexCreation;
 
-	FTransform LocalToWorld = MirrorTool->TargetComponentInterface(ComponentIndex)->GetWorldTransform();
+	FTransform LocalToWorld = (FTransform) UE::ToolTarget::GetLocalToWorldTransform(MirrorTool->Targets[ComponentIndex]);
 	MirrorOp->SetTransform(LocalToWorld);
 
 	// We also need WorldToLocal. Threshold the LocalToWorld scaling transform so we can get the inverse.
@@ -124,10 +103,6 @@ bool UMirrorTool::CanAccept() const
 	return Super::CanAccept();
 }
 
-void UMirrorTool::SetWorld(UWorld* World)
-{
-	TargetWorld = World;
-}
 
 void UMirrorTool::OnPropertyModified(UObject* PropertySet, FProperty* Property)
 {
@@ -136,7 +111,7 @@ void UMirrorTool::OnPropertyModified(UObject* PropertySet, FProperty* Property)
 	{
 		for (int32 ComponentIdx = 0; ComponentIdx < Targets.Num(); ComponentIdx++)
 		{
-			TargetComponentInterface(ComponentIdx)->SetOwnerVisibility(!Settings->bShowPreview);
+			UE::ToolTarget::SetSourceObjectVisible(Targets[ComponentIdx], !Settings->bShowPreview);
 		}
 		for (UMeshOpPreviewWithBackgroundCompute* Preview : Previews)
 		{
@@ -200,13 +175,10 @@ void UMirrorTool::Setup()
 	// Fill in the MeshesToMirror array with suitably converted meshes.
 	for (int i = 0; i < Targets.Num(); i++)
 	{
-		IPrimitiveComponentBackedTarget* TargetComponent = TargetComponentInterface(i);
-		IMeshDescriptionProvider* TargetMeshProvider = TargetMeshProviderInterface(i);
-
 		// Convert into dynamic mesh
 		TSharedPtr<FDynamicMesh3, ESPMode::ThreadSafe> DynamicMesh = MakeShared<FDynamicMesh3, ESPMode::ThreadSafe>();
 		FMeshDescriptionToDynamicMesh Converter;
-		Converter.Convert(TargetMeshProvider->GetMeshDescription(), *DynamicMesh);
+		Converter.Convert(UE::ToolTarget::GetMeshDescription(Targets[i]), *DynamicMesh);
 
 		// Wrap the dynamic mesh in a replacement change target
 		UDynamicMeshReplacementChangeTarget* WrappedTarget = MeshesToMirror.Add_GetRef(NewObject<UDynamicMeshReplacementChangeTarget>());
@@ -219,7 +191,7 @@ void UMirrorTool::Setup()
 	// Set the visibility of the StaticMeshComponents depending on whether we are showing them or the preview.
 	for (int32 ComponentIdx = 0; ComponentIdx < Targets.Num(); ComponentIdx++)
 	{
-		TargetComponentInterface(ComponentIdx)->SetOwnerVisibility(!Settings->bShowPreview);
+		UE::ToolTarget::SetSourceObjectVisible(Targets[ComponentIdx], !Settings->bShowPreview);
 	}
 
 	// Initialize the PreviewMesh and BackgroundCompute objects
@@ -230,7 +202,7 @@ void UMirrorTool::Setup()
 	for (int32 ComponentIdx = 0; ComponentIdx < Targets.Num(); ComponentIdx++)
 	{
 		FVector ComponentOrigin, ComponentExtents;
-		TargetComponentInterface(ComponentIdx)->GetOwnerActor()->GetActorBounds(false, ComponentOrigin, ComponentExtents);
+		UE::ToolTarget::GetTargetActor(Targets[ComponentIdx])->GetActorBounds(false, ComponentOrigin, ComponentExtents);
 		CombinedBounds += FBox::BuildAABB(ComponentOrigin, ComponentExtents);
 	}
 
@@ -268,8 +240,7 @@ void UMirrorTool::Setup()
 	// invisible, since we want to be able to reposition the plane onto the original mesh.
 	for (int32 ComponentIdx = 0; ComponentIdx < Targets.Num(); ComponentIdx++)
 	{
-		IPrimitiveComponentBackedTarget* TargetComponent = TargetComponentInterface(ComponentIdx);
-		PlaneMechanic->SetPlaneCtrlClickBehaviorTarget->InvisibleComponentsToHitTest.Add(TargetComponent->GetOwnerComponent());
+		PlaneMechanic->SetPlaneCtrlClickBehaviorTarget->InvisibleComponentsToHitTest.Add(UE::ToolTarget::GetTargetComponent(Targets[ComponentIdx]));
 	}
 
 	// Start the preview calculations
@@ -295,13 +266,12 @@ void UMirrorTool::SetupPreviews()
 		ToolSetupUtil::ApplyRenderingConfigurationToPreview(Preview->PreviewMesh, nullptr);
 		Preview->PreviewMesh->SetTangentsMode(EDynamicMeshComponentTangentsMode::AutoCalculated);
 
-		FComponentMaterialSet MaterialSet;
-		TargetMaterialInterface(PreviewIndex)->GetMaterialSet(MaterialSet);
+		const FComponentMaterialSet MaterialSet = UE::ToolTarget::GetMaterialSet(Targets[PreviewIndex]);
 		Preview->ConfigureMaterials(MaterialSet.Materials, ToolSetupUtil::GetDefaultWorkingMaterial(GetToolManager()));
 
 		// Set initial preview to unprocessed mesh, so that things don't disappear initially
 		Preview->PreviewMesh->UpdatePreview(MeshesToMirror[PreviewIndex]->GetMesh().Get());
-		Preview->PreviewMesh->SetTransform(TargetComponentInterface(PreviewIndex)->GetWorldTransform());
+		Preview->PreviewMesh->SetTransform((FTransform) UE::ToolTarget::GetLocalToWorldTransform(Targets[PreviewIndex]));
 		Preview->SetVisibility(Settings->bShowPreview);
 	}
 }
@@ -327,10 +297,11 @@ void UMirrorTool::CheckAndDisplayWarnings()
 	IPrimitiveComponentBackedTarget* NonUniformScalingTarget = nullptr;
 	for (int32 i = 0; i < Targets.Num(); ++i)
 	{
-		FVector Scaling = TargetComponentInterface(i)->GetWorldTransform().GetScale3D();
+		IPrimitiveComponentBackedTarget* Component = Cast<IPrimitiveComponentBackedTarget>(Targets[i]);
+		const FVector Scaling = Component->GetWorldTransform().GetScale3D();
 		if (Scaling.X != Scaling.Y || Scaling.Y != Scaling.Z)
 		{
-			NonUniformScalingTarget = TargetComponentInterface(i);
+			NonUniformScalingTarget = Component;
 			break;
 		}
 	}
@@ -367,7 +338,7 @@ void UMirrorTool::Shutdown(EToolShutdownType ShutdownType)
 	// Restore (unhide) the source meshes
 	for (int32 ComponentIdx = 0; ComponentIdx < Targets.Num(); ComponentIdx++)
 	{
-		TargetComponentInterface(ComponentIdx)->SetOwnerVisibility(true);
+		UE::ToolTarget::ShowSourceObject(Targets[ComponentIdx]);
 	}
 
 	// Swap in results, if appropriate
@@ -432,7 +403,6 @@ void UMirrorTool::GenerateAsset(const TArray<FDynamicMeshOpResult>& Results)
 	NewSelection.ModificationType = ESelectedObjectsModificationType::Replace;
 	for (int OrigMeshIdx = 0; OrigMeshIdx < NumSourceMeshes; OrigMeshIdx++)
 	{
-		IPrimitiveComponentBackedTarget* TargetComponent = TargetComponentInterface(OrigMeshIdx);
 		FDynamicMesh3* Mesh = Results[OrigMeshIdx].Mesh.Get();
 		check(Mesh != nullptr);
 
@@ -440,15 +410,15 @@ void UMirrorTool::GenerateAsset(const TArray<FDynamicMeshOpResult>& Results)
 		{
 			if (bWantToDestroy)
 			{
-				TargetComponent->GetOwnerComponent()->DestroyComponent();
+				UE::ToolTarget::GetTargetComponent(Targets[OrigMeshIdx])->DestroyComponent();
 			}
 			continue;
 		}
 		else if (Settings->SaveMode == EMirrorSaveMode::UpdateAssets)
 		{
-			NewSelection.Actors.Add(TargetComponent->GetOwnerActor());
+			NewSelection.Actors.Add(UE::ToolTarget::GetTargetActor(Targets[OrigMeshIdx]));
 
-			TargetMeshCommitterInterface(OrigMeshIdx)->CommitMeshDescription([&Mesh](const IMeshDescriptionCommitter::FCommitterParams& CommitParams)
+			Cast<IMeshDescriptionCommitter>(Targets[OrigMeshIdx])->CommitMeshDescription([&Mesh](const IMeshDescriptionCommitter::FCommitterParams& CommitParams)
 				{
 					FDynamicMeshToMeshDescription Converter;
 					Converter.Convert(Mesh, *CommitParams.MeshDescriptionOut);
@@ -458,7 +428,7 @@ void UMirrorTool::GenerateAsset(const TArray<FDynamicMeshOpResult>& Results)
 		{
 			// Build array of materials from the original.
 			TArray<UMaterialInterface*> Materials;
-			IMaterialProvider* TargetMaterial = TargetMaterialInterface(OrigMeshIdx);
+			IMaterialProvider* TargetMaterial = Cast<IMaterialProvider>(Targets[OrigMeshIdx]);
 			for (int MaterialIdx = 0, NumMaterials = TargetMaterial->GetNumMaterials(); MaterialIdx < NumMaterials; MaterialIdx++)
 			{
 				Materials.Add(TargetMaterial->GetMaterial(MaterialIdx));
@@ -477,7 +447,7 @@ void UMirrorTool::GenerateAsset(const TArray<FDynamicMeshOpResult>& Results)
 			}
 
 			// Remove the original actor
-			TargetComponent->GetOwnerComponent()->DestroyComponent();
+			UE::ToolTarget::GetTargetComponent(Targets[OrigMeshIdx])->DestroyComponent();
 		}
 	}
 
diff --git a/Engine/Plugins/Experimental/MeshModelingToolsetExp/Source/MeshModelingToolsExp/Private/Physics/PhysicsInspectorTool.cpp b/Engine/Plugins/Experimental/MeshModelingToolsetExp/Source/MeshModelingToolsExp/Private/Physics/PhysicsInspectorTool.cpp
--- a/Engine/Plugins/Experimental/MeshModelingToolsetExp/Source/MeshModelingToolsExp/Private/Physics/PhysicsInspectorTool.cpp
+++ b/Engine/Plugins/Experimental/MeshModelingToolsetExp/Source/MeshModelingToolsExp/Private/Physics/PhysicsInspectorTool.cpp
@@ -4,6 +4,7 @@
 #include "InteractiveToolManager.h"
 #include "ToolBuilderUtil.h"
 #include "ToolSetupUtil.h"
+#include "ModelingToolTargetUtil.h"
 #include "Drawing/PreviewGeometryActor.h"
 
 #include "Physics/PhysicsDataCollection.h"
@@ -33,20 +34,10 @@ const FToolTargetTypeRequirements& UPhysicsInspectorToolBuilder::GetTargetRequir
 	return TypeRequirements;
 }
 
-bool UPhysicsInspectorToolBuilder::CanBuildTool(const FToolBuilderState& SceneState) const
-{
-	return SceneState.TargetManager->CountSelectedAndTargetable(SceneState, GetTargetRequirements()) > 0;
-}
-
 
-UInteractiveTool* UPhysicsInspectorToolBuilder::BuildTool(const FToolBuilderState& SceneState) const
+UMultiSelectionMeshEditingTool* UPhysicsInspectorToolBuilder::CreateNewTool(const FToolBuilderState& SceneState) const
 {
-	UPhysicsInspectorTool* NewTool = NewObject<UPhysicsInspectorTool>(SceneState.ToolManager);
-
-	TArray<TObjectPtr<UToolTarget>> Targets = SceneState.TargetManager->BuildAllSelectedTargetable(SceneState, GetTargetRequirements());
-	NewTool->SetTargets(MoveTemp(Targets));
-
-	return NewTool;
+	return NewObject<UPhysicsInspectorTool>(SceneState.ToolManager);
 }
 
 
@@ -63,8 +54,7 @@ void UPhysicsInspectorTool::Setup()
 
 	for (int32 ComponentIdx = 0; ComponentIdx < Targets.Num(); ComponentIdx++)
 	{
-		IPrimitiveComponentBackedTarget* TargetComponent = TargetComponentInterface(ComponentIdx);
-		const UStaticMeshComponent* Component = CastChecked<UStaticMeshComponent>(TargetComponent->GetOwnerComponent());
+		const UStaticMeshComponent* Component = CastChecked<UStaticMeshComponent>(UE::ToolTarget::GetTargetComponent(Targets[ComponentIdx]));
 		const UStaticMesh* StaticMesh = Component->GetStaticMesh();
 		if (ensure(StaticMesh && StaticMesh->GetBodySetup()))
 		{
@@ -76,10 +66,10 @@ void UPhysicsInspectorTool::Setup()
 			PhysicsInfos.Add(PhysicsData);
 
 			UPreviewGeometry* PreviewGeom = NewObject<UPreviewGeometry>(this);
-			FTransform TargetTransform = TargetComponent->GetWorldTransform();
+			FTransform TargetTransform = (FTransform) UE::ToolTarget::GetLocalToWorldTransform(Targets[ComponentIdx]);
 			PhysicsData->ExternalScale3D = TargetTransform.GetScale3D();
 			TargetTransform.SetScale3D(FVector::OneVector);
-			PreviewGeom->CreateInWorld(TargetComponent->GetOwnerActor()->GetWorld(), TargetTransform);
+			PreviewGeom->CreateInWorld(UE::ToolTarget::GetTargetActor(Targets[ComponentIdx])->GetWorld(), TargetTransform);
 			PreviewElements.Add(PreviewGeom);
 
 			InitializeGeometry(*PhysicsData, PreviewGeom);
diff --git a/Engine/Plugins/Experimental/MeshModelingToolsetExp/Source/MeshModelingToolsExp/Private/Physics/SetCollisionGeometryTool.cpp b/Engine/Plugins/Experimental/MeshModelingToolsetExp/Source/MeshModelingToolsExp/Private/Physics/SetCollisionGeometryTool.cpp
--- a/Engine/Plugins/Experimental/MeshModelingToolsetExp/Source/MeshModelingToolsExp/Private/Physics/SetCollisionGeometryTool.cpp
+++ b/Engine/Plugins/Experimental/MeshModelingToolsetExp/Source/MeshModelingToolsExp/Private/Physics/SetCollisionGeometryTool.cpp
@@ -51,12 +51,9 @@ bool USetCollisionGeometryToolBuilder::CanBuildTool(const FToolBuilderState& Sce
 }
 
 
-UInteractiveTool* USetCollisionGeometryToolBuilder::BuildTool(const FToolBuilderState& SceneState) const
+UMultiSelectionMeshEditingTool* USetCollisionGeometryToolBuilder::CreateNewTool(const FToolBuilderState& SceneState) const
 {
-	USetCollisionGeometryTool* NewTool = NewObject<USetCollisionGeometryTool>(SceneState.ToolManager);
-	TArray<TObjectPtr<UToolTarget>> Targets = SceneState.TargetManager->BuildAllSelectedTargetable(SceneState, GetTargetRequirements());
-	NewTool->SetTargets(MoveTemp(Targets));
-	return NewTool;
+	return NewObject<USetCollisionGeometryTool>(SceneState.ToolManager);
 }
 
 
diff --git a/Engine/Plugins/Experimental/MeshModelingToolsetExp/Source/MeshModelingToolsExp/Private/PlaneCutTool.cpp b/Engine/Plugins/Experimental/MeshModelingToolsetExp/Source/MeshModelingToolsExp/Private/PlaneCutTool.cpp
--- a/Engine/Plugins/Experimental/MeshModelingToolsetExp/Source/MeshModelingToolsExp/Private/PlaneCutTool.cpp
+++ b/Engine/Plugins/Experimental/MeshModelingToolsetExp/Source/MeshModelingToolsExp/Private/PlaneCutTool.cpp
@@ -43,31 +43,10 @@ using namespace UE::Geometry;
  * ToolBuilder
  */
 
-const FToolTargetTypeRequirements& UPlaneCutToolBuilder::GetTargetRequirements() const
-{
-	static FToolTargetTypeRequirements TypeRequirements({
-		UMaterialProvider::StaticClass(),
-		UMeshDescriptionCommitter::StaticClass(),
-		UMeshDescriptionProvider::StaticClass(),
-		UPrimitiveComponentBackedTarget::StaticClass()
-		});
-	return TypeRequirements;
-}
 
-bool UPlaneCutToolBuilder::CanBuildTool(const FToolBuilderState& SceneState) const
+UMultiSelectionMeshEditingTool* UPlaneCutToolBuilder::CreateNewTool(const FToolBuilderState& SceneState) const
 {
-	return SceneState.TargetManager->CountSelectedAndTargetable(SceneState, GetTargetRequirements()) > 0;
-}
-
-UInteractiveTool* UPlaneCutToolBuilder::BuildTool(const FToolBuilderState& SceneState) const
-{
-	UPlaneCutTool* NewTool = NewObject<UPlaneCutTool>(SceneState.ToolManager);
-
-	TArray<TObjectPtr<UToolTarget>> Targets = SceneState.TargetManager->BuildAllSelectedTargetable(SceneState, GetTargetRequirements());
-	NewTool->SetTargets(MoveTemp(Targets));
-	NewTool->SetWorld(SceneState.World);
-
-	return NewTool;
+	return NewObject<UPlaneCutTool>(SceneState.ToolManager);
 }
 
 
@@ -80,11 +59,6 @@ UPlaneCutTool::UPlaneCutTool()
 {
 }
 
-void UPlaneCutTool::SetWorld(UWorld* World)
-{
-	this->TargetWorld = World;
-}
-
 void UPlaneCutTool::Setup()
 {
 	UInteractiveTool::Setup();
@@ -92,7 +66,7 @@ void UPlaneCutTool::Setup()
 	// hide input StaticMeshComponents
 	for (int32 ComponentIdx = 0; ComponentIdx < Targets.Num(); ComponentIdx++)
 	{
-		TargetComponentInterface(ComponentIdx)->SetOwnerVisibility(false);
+		UE::ToolTarget::HideSourceObject(Targets[ComponentIdx]);
 	}
 
 	TArray<int32> MapToFirstOccurrences;
@@ -108,10 +82,9 @@ void UPlaneCutTool::Setup()
 	// Convert input mesh descriptions to dynamic mesh
 	for (int Idx = 0; Idx < Targets.Num(); Idx++)
 	{
-		IMeshDescriptionProvider* TargetMesh = TargetMeshProviderInterface(Idx);
 		FDynamicMesh3* OriginalDynamicMesh = new FDynamicMesh3;
 		FMeshDescriptionToDynamicMesh Converter;
-		Converter.Convert(TargetMesh->GetMeshDescription(), *OriginalDynamicMesh);
+		Converter.Convert(UE::ToolTarget::GetMeshDescription(Targets[Idx]), *OriginalDynamicMesh);
 		OriginalDynamicMesh->EnableAttributes();
 		TDynamicMeshScalarTriangleAttribute<int>* SubObjectIDs = new TDynamicMeshScalarTriangleAttribute<int>(OriginalDynamicMesh);
 		SubObjectIDs->Initialize(0);
@@ -146,7 +119,7 @@ void UPlaneCutTool::Setup()
 	for (int Idx = 0; Idx < Targets.Num(); Idx++)
 	{
 		FVector ComponentOrigin, ComponentExtents;
-		TargetComponentInterface(Idx)->GetOwnerActor()->GetActorBounds(false, ComponentOrigin, ComponentExtents);
+		UE::ToolTarget::GetTargetActor(Targets[Idx])->GetActorBounds(false, ComponentOrigin, ComponentExtents);
 		CombinedBounds += FBox::BuildAABB(ComponentOrigin, ComponentExtents);
 	}
 
@@ -160,7 +133,7 @@ void UPlaneCutTool::Setup()
 		});
 	for (int Idx = 0; Idx < Targets.Num(); Idx++)
 	{
-		PlaneMechanic->SetPlaneCtrlClickBehaviorTarget->InvisibleComponentsToHitTest.Add(TargetComponentInterface(Idx)->GetOwnerComponent());
+		PlaneMechanic->SetPlaneCtrlClickBehaviorTarget->InvisibleComponentsToHitTest.Add(UE::ToolTarget::GetTargetComponent(Targets[Idx]));
 	}
 
 	InvalidatePreviews();
@@ -213,15 +186,14 @@ void UPlaneCutTool::SetupPreviews()
 		ToolSetupUtil::ApplyRenderingConfigurationToPreview(Preview->PreviewMesh, nullptr);
 		Preview->PreviewMesh->SetTangentsMode(EDynamicMeshComponentTangentsMode::AutoCalculated);
 
-		FComponentMaterialSet MaterialSet;
-		TargetMaterialInterface(PreviewIdx)->GetMaterialSet(MaterialSet);
+		const FComponentMaterialSet MaterialSet = UE::ToolTarget::GetMaterialSet(Targets[PreviewIdx]);
 		Preview->ConfigureMaterials(MaterialSet.Materials,
 			ToolSetupUtil::GetDefaultWorkingMaterial(GetToolManager())
 		);
 
 		// set initial preview to un-processed mesh, so stuff doesn't just disappear if the first cut takes a while
 		Preview->PreviewMesh->UpdatePreview(MeshesToCut[PreviewIdx]->GetMesh().Get());
-		Preview->PreviewMesh->SetTransform(TargetComponentInterface(PreviewIdx)->GetWorldTransform());
+		Preview->PreviewMesh->SetTransform((FTransform)UE::ToolTarget::GetLocalToWorldTransform(Targets[PreviewIdx]));
 		Preview->SetVisibility(BasicProperties->bShowPreview);
 	}
 }
@@ -279,7 +251,7 @@ void UPlaneCutTool::Shutdown(EToolShutdownType ShutdownType)
 	// Restore (unhide) the source meshes
 	for (int Idx = 0; Idx < Targets.Num(); Idx++)
 	{
-		TargetComponentInterface(Idx)->SetOwnerVisibility(true);
+		UE::ToolTarget::ShowSourceObject(Targets[Idx]);
 	}
 
 	TArray<FDynamicMeshOpResult> Results;
@@ -299,7 +271,7 @@ TUniquePtr<FDynamicMeshOperator> UPlaneCutOperatorFactory::MakeNewOperator()
 	CutOp->bFillCutHole = CutTool->BasicProperties->bFillCutHole;
 	CutOp->bFillSpans = CutTool->BasicProperties->bFillSpans;
 
-	FTransform LocalToWorld = CutTool->TargetComponentInterface(ComponentIndex)->GetWorldTransform();
+	FTransform LocalToWorld = (FTransform) UE::ToolTarget::GetLocalToWorldTransform(CutTool->Targets[ComponentIndex]);
 	CutOp->SetTransform(LocalToWorld);
 	// for all plane computation, change LocalToWorld to not have any zero scale dims
 	FVector LocalToWorldScale = LocalToWorld.GetScale3D();
@@ -382,7 +354,7 @@ void UPlaneCutTool::OnPropertyModified(UObject* PropertySet, FProperty* Property
 	{
 		for (int Idx = 0; Idx < Targets.Num(); Idx++)
 		{
-			TargetComponentInterface(Idx)->SetOwnerVisibility(!BasicProperties->bShowPreview);
+			UE::ToolTarget::SetSourceObjectVisible(Targets[Idx], !BasicProperties->bShowPreview);
 		}
 		for (UMeshOpPreviewWithBackgroundCompute* Preview : Previews)
 		{
@@ -443,8 +415,7 @@ void UPlaneCutTool::GenerateAsset(const TArray<FDynamicMeshOpResult>& Results)
 	NewSelection.ModificationType = ESelectedObjectsModificationType::Replace;
 	for (int OrigMeshIdx = 0; OrigMeshIdx < NumSourceMeshes; OrigMeshIdx++)
 	{
-		IPrimitiveComponentBackedTarget* TargetComponent = TargetComponentInterface(OrigMeshIdx);
-		NewSelection.Actors.Add(TargetComponent->GetOwnerActor());
+		NewSelection.Actors.Add(UE::ToolTarget::GetTargetActor(Targets[OrigMeshIdx]));
 	}
 
 	// check if we entirely cut away any meshes
@@ -475,7 +446,7 @@ void UPlaneCutTool::GenerateAsset(const TArray<FDynamicMeshOpResult>& Results)
 		{
 			if (bWantDestroy)
 			{
-				TargetComponentInterface(OrigMeshIdx)->GetOwnerComponent()->DestroyComponent();
+				UE::ToolTarget::GetTargetComponent(Targets[OrigMeshIdx])->DestroyComponent();
 			}
 			continue;
 		}
@@ -503,7 +474,7 @@ void UPlaneCutTool::GenerateAsset(const TArray<FDynamicMeshOpResult>& Results)
 			}
 		}
 
-		TargetMeshCommitterInterface(OrigMeshIdx)->CommitMeshDescription([&UseMesh](const IMeshDescriptionCommitter::FCommitterParams& CommitParams)
+		Cast<IMeshDescriptionCommitter>(Targets[OrigMeshIdx])->CommitMeshDescription([&UseMesh](const IMeshDescriptionCommitter::FCommitterParams& CommitParams)
 		{
 			FDynamicMeshToMeshDescription Converter;
 			Converter.Convert(UseMesh, *CommitParams.MeshDescriptionOut);
@@ -521,10 +492,8 @@ void UPlaneCutTool::GenerateAsset(const TArray<FDynamicMeshOpResult>& Results)
 			}
 
 			// get materials for both the component and the asset
-			IMaterialProvider* TargetMaterial = TargetMaterialInterface(OrigMeshIdx);
-			FComponentMaterialSet ComponentMaterialSet, AssetMaterialSet;
-			TargetMaterial->GetMaterialSet(ComponentMaterialSet, false);
-			TargetMaterial->GetMaterialSet(AssetMaterialSet, true /*prefer asset materials*/);
+			const FComponentMaterialSet ComponentMaterialSet = UE::ToolTarget::GetMaterialSet(Targets[OrigMeshIdx], false);
+			const FComponentMaterialSet AssetMaterialSet = UE::ToolTarget::GetMaterialSet(Targets[OrigMeshIdx], true /*prefer asset materials*/);
 
 			// add all the additional meshes
 			for (int AddMeshIdx = 1; AddMeshIdx < SplitMeshes.Num(); AddMeshIdx++)
diff --git a/Engine/Plugins/Experimental/MeshModelingToolsetExp/Source/MeshModelingToolsExp/Private/ProjectToTargetTool.cpp b/Engine/Plugins/Experimental/MeshModelingToolsetExp/Source/MeshModelingToolsExp/Private/ProjectToTargetTool.cpp
--- a/Engine/Plugins/Experimental/MeshModelingToolsetExp/Source/MeshModelingToolsExp/Private/ProjectToTargetTool.cpp
+++ b/Engine/Plugins/Experimental/MeshModelingToolsetExp/Source/MeshModelingToolsExp/Private/ProjectToTargetTool.cpp
@@ -3,6 +3,7 @@
 #include "ProjectToTargetTool.h"
 #include "MeshDescriptionToDynamicMesh.h"
 #include "InteractiveToolManager.h"
+#include "ModelingToolTargetUtil.h"
 
 #include "TargetInterfaces/MaterialProvider.h"
 #include "TargetInterfaces/MeshDescriptionCommitter.h"
@@ -14,31 +15,14 @@ using namespace UE::Geometry;
 
 #define LOCTEXT_NAMESPACE "UProjectToTargetTool"
 
-const FToolTargetTypeRequirements& UProjectToTargetToolBuilder::GetTargetRequirements() const
-{
-	static FToolTargetTypeRequirements TypeRequirements({
-		UMeshDescriptionCommitter::StaticClass(),
-		UMeshDescriptionProvider::StaticClass(),
-		UPrimitiveComponentBackedTarget::StaticClass(),
-		UMaterialProvider::StaticClass()
-		});
-	return TypeRequirements;
-}
-
 bool UProjectToTargetToolBuilder::CanBuildTool(const FToolBuilderState& SceneState) const
 {
 	return SceneState.TargetManager->CountSelectedAndTargetable(SceneState, GetTargetRequirements()) == 2;
 }
 
-UInteractiveTool* UProjectToTargetToolBuilder::BuildTool(const FToolBuilderState& SceneState) const
+UMultiSelectionMeshEditingTool* UProjectToTargetToolBuilder::CreateNewTool(const FToolBuilderState& SceneState) const
 {
-	UProjectToTargetTool* NewTool = NewObject<UProjectToTargetTool>(SceneState.ToolManager);
-
-	TArray<TObjectPtr<UToolTarget>> Targets = SceneState.TargetManager->BuildAllSelectedTargetable(SceneState, GetTargetRequirements());
-	NewTool->SetTargets(MoveTemp(Targets));
-	NewTool->SetWorld(SceneState.World);
-
-	return NewTool;
+	return NewObject<UProjectToTargetTool>(SceneState.ToolManager);
 }
 
 void UProjectToTargetTool::Setup()
@@ -47,7 +31,7 @@ void UProjectToTargetTool::Setup()
 	FMeshDescriptionToDynamicMesh ProjectionConverter;
 	check(Targets.Num() == 2);
 	ProjectionTarget = MakeUnique<FDynamicMesh3>();
-	ProjectionConverter.Convert(TargetMeshProviderInterface(1)->GetMeshDescription(), *ProjectionTarget);
+	ProjectionConverter.Convert(UE::ToolTarget::GetMeshDescription(Targets[1]), *ProjectionTarget);
 	ProjectionTargetSpatial = MakeUnique<FDynamicMeshAABBTree3>(ProjectionTarget.Get(), true);
 
 	// Now setup parent RemeshMeshTool class
@@ -78,8 +62,8 @@ TUniquePtr<FDynamicMeshOperator> UProjectToTargetTool::MakeNewOperator()
 	RemeshOp->ProjectionTargetSpatial = ProjectionTargetSpatial.Get();
 
 	using FTransform3d = UE::Geometry::FTransform3d;
-	RemeshOp->ToolMeshLocalToWorld = FTransform3d(Cast<IPrimitiveComponentBackedTarget>(Targets[0])->GetWorldTransform());
-	RemeshOp->TargetMeshLocalToWorld = FTransform3d(Cast<IPrimitiveComponentBackedTarget>(Targets[1])->GetWorldTransform());
+	RemeshOp->ToolMeshLocalToWorld = UE::ToolTarget::GetLocalToWorldTransform(Targets[0]);
+	RemeshOp->TargetMeshLocalToWorld = UE::ToolTarget::GetLocalToWorldTransform(Targets[1]);
 	RemeshOp->bUseWorldSpace = ProjectProperties->bWorldSpace;
 	RemeshOp->bParallel = ProjectProperties->bParallel;
 
diff --git a/Engine/Plugins/Experimental/MeshModelingToolsetExp/Source/MeshModelingToolsExp/Private/RemeshMeshTool.cpp b/Engine/Plugins/Experimental/MeshModelingToolsetExp/Source/MeshModelingToolsExp/Private/RemeshMeshTool.cpp
--- a/Engine/Plugins/Experimental/MeshModelingToolsetExp/Source/MeshModelingToolsExp/Private/RemeshMeshTool.cpp
+++ b/Engine/Plugins/Experimental/MeshModelingToolsetExp/Source/MeshModelingToolsExp/Private/RemeshMeshTool.cpp
@@ -32,31 +32,15 @@ using namespace UE::Geometry;
 /*
  * ToolBuilder
  */
-const FToolTargetTypeRequirements& URemeshMeshToolBuilder::GetTargetRequirements() const
-{
-	static FToolTargetTypeRequirements TypeRequirements({
-		UMaterialProvider::StaticClass(),
-		UMeshDescriptionCommitter::StaticClass(),
-		UMeshDescriptionProvider::StaticClass(),
-		UPrimitiveComponentBackedTarget::StaticClass()
-		});
-	return TypeRequirements;
-}
 
 bool URemeshMeshToolBuilder::CanBuildTool(const FToolBuilderState& SceneState) const
 {
 	return SceneState.TargetManager->CountSelectedAndTargetable(SceneState, GetTargetRequirements()) == 1;
 }
 
-UInteractiveTool* URemeshMeshToolBuilder::BuildTool(const FToolBuilderState& SceneState) const
+UMultiSelectionMeshEditingTool* URemeshMeshToolBuilder::CreateNewTool(const FToolBuilderState& SceneState) const
 {
-	URemeshMeshTool* NewTool = NewObject<URemeshMeshTool>(SceneState.ToolManager);
-
-	TArray<TObjectPtr<UToolTarget>> Targets = SceneState.TargetManager->BuildAllSelectedTargetable(SceneState, GetTargetRequirements());
-	NewTool->SetTargets(MoveTemp(Targets));
-	NewTool->SetWorld(SceneState.World);
-
-	return NewTool;
+	return NewObject<URemeshMeshTool>(SceneState.ToolManager);
 }
 
 /*
@@ -91,10 +75,6 @@ URemeshMeshTool::URemeshMeshTool(const FObjectInitializer&)
 	BasicProperties->ClearFlags(RF_Transactional);
 }
 
-void URemeshMeshTool::SetWorld(UWorld* World)
-{
-	this->TargetWorld = World;
-}
 
 void URemeshMeshTool::Setup()
 {
@@ -106,17 +86,14 @@ void URemeshMeshTool::Setup()
 
 	check(Targets.Num() > 0);
 	check(Targets[0]);
-	IPrimitiveComponentBackedTarget* TargetComponent = TargetComponentInterface(0);
-	IMeshDescriptionProvider* TargetMeshProvider = TargetMeshProviderInterface(0);
-
 
 	// hide component and create + show preview
-	TargetComponent->SetOwnerVisibility(false);
+	UE::ToolTarget::HideSourceObject(Targets[0]);
 	Preview = NewObject<UMeshOpPreviewWithBackgroundCompute>(this);
 	Preview->Setup(this->TargetWorld, this);
 	ToolSetupUtil::ApplyRenderingConfigurationToPreview(Preview->PreviewMesh, Targets[0]);
 
-	FComponentMaterialSet MaterialSet = UE::ToolTarget::GetMaterialSet(Targets[0]);
+	const FComponentMaterialSet MaterialSet = UE::ToolTarget::GetMaterialSet(Targets[0]);
 	Preview->ConfigureMaterials( MaterialSet.Materials,
 								 ToolSetupUtil::GetDefaultWorkingMaterial(GetToolManager())
 	);
@@ -125,9 +102,9 @@ void URemeshMeshTool::Setup()
 
 	OriginalMesh = MakeShared<FDynamicMesh3, ESPMode::ThreadSafe>();
 	FMeshDescriptionToDynamicMesh Converter;
-	Converter.Convert(TargetMeshProvider->GetMeshDescription(), *OriginalMesh);
+	Converter.Convert(UE::ToolTarget::GetMeshDescription(Targets[0]), *OriginalMesh);
 
-	Preview->PreviewMesh->SetTransform(TargetComponent->GetWorldTransform());
+	Preview->PreviewMesh->SetTransform((FTransform) UE::ToolTarget::GetLocalToWorldTransform(Targets[0]));
 	Preview->PreviewMesh->SetTangentsMode(EDynamicMeshComponentTangentsMode::AutoCalculated);
 	Preview->PreviewMesh->UpdatePreview(OriginalMesh.Get());
 
diff --git a/Engine/Plugins/Experimental/MeshModelingToolsetExp/Source/MeshModelingToolsExp/Private/RemoveOccludedTrianglesTool.cpp b/Engine/Plugins/Experimental/MeshModelingToolsetExp/Source/MeshModelingToolsExp/Private/RemoveOccludedTrianglesTool.cpp
--- a/Engine/Plugins/Experimental/MeshModelingToolsetExp/Source/MeshModelingToolsExp/Private/RemoveOccludedTrianglesTool.cpp
+++ b/Engine/Plugins/Experimental/MeshModelingToolsetExp/Source/MeshModelingToolsExp/Private/RemoveOccludedTrianglesTool.cpp
@@ -5,6 +5,7 @@
 #include "ToolBuilderUtil.h"
 
 #include "ToolSetupUtil.h"
+#include "ModelingToolTargetUtil.h"
 
 #include "DynamicMesh/DynamicMesh3.h"
 #include "Polygroups/PolygroupUtil.h"
@@ -38,31 +39,9 @@ using namespace UE::Geometry;
  */
 
 
-const FToolTargetTypeRequirements& URemoveOccludedTrianglesToolBuilder::GetTargetRequirements() const
+UMultiSelectionMeshEditingTool* URemoveOccludedTrianglesToolBuilder::CreateNewTool(const FToolBuilderState& SceneState) const
 {
-	static FToolTargetTypeRequirements TypeRequirements({
-		UMaterialProvider::StaticClass(),
-		UMeshDescriptionCommitter::StaticClass(),
-		UMeshDescriptionProvider::StaticClass(),
-		UPrimitiveComponentBackedTarget::StaticClass()
-		});
-	return TypeRequirements;
-}
-
-bool URemoveOccludedTrianglesToolBuilder::CanBuildTool(const FToolBuilderState& SceneState) const
-{
-	return SceneState.TargetManager->CountSelectedAndTargetable(SceneState, GetTargetRequirements()) > 0;
-}
-
-UInteractiveTool* URemoveOccludedTrianglesToolBuilder::BuildTool(const FToolBuilderState& SceneState) const
-{
-	URemoveOccludedTrianglesTool* NewTool = NewObject<URemoveOccludedTrianglesTool>(SceneState.ToolManager);
-
-	TArray<TObjectPtr<UToolTarget>> Targets = SceneState.TargetManager->BuildAllSelectedTargetable(SceneState, GetTargetRequirements());
-	NewTool->SetTargets(MoveTemp(Targets));
-	NewTool->SetWorld(SceneState.World);
-
-	return NewTool;
+	return NewObject<URemoveOccludedTrianglesTool>(SceneState.ToolManager);
 }
 
 
@@ -87,11 +66,6 @@ URemoveOccludedTrianglesTool::URemoveOccludedTrianglesTool()
 	SetToolDisplayName(LOCTEXT("ProjectToTargetToolName", "Jacket"));
 }
 
-void URemoveOccludedTrianglesTool::SetWorld(UWorld* World)
-{
-	this->TargetWorld = World;
-}
-
 void URemoveOccludedTrianglesTool::Setup()
 {
 	UInteractiveTool::Setup();
@@ -99,7 +73,7 @@ void URemoveOccludedTrianglesTool::Setup()
 	// hide input StaticMeshComponent
 	for (int Idx = 0; Idx < Targets.Num(); Idx++)
 	{
-		TargetComponentInterface(Idx)->SetOwnerVisibility(false);
+		UE::ToolTarget::HideSourceObject(Targets[Idx]);
 	}
 
 
@@ -243,7 +217,7 @@ void URemoveOccludedTrianglesTool::SetupPreviews()
 		};
 
 		bool bHasConverted = OriginalDynamicMeshes[PreviewIdx].IsValid();
-
+		
 		if (!bHasConverted)
 		{
 			URemoveOccludedTrianglesOperatorFactory *OpFactory = NewObject<URemoveOccludedTrianglesOperatorFactory>();
@@ -251,26 +225,25 @@ void URemoveOccludedTrianglesTool::SetupPreviews()
 			OpFactory->PreviewIdx = PreviewIdx;
 			OriginalDynamicMeshes[PreviewIdx] = MakeShared<FDynamicMesh3, ESPMode::ThreadSafe>();
 			FMeshDescriptionToDynamicMesh Converter;
-			Converter.Convert(TargetMeshProviderInterface(TargetIdx)->GetMeshDescription(), *OriginalDynamicMeshes[PreviewIdx]);
+			Converter.Convert(UE::ToolTarget::GetMeshDescription(Targets[TargetIdx]), *OriginalDynamicMeshes[PreviewIdx]);
 
 			UMeshOpPreviewWithBackgroundCompute* Preview = Previews.Add_GetRef(NewObject<UMeshOpPreviewWithBackgroundCompute>(OpFactory, "Preview"));
 			Preview->Setup(this->TargetWorld, OpFactory);
 			ToolSetupUtil::ApplyRenderingConfigurationToPreview(Preview->PreviewMesh, nullptr);
 			Preview->PreviewMesh->SetTangentsMode(EDynamicMeshComponentTangentsMode::AutoCalculated);
 
-			FComponentMaterialSet MaterialSet;
-			TargetMaterialInterface(TargetIdx)->GetMaterialSet(MaterialSet);
+			const FComponentMaterialSet MaterialSet = UE::ToolTarget::GetMaterialSet(Targets[TargetIdx]);
 			Preview->ConfigureMaterials(MaterialSet.Materials,
 				ToolSetupUtil::GetDefaultWorkingMaterial(GetToolManager())
 			);
 
-			Preview->PreviewMesh->SetTransform(TargetComponentInterface(TargetIdx)->GetWorldTransform());
+			Preview->PreviewMesh->SetTransform((FTransform) UE::ToolTarget::GetLocalToWorldTransform(Targets[TargetIdx]));
 			Preview->PreviewMesh->UpdatePreview(OriginalDynamicMeshes[PreviewIdx].Get());
 			Preview->SetVisibility(true);
 
 			OccluderTrees[TargetIdx] = MakeShared<FDynamicMeshAABBTree3, ESPMode::ThreadSafe>(OriginalDynamicMeshes[PreviewIdx].Get());
 			OccluderWindings[TargetIdx] = MakeShared<TFastWindingTree<FDynamicMesh3>, ESPMode::ThreadSafe>(OccluderTrees[TargetIdx].Get());
-			OccluderTransforms[TargetIdx] = (UE::Geometry::FTransform3d)TargetComponentInterface(TargetIdx)->GetWorldTransform();
+			OccluderTransforms[TargetIdx] = UE::ToolTarget::GetLocalToWorldTransform(Targets[TargetIdx]);
 
 			// configure secondary render material
 			Preview->SecondaryMaterial = OccludedMaterial;
@@ -290,21 +263,20 @@ void URemoveOccludedTrianglesTool::SetupPreviews()
 			// already did the conversion for a full UMeshOpPreviewWithBackgroundCompute -- just make a light version of that and hook it up to copy the other's work
 			int CopyIdx = PreviewCopies.Num();
 			UPreviewMesh* PreviewMesh = PreviewCopies.Add_GetRef(NewObject<UPreviewMesh>(this));
-			PreviewMesh->CreateInWorld(this->TargetWorld, TargetComponentInterface(TargetIdx)->GetWorldTransform());
+			PreviewMesh->CreateInWorld(this->TargetWorld, (FTransform) UE::ToolTarget::GetLocalToWorldTransform(Targets[TargetIdx]));
 
 			PreviewToCopyIdx[PreviewIdx].Add(CopyIdx);
 
 			PreviewMesh->UpdatePreview(OriginalDynamicMeshes[PreviewIdx].Get());
 
-			FComponentMaterialSet MaterialSet;
-			TargetMaterialInterface(TargetIdx)->GetMaterialSet(MaterialSet);
+			const FComponentMaterialSet MaterialSet = UE::ToolTarget::GetMaterialSet(Targets[TargetIdx]);
 			PreviewMesh->SetMaterials(MaterialSet.Materials);
 			
 			PreviewMesh->SetVisible(true);
 
 			OccluderTrees[TargetIdx] = OccluderTrees[PreviewToTargetIdx[PreviewIdx]];
 			OccluderWindings[TargetIdx] = OccluderWindings[PreviewToTargetIdx[PreviewIdx]];
-			OccluderTransforms[TargetIdx] = (UE::Geometry::FTransform3d)TargetComponentInterface(TargetIdx)->GetWorldTransform();
+			OccluderTransforms[TargetIdx] = UE::ToolTarget::GetLocalToWorldTransform(Targets[TargetIdx]);
 
 			PreviewMesh->SetSecondaryRenderMaterial(OccludedMaterial);
 			PreviewMesh->EnableSecondaryTriangleBuffers(MoveTemp(IsOccludedGroupFn));
@@ -329,7 +301,7 @@ void URemoveOccludedTrianglesTool::Shutdown(EToolShutdownType ShutdownType)
 	// Restore (unhide) the source meshes
 	for (int ComponentIdx = 0; ComponentIdx < Targets.Num(); ComponentIdx++)
 	{
-		TargetComponentInterface(ComponentIdx)->SetOwnerVisibility(true);
+		UE::ToolTarget::ShowSourceObject(Targets[ComponentIdx]);
 	}
 
 	// clear all the preview copies
@@ -389,7 +361,7 @@ TUniquePtr<FDynamicMeshOperator> URemoveOccludedTrianglesOperatorFactory::MakeNe
 	Op->WindingIsoValue = Tool->BasicProperties->WindingIsoValue;
 
 	int ComponentIndex = Tool->PreviewToTargetIdx[PreviewIdx];
-	FTransform LocalToWorld = Tool->TargetComponentInterface(ComponentIndex)->GetWorldTransform();
+	FTransform LocalToWorld = (FTransform) UE::ToolTarget::GetLocalToWorldTransform(Tool->Targets[ComponentIndex]);
 	Op->OriginalMesh = Tool->OriginalDynamicMeshes[PreviewIdx];
 
 	if (Tool->BasicProperties->bOnlySelfOcclude)
@@ -526,14 +498,14 @@ void URemoveOccludedTrianglesTool::GenerateAsset(const TArray<FDynamicMeshOpResu
 				{
 					if (TargetToPreviewIdx[TargetIdx] == PreviewIdx)
 					{
-						TargetComponentInterface(TargetIdx)->GetOwnerComponent()->DestroyComponent();
+						UE::ToolTarget::GetTargetComponent(Targets[TargetIdx])->DestroyComponent();
 					}
 				}
 			}
 			continue;
 		}
 
-		TargetMeshCommitterInterface(ComponentIdx)->CommitMeshDescription([&Results, &PreviewIdx, this](const IMeshDescriptionCommitter::FCommitterParams& CommitParams)
+		Cast<IMeshDescriptionCommitter>(Targets[ComponentIdx])->CommitMeshDescription([&Results, &PreviewIdx, this](const IMeshDescriptionCommitter::FCommitterParams& CommitParams)
 		{
 			FDynamicMeshToMeshDescription Converter;
 			Converter.Convert(Results[PreviewIdx].Mesh.Get(), *CommitParams.MeshDescriptionOut);
diff --git a/Engine/Plugins/Experimental/MeshModelingToolsetExp/Source/MeshModelingToolsExp/Private/SelfUnionMeshesTool.cpp b/Engine/Plugins/Experimental/MeshModelingToolsetExp/Source/MeshModelingToolsExp/Private/SelfUnionMeshesTool.cpp
--- a/Engine/Plugins/Experimental/MeshModelingToolsetExp/Source/MeshModelingToolsExp/Private/SelfUnionMeshesTool.cpp
+++ b/Engine/Plugins/Experimental/MeshModelingToolsetExp/Source/MeshModelingToolsExp/Private/SelfUnionMeshesTool.cpp
@@ -5,6 +5,7 @@
 #include "InteractiveToolManager.h"
 #include "ToolSetupUtil.h"
 #include "BaseGizmos/TransformProxy.h"
+#include "ModelingToolTargetUtil.h"
 
 #include "DynamicMesh/DynamicMesh3.h"
 #include "DynamicMeshEditor.h"
@@ -13,6 +14,8 @@
 #include "MeshDescriptionToDynamicMesh.h"
 #include "DynamicMeshToMeshDescription.h"
 
+#include "TargetInterfaces/MeshDescriptionProvider.h"
+
 using namespace UE::Geometry;
 
 #define LOCTEXT_NAMESPACE "USelfUnionMeshesTool"
@@ -56,8 +59,7 @@ void USelfUnionMeshesTool::ConvertInputsAndSetPreviewMaterials(bool bSetPreviewM
 	{
 		for (int ComponentIdx = 0; ComponentIdx < Targets.Num(); ComponentIdx++)
 		{
-			FComponentMaterialSet ComponentMaterialSet;
-			TargetMaterialInterface(ComponentIdx)->GetMaterialSet(ComponentMaterialSet);
+			const FComponentMaterialSet ComponentMaterialSet = UE::ToolTarget::GetMaterialSet(Targets[ComponentIdx]);
 			for (UMaterialInterface* Mat : ComponentMaterialSet.Materials)
 			{
 				int* FoundMatIdx = KnownMaterials.Find(Mat);
@@ -77,14 +79,14 @@ void USelfUnionMeshesTool::ConvertInputsAndSetPreviewMaterials(bool bSetPreviewM
 	}
 	else
 	{
-		TargetMaterialInterface(0)->GetMaterialSet(AllMaterialSet);
+		AllMaterialSet = UE::ToolTarget::GetMaterialSet(Targets[0]);
 		for (int MatIdx = 0; MatIdx < AllMaterialSet.Materials.Num(); MatIdx++)
 		{
 			MaterialRemap[0].Add(MatIdx);
 		}
 		for (int ComponentIdx = 1; ComponentIdx < Targets.Num(); ComponentIdx++)
 		{
-			MaterialRemap[ComponentIdx].Init(0, TargetMaterialInterface(ComponentIdx)->GetNumMaterials());
+			MaterialRemap[ComponentIdx].Init(0, Cast<IMaterialProvider>(Targets[ComponentIdx])->GetNumMaterials());
 		}
 	}
 
@@ -100,7 +102,7 @@ void USelfUnionMeshesTool::ConvertInputsAndSetPreviewMaterials(bool bSetPreviewM
 	{
 		FDynamicMesh3 ComponentMesh;
 		FMeshDescriptionToDynamicMesh Converter;
-		Converter.Convert(TargetMeshProviderInterface(ComponentIdx)->GetMeshDescription(), ComponentMesh);
+		Converter.Convert(UE::ToolTarget::GetMeshDescription(Targets[ComponentIdx]), ComponentMesh);
 		bNeedColorAttr = bNeedColorAttr || (ComponentMesh.Attributes()->HasPrimaryColors());
 		// ensure materials and attributes are always enabled
 		ComponentMesh.EnableAttributes();
diff --git a/Engine/Plugins/Experimental/MeshModelingToolsetExp/Source/MeshModelingToolsExp/Private/SkinWeightsBindingTool.cpp b/Engine/Plugins/Experimental/MeshModelingToolsetExp/Source/MeshModelingToolsExp/Private/SkinWeightsBindingTool.cpp
--- a/Engine/Plugins/Experimental/MeshModelingToolsetExp/Source/MeshModelingToolsExp/Private/SkinWeightsBindingTool.cpp
+++ b/Engine/Plugins/Experimental/MeshModelingToolsetExp/Source/MeshModelingToolsExp/Private/SkinWeightsBindingTool.cpp
@@ -11,11 +11,16 @@
 #include "SkeletalMeshAttributes.h"
 #include "ToolSetupUtil.h"
 #include "ToolTargetManager.h"
+#include "ModelingToolTargetUtil.h"
 #include "Components/SkeletalMeshComponent.h"
 #include "Engine/SkeletalMesh.h"
 #include "Spatial/FastWinding.h"
 #include "Spatial/MeshWindingNumberGrid.h"
 
+#include "TargetInterfaces/PrimitiveComponentBackedTarget.h"
+#include "TargetInterfaces/MeshDescriptionProvider.h"
+#include "TargetInterfaces/MeshDescriptionCommitter.h"
+
 // #pragma optimize( "", off )
 
 #define LOCTEXT_NAMESPACE "USkinWeightsBindingTool"
@@ -724,15 +729,9 @@ bool USkinWeightsBindingToolBuilder::CanBuildTool(const FToolBuilderState& Scene
 }
 
 
-UInteractiveTool* USkinWeightsBindingToolBuilder::BuildTool(const FToolBuilderState& SceneState) const
+UMultiSelectionMeshEditingTool* USkinWeightsBindingToolBuilder::CreateNewTool(const FToolBuilderState& SceneState) const
 {
-	USkinWeightsBindingTool* NewTool = NewObject<USkinWeightsBindingTool>(SceneState.ToolManager);
-
-	TArray<TObjectPtr<UToolTarget>> Targets = SceneState.TargetManager->BuildAllSelectedTargetable(SceneState, FToolTargetTypeRequirements());
-	NewTool->SetTargets(MoveTemp(Targets));
-	NewTool->SetWorld(SceneState.World);
-
-	return NewTool;
+	return NewObject<USkinWeightsBindingTool>(SceneState.ToolManager);
 }
 
 
@@ -756,12 +755,6 @@ USkinWeightsBindingTool::~USkinWeightsBindingTool()
 }
 
 
-void USkinWeightsBindingTool::SetWorld(UWorld* InWorld)
-{
-	TargetWorld = InWorld;
-}
-
-
 void USkinWeightsBindingTool::Setup()
 {
 	Super::Setup();
@@ -776,11 +769,7 @@ void USkinWeightsBindingTool::Setup()
 		return;
 	}
 	
-	IPrimitiveComponentBackedTarget* TargetComponent = TargetComponentInterface(0);
-	IMeshDescriptionProvider* TargetMeshProvider = TargetMeshProviderInterface(0);
-	IMaterialProvider* TargetMaterial = TargetMaterialInterface(0);
-
-	const USkeletalMeshComponent* SkelMeshComponent = Cast<USkeletalMeshComponent>(TargetComponent->GetOwnerComponent());
+	const USkeletalMeshComponent* SkelMeshComponent = Cast<USkeletalMeshComponent>(UE::ToolTarget::GetTargetComponent(Targets[0]));
 	
 	if (SkelMeshComponent && SkelMeshComponent->SkeletalMesh)
 	{
@@ -817,7 +806,7 @@ void USkinWeightsBindingTool::Setup()
 		}
 	}
 	
-	TargetComponent->SetOwnerVisibility(false);
+	UE::ToolTarget::HideSourceObject(Targets[0]);
 	Preview = NewObject<UMeshOpPreviewWithBackgroundCompute>(this, "Preview");
 	Preview->Setup(this->TargetWorld, this);
 	Preview->SetIsMeshTopologyConstant(true, EMeshRenderAttributeFlags::VertexColors);
@@ -826,8 +815,7 @@ void USkinWeightsBindingTool::Setup()
 		UpdateVisualization();
 	});
 	
-	FComponentMaterialSet MaterialSet;
-	TargetMaterial->GetMaterialSet(MaterialSet);
+	FComponentMaterialSet MaterialSet = UE::ToolTarget::GetMaterialSet(Targets[0]);
 
 	UMaterialInterface* VtxColorMaterial = GetToolManager()->GetContextQueriesAPI()->GetStandardMaterial(EStandardToolContextMaterials::VertexColorMaterial);
 	if (VtxColorMaterial != nullptr)
@@ -847,7 +835,7 @@ void USkinWeightsBindingTool::Setup()
 
 	OriginalMesh = MakeShared<FDynamicMesh3, ESPMode::ThreadSafe>();
 	FMeshDescriptionToDynamicMesh Converter;
-	Converter.Convert(TargetMeshProvider->GetMeshDescription(), *OriginalMesh);
+	Converter.Convert(UE::ToolTarget::GetMeshDescription(Targets[0]), *OriginalMesh);
 
 	// Enable or override vertex colors on the original mesh.
 	OriginalMesh->EnableAttributes();
@@ -856,7 +844,7 @@ void USkinWeightsBindingTool::Setup()
 	// Create an overlay that has no split elements, init with zero value.
 	OriginalMesh->Attributes()->PrimaryColors()->CreateFromPredicate([](int ParentVID, int TriIDA, int TriIDB){return true;}, 0.f);
 
-	Preview->PreviewMesh->SetTransform(TargetComponent->GetWorldTransform());
+	Preview->PreviewMesh->SetTransform((FTransform) UE::ToolTarget::GetLocalToWorldTransform(Targets[0]));
 	Preview->PreviewMesh->SetTangentsMode(EDynamicMeshComponentTangentsMode::AutoCalculated);
 	Preview->PreviewMesh->SetShadowsEnabled(false);
 	Preview->PreviewMesh->UpdatePreview(OriginalMesh.Get());
@@ -881,7 +869,7 @@ void USkinWeightsBindingTool::Shutdown(EToolShutdownType ShutdownType)
 {
 	Properties->SaveProperties(this);
 
-	TargetComponentInterface(0)->SetOwnerVisibility(true);
+	UE::ToolTarget::ShowSourceObject(Targets[0]);
 	
 	FDynamicMeshOpResult Result = Preview->Shutdown();
 	if (ShutdownType == EToolShutdownType::Accept)
@@ -948,7 +936,7 @@ void USkinWeightsBindingTool::Render(IToolsContextRenderAPI* RenderAPI)
 		bool bShowBoundary = true;
 
 		FPrimitiveDrawInterface* PDI = RenderAPI->GetPrimitiveDrawInterface();
-		FTransform Transform = TargetComponentInterface(0)->GetWorldTransform();
+		const FTransform Transform = (FTransform) UE::ToolTarget::GetLocalToWorldTransform(Targets[0]);
 		float PDIScale = RenderAPI->GetCameraState().GetPDIScalingFactor();
 
 		for (int32 I = 0; I < Occupancy->Occupancy.Size(); I++)
@@ -1006,7 +994,7 @@ TUniquePtr<UE::Geometry::FDynamicMeshOperator> USkinWeightsBindingTool::MakeNewO
 	Op->OriginalMesh = OriginalMesh;
 	Op->TransformHierarchy = TransformHierarchy;
 
-	const FTransform LocalToWorld = TargetComponentInterface(0)->GetWorldTransform();
+	const FTransform LocalToWorld = (FTransform) UE::ToolTarget::GetLocalToWorldTransform(Targets[0]);
 	Op->SetResultTransform(static_cast<UE::Geometry::FTransform3d>(LocalToWorld));
 	
 	return Op;
@@ -1019,7 +1007,7 @@ void USkinWeightsBindingTool::GenerateAsset(const FDynamicMeshOpResult& Result)
 	GetToolManager()->BeginUndoTransaction(LOCTEXT("SkinWeightsBindingToolTransactionName", "Create Rigid Binding"));
 
 	check(Result.Mesh.Get() != nullptr);
-	TargetMeshCommitterInterface(0)->CommitMeshDescription([&Result](const IMeshDescriptionCommitter::FCommitterParams& CommitParams)
+	Cast<IMeshDescriptionCommitter>(Targets[0])->CommitMeshDescription([&Result](const IMeshDescriptionCommitter::FCommitterParams& CommitParams)
 	{
 		FDynamicMeshToMeshDescription Converter;
 
diff --git a/Engine/Plugins/Experimental/MeshModelingToolsetExp/Source/MeshModelingToolsExp/Private/SplitMeshesTool.cpp b/Engine/Plugins/Experimental/MeshModelingToolsetExp/Source/MeshModelingToolsExp/Private/SplitMeshesTool.cpp
--- a/Engine/Plugins/Experimental/MeshModelingToolsetExp/Source/MeshModelingToolsExp/Private/SplitMeshesTool.cpp
+++ b/Engine/Plugins/Experimental/MeshModelingToolsetExp/Source/MeshModelingToolsExp/Private/SplitMeshesTool.cpp
@@ -13,6 +13,9 @@
 #include "DynamicMesh/MeshTransforms.h"
 #include "DynamicSubmesh3.h"
 
+#include "TargetInterfaces/MeshDescriptionProvider.h"
+#include "TargetInterfaces/PrimitiveComponentBackedTarget.h"
+
 using namespace UE::Geometry;
 
 #define LOCTEXT_NAMESPACE "USplitMeshesTool"
@@ -30,18 +33,9 @@ const FToolTargetTypeRequirements& USplitMeshesToolBuilder::GetTargetRequirement
 	return TypeRequirements;
 }
 
-bool USplitMeshesToolBuilder::CanBuildTool(const FToolBuilderState& SceneState) const
-{
-	return SceneState.TargetManager->CountSelectedAndTargetable(SceneState, GetTargetRequirements()) > 0;
-}
-
-UInteractiveTool* USplitMeshesToolBuilder::BuildTool(const FToolBuilderState& SceneState) const
+UMultiSelectionMeshEditingTool* USplitMeshesToolBuilder::CreateNewTool(const FToolBuilderState& SceneState) const
 {
-	USplitMeshesTool* NewTool = NewObject<USplitMeshesTool>(SceneState.ToolManager);
-	TArray<TObjectPtr<UToolTarget>> Targets = SceneState.TargetManager->BuildAllSelectedTargetable(SceneState, GetTargetRequirements());
-	NewTool->SetTargets(MoveTemp(Targets));
-	NewTool->SetWorld(SceneState.World);
-	return NewTool;
+	return NewObject<USplitMeshesTool>(SceneState.ToolManager);
 }
 
 
@@ -51,11 +45,6 @@ UInteractiveTool* USplitMeshesToolBuilder::BuildTool(const FToolBuilderState& Sc
  */
 
 
-void USplitMeshesTool::SetWorld(UWorld* World)
-{
-	this->TargetWorld = World;
-}
-
 void USplitMeshesTool::Setup()
 {
 	UInteractiveTool::Setup();
diff --git a/Engine/Plugins/Experimental/MeshModelingToolsetExp/Source/MeshModelingToolsExp/Private/TransferMeshTool.cpp b/Engine/Plugins/Experimental/MeshModelingToolsetExp/Source/MeshModelingToolsExp/Private/TransferMeshTool.cpp
--- a/Engine/Plugins/Experimental/MeshModelingToolsetExp/Source/MeshModelingToolsExp/Private/TransferMeshTool.cpp
+++ b/Engine/Plugins/Experimental/MeshModelingToolsetExp/Source/MeshModelingToolsExp/Private/TransferMeshTool.cpp
@@ -15,40 +15,21 @@ using namespace UE::Geometry;
 /*
  * ToolBuilder
  */
-const FToolTargetTypeRequirements& UTransferMeshToolBuilder::GetTargetRequirements() const
-{
-	static FToolTargetTypeRequirements TypeRequirements({
-		UMaterialProvider::StaticClass(),
-		UMeshDescriptionCommitter::StaticClass(),
-		UMeshDescriptionProvider::StaticClass(),
-		UPrimitiveComponentBackedTarget::StaticClass()
-		});
-	return TypeRequirements;
-}
 
 bool UTransferMeshToolBuilder::CanBuildTool(const FToolBuilderState& SceneState) const
 {
 	return SceneState.TargetManager->CountSelectedAndTargetable(SceneState, GetTargetRequirements()) == 2;
 }
 
-UInteractiveTool* UTransferMeshToolBuilder::BuildTool(const FToolBuilderState& SceneState) const
+UMultiSelectionMeshEditingTool* UTransferMeshToolBuilder::CreateNewTool(const FToolBuilderState& SceneState) const
 {
-	UTransferMeshTool* NewTool = NewObject<UTransferMeshTool>(SceneState.ToolManager);
-	TArray<TObjectPtr<UToolTarget>> Targets = SceneState.TargetManager->BuildAllSelectedTargetable(SceneState, GetTargetRequirements());
-	NewTool->SetTargets(MoveTemp(Targets));
-	NewTool->SetWorld(SceneState.World);
-	return NewTool;
+	return NewObject<UTransferMeshTool>(SceneState.ToolManager);
 }
 
 /*
  * Tool
  */
 
-void UTransferMeshTool::SetWorld(UWorld* World)
-{
-	this->TargetWorld = World;
-}
-
 void UTransferMeshTool::Setup()
 {
 	UInteractiveTool::Setup();
diff --git a/Engine/Plugins/Experimental/MeshModelingToolsetExp/Source/MeshModelingToolsExp/Private/TransformMeshesTool.cpp b/Engine/Plugins/Experimental/MeshModelingToolsetExp/Source/MeshModelingToolsExp/Private/TransformMeshesTool.cpp
--- a/Engine/Plugins/Experimental/MeshModelingToolsetExp/Source/MeshModelingToolsExp/Private/TransformMeshesTool.cpp
+++ b/Engine/Plugins/Experimental/MeshModelingToolsetExp/Source/MeshModelingToolsExp/Private/TransformMeshesTool.cpp
@@ -39,20 +39,9 @@ const FToolTargetTypeRequirements& UTransformMeshesToolBuilder::GetTargetRequire
 	return TypeRequirements;
 }
 
-bool UTransformMeshesToolBuilder::CanBuildTool(const FToolBuilderState& SceneState) const
+UMultiSelectionMeshEditingTool* UTransformMeshesToolBuilder::CreateNewTool(const FToolBuilderState& SceneState) const
 {
-	return SceneState.TargetManager->CountSelectedAndTargetable(SceneState, GetTargetRequirements()) >= 1;
-}
-
-UInteractiveTool* UTransformMeshesToolBuilder::BuildTool(const FToolBuilderState& SceneState) const
-{
-	UTransformMeshesTool* NewTool = NewObject<UTransformMeshesTool>(SceneState.ToolManager);
-
-	TArray<TObjectPtr<UToolTarget>> Targets = SceneState.TargetManager->BuildAllSelectedTargetable(SceneState, GetTargetRequirements());
-	NewTool->SetTargets(MoveTemp(Targets));
-	NewTool->SetWorld(SceneState.World, SceneState.GizmoManager);
-
-	return NewTool;
+	return NewObject<UTransformMeshesTool>(SceneState.ToolManager);
 }
 
 
@@ -66,12 +55,6 @@ UTransformMeshesTool::UTransformMeshesTool()
 {
 }
 
-void UTransformMeshesTool::SetWorld(UWorld* World, UInteractiveGizmoManager* GizmoManagerIn)
-{
-	this->TargetWorld = World;
-	this->GizmoManager = GizmoManagerIn;
-}
-
 
 void UTransformMeshesTool::Setup()
 {
@@ -122,7 +105,7 @@ void UTransformMeshesTool::Shutdown(EToolShutdownType ShutdownType)
 
 	DragAlignmentMechanic->Shutdown();
 
-	GizmoManager->DestroyAllGizmosByOwner(this);
+	GetToolManager()->GetPairedGizmoManager()->DestroyAllGizmosByOwner(this);
 }
 
 
@@ -336,7 +319,7 @@ void UTransformMeshesTool::SetActiveGizmos_PerObject()
 
 void UTransformMeshesTool::ResetActiveGizmos()
 {
-	GizmoManager->DestroyAllGizmosByOwner(this);
+	GetToolManager()->GetPairedGizmoManager()->DestroyAllGizmosByOwner(this);
 	ActiveGizmos.Reset();
 }
 
@@ -357,7 +340,7 @@ FInputRayHit UTransformMeshesTool::CanBeginClickDragSequence(const FInputDeviceR
 
 	for ( int k = 0; k < Targets.Num(); ++k )
 	{
-		IPrimitiveComponentBackedTarget* Target = TargetComponentInterface(k);
+		const IPrimitiveComponentBackedTarget* Target = Cast<IPrimitiveComponentBackedTarget>(Targets[k]);
 
 		FHitResult WorldHit;
 		if (Target->HitTestComponent(PressPos.WorldRay, WorldHit))
diff --git a/Engine/Plugins/Experimental/MeshModelingToolsetExp/Source/MeshModelingToolsExp/Private/UVLayoutTool.cpp b/Engine/Plugins/Experimental/MeshModelingToolsetExp/Source/MeshModelingToolsExp/Private/UVLayoutTool.cpp
--- a/Engine/Plugins/Experimental/MeshModelingToolsetExp/Source/MeshModelingToolsExp/Private/UVLayoutTool.cpp
+++ b/Engine/Plugins/Experimental/MeshModelingToolsetExp/Source/MeshModelingToolsExp/Private/UVLayoutTool.cpp
@@ -5,6 +5,7 @@
 #include "ToolBuilderUtil.h"
 
 #include "ToolSetupUtil.h"
+#include "ModelingToolTargetUtil.h"
 
 #include "DynamicMesh/DynamicMesh3.h"
 
@@ -28,32 +29,9 @@ using namespace UE::Geometry;
  * ToolBuilder
  */
 
-
-const FToolTargetTypeRequirements& UUVLayoutToolBuilder::GetTargetRequirements() const
+UMultiSelectionMeshEditingTool* UUVLayoutToolBuilder::CreateNewTool(const FToolBuilderState& SceneState) const
 {
-	static FToolTargetTypeRequirements TypeRequirements({
-		UMaterialProvider::StaticClass(),
-		UMeshDescriptionCommitter::StaticClass(),
-		UMeshDescriptionProvider::StaticClass(),
-		UPrimitiveComponentBackedTarget::StaticClass()
-		});
-	return TypeRequirements;
-}
-
-bool UUVLayoutToolBuilder::CanBuildTool(const FToolBuilderState& SceneState) const
-{
-	return SceneState.TargetManager->CountSelectedAndTargetable(SceneState, GetTargetRequirements()) >= 1;
-}
-
-UInteractiveTool* UUVLayoutToolBuilder::BuildTool(const FToolBuilderState& SceneState) const
-{
-	UUVLayoutTool* NewTool = NewObject<UUVLayoutTool>(SceneState.ToolManager);
-
-	TArray<TObjectPtr<UToolTarget>> Targets = SceneState.TargetManager->BuildAllSelectedTargetable(SceneState, GetTargetRequirements());
-	NewTool->SetTargets(MoveTemp(Targets));
-	NewTool->SetWorld(SceneState.World, SceneState.GizmoManager);
-
-	return NewTool;
+	return NewObject<UUVLayoutTool>(SceneState.ToolManager);
 }
 
 
@@ -66,11 +44,6 @@ UUVLayoutTool::UUVLayoutTool()
 {
 }
 
-void UUVLayoutTool::SetWorld(UWorld* World, UInteractiveGizmoManager* GizmoManagerIn)
-{
-	this->TargetWorld = World;
-}
-
 void UUVLayoutTool::Setup()
 {
 	UInteractiveTool::Setup();
@@ -78,7 +51,7 @@ void UUVLayoutTool::Setup()
 	// hide input StaticMeshComponent
 	for (int32 ComponentIdx = 0; ComponentIdx < Targets.Num(); ComponentIdx++)
 	{
-		TargetComponentInterface(ComponentIdx)->SetOwnerVisibility(false);
+		UE::ToolTarget::HideSourceObject(Targets[ComponentIdx]);
 	}
 
 	// if we only have one object, add ability to set UV channel
@@ -86,7 +59,7 @@ void UUVLayoutTool::Setup()
 	{
 		UVChannelProperties = NewObject<UMeshUVChannelProperties>(this);
 		UVChannelProperties->RestoreProperties(this);
-		UVChannelProperties->Initialize(TargetMeshProviderInterface(0)->GetMeshDescription(), false);
+		UVChannelProperties->Initialize(UE::ToolTarget::GetMeshDescription(Targets[0]), false);
 		UVChannelProperties->ValidateSelection(true);
 		AddToolPropertySource(UVChannelProperties);
 		UVChannelProperties->WatchProperty(UVChannelProperties->UVChannel, [this](const FString& NewValue)
@@ -111,13 +84,13 @@ void UUVLayoutTool::Setup()
 		UVLayoutView = NewObject<UUVLayoutPreview>(this);
 		UVLayoutView->CreateInWorld(TargetWorld);
 
-		FComponentMaterialSet MaterialSet;
-		TargetMaterialInterface(0)->GetMaterialSet(MaterialSet);
+		const FComponentMaterialSet MaterialSet = UE::ToolTarget::GetMaterialSet(Targets[0]);
 		UVLayoutView->SetSourceMaterials(MaterialSet);
 
+		const AActor* Actor = UE::ToolTarget::GetTargetActor(Targets[0]);
 		UVLayoutView->SetSourceWorldPosition(
-			TargetComponentInterface(0)->GetOwnerActor()->GetTransform(),
-			TargetComponentInterface(0)->GetOwnerActor()->GetComponentsBoundingBox());
+			Actor->GetTransform(),
+			Actor->GetComponentsBoundingBox());
 
 		UVLayoutView->Settings->RestoreProperties(this);
 		AddToolPropertySource(UVLayoutView->Settings);
@@ -153,25 +126,24 @@ void UUVLayoutTool::UpdateNumPreviews()
 		{
 			OriginalDynamicMeshes[PreviewIdx] = MakeShared<FDynamicMesh3, ESPMode::ThreadSafe>();
 			FMeshDescriptionToDynamicMesh Converter;
-			Converter.Convert(TargetMeshProviderInterface(PreviewIdx)->GetMeshDescription(), *OriginalDynamicMeshes[PreviewIdx]);
+			Converter.Convert(UE::ToolTarget::GetMeshDescription(Targets[PreviewIdx]), *OriginalDynamicMeshes[PreviewIdx]);
 
 			Factories[PreviewIdx]= NewObject<UUVLayoutOperatorFactory>();
 			Factories[PreviewIdx]->OriginalMesh = OriginalDynamicMeshes[PreviewIdx];
 			Factories[PreviewIdx]->Settings = BasicProperties;
-			Factories[PreviewIdx]->TargetTransform = TargetComponentInterface(PreviewIdx)->GetWorldTransform();
+			Factories[PreviewIdx]->TargetTransform = (FTransform) UE::ToolTarget::GetLocalToWorldTransform(Targets[PreviewIdx]);
 			Factories[PreviewIdx]->GetSelectedUVChannel = [this]() { return GetSelectedUVChannel(); };
 
 			UMeshOpPreviewWithBackgroundCompute* Preview = Previews.Add_GetRef(NewObject<UMeshOpPreviewWithBackgroundCompute>(Factories[PreviewIdx], "Preview"));
 			Preview->Setup(this->TargetWorld, Factories[PreviewIdx]);
 			ToolSetupUtil::ApplyRenderingConfigurationToPreview(Preview->PreviewMesh, Targets[PreviewIdx]); 
 
-			FComponentMaterialSet MaterialSet;
-			TargetMaterialInterface(PreviewIdx)->GetMaterialSet(MaterialSet);
+			const FComponentMaterialSet MaterialSet = UE::ToolTarget::GetMaterialSet(Targets[PreviewIdx]);
 			Preview->ConfigureMaterials(MaterialSet.Materials,
 				ToolSetupUtil::GetDefaultWorkingMaterial(GetToolManager())
 			);
 			Preview->PreviewMesh->UpdatePreview(OriginalDynamicMeshes[PreviewIdx].Get());
-			Preview->PreviewMesh->SetTransform(TargetComponentInterface(PreviewIdx)->GetWorldTransform());
+			Preview->PreviewMesh->SetTransform((FTransform) UE::ToolTarget::GetLocalToWorldTransform(Targets[PreviewIdx]));
 
 			Preview->OnMeshUpdated.AddLambda([this](UMeshOpPreviewWithBackgroundCompute* Compute)
 			{
@@ -198,7 +170,7 @@ void UUVLayoutTool::Shutdown(EToolShutdownType ShutdownType)
 	// Restore (unhide) the source meshes
 	for (int32 ComponentIdx = 0; ComponentIdx < Targets.Num(); ComponentIdx++)
 	{
-		TargetComponentInterface(ComponentIdx)->SetOwnerVisibility(true);
+		UE::ToolTarget::ShowSourceObject(Targets[ComponentIdx]);
 	}
 
 	TArray<FDynamicMeshOpResult> Results;
@@ -321,7 +293,7 @@ void UUVLayoutTool::GenerateAsset(const TArray<FDynamicMeshOpResult>& Results)
 	for (int32 ComponentIdx = 0; ComponentIdx < Targets.Num(); ComponentIdx++)
 	{
 		check(Results[ComponentIdx].Mesh.Get() != nullptr);
-		TargetMeshCommitterInterface(ComponentIdx)->CommitMeshDescription([&Results, &ComponentIdx, this](const IMeshDescriptionCommitter::FCommitterParams& CommitParams)
+		Cast<IMeshDescriptionCommitter>(Targets[ComponentIdx])->CommitMeshDescription([&Results, &ComponentIdx, this](const IMeshDescriptionCommitter::FCommitterParams& CommitParams)
 		{
 			const FDynamicMesh3* DynamicMesh = Results[ComponentIdx].Mesh.Get();
 			FMeshDescription* MeshDescription = CommitParams.MeshDescriptionOut;
diff --git a/Engine/Plugins/Experimental/MeshModelingToolsetExp/Source/MeshModelingToolsExp/Public/AlignObjectsTool.h b/Engine/Plugins/Experimental/MeshModelingToolsetExp/Source/MeshModelingToolsExp/Public/AlignObjectsTool.h
--- a/Engine/Plugins/Experimental/MeshModelingToolsetExp/Source/MeshModelingToolsExp/Public/AlignObjectsTool.h
+++ b/Engine/Plugins/Experimental/MeshModelingToolsetExp/Source/MeshModelingToolsExp/Public/AlignObjectsTool.h
@@ -3,7 +3,7 @@
 #pragma once
 
 #include "CoreMinimal.h"
-#include "MultiSelectionTool.h"
+#include "BaseTools/MultiSelectionMeshEditingTool.h"
 #include "InteractiveToolBuilder.h"
 #include "BaseBehaviors/BehaviorTargetInterfaces.h"
 #include "Changes/TransformChange.h"
@@ -17,16 +17,16 @@ class UPrimitiveComponent;
  *
  */
 UCLASS()
-class MESHMODELINGTOOLSEXP_API UAlignObjectsToolBuilder : public UInteractiveToolBuilder
+class MESHMODELINGTOOLSEXP_API UAlignObjectsToolBuilder : public UMultiSelectionMeshEditingToolBuilder
 {
 	GENERATED_BODY()
 
 public:
 	virtual bool CanBuildTool(const FToolBuilderState& SceneState) const override;
-	virtual UInteractiveTool* BuildTool(const FToolBuilderState& SceneState) const override;
+	virtual UMultiSelectionMeshEditingTool* CreateNewTool(const FToolBuilderState& SceneState) const override;
 
 protected:
-	virtual const FToolTargetTypeRequirements& GetTargetRequirements() const;
+	virtual const FToolTargetTypeRequirements& GetTargetRequirements() const override;
 };
 
 
@@ -99,7 +99,7 @@ public:
  * On cancel the original positions are restored, and on accept the positions are updated with a transaction.
  */
 UCLASS()
-class MESHMODELINGTOOLSEXP_API UAlignObjectsTool : public UMultiSelectionTool, public IClickDragBehaviorTarget
+class MESHMODELINGTOOLSEXP_API UAlignObjectsTool : public UMultiSelectionMeshEditingTool, public IClickDragBehaviorTarget
 {
 	GENERATED_BODY()
 	using FTransform3d = UE::Geometry::FTransform3d;
@@ -109,8 +109,6 @@ public:
 
 	virtual void RegisterActions(FInteractiveToolActionSet& ActionSet) override;
 
-	virtual void SetWorld(UWorld* World, UInteractiveGizmoManager* GizmoManager);
-
 	virtual void Setup() override;
 	virtual void Shutdown(EToolShutdownType ShutdownType) override;
 
@@ -137,9 +135,6 @@ public:
 
 
 protected:
-	UWorld* TargetWorld;
-	UInteractiveGizmoManager* GizmoManager;
-
 
 	struct FAlignInfo
 	{
diff --git a/Engine/Plugins/Experimental/MeshModelingToolsetExp/Source/MeshModelingToolsExp/Public/BakeMeshAttributeMapsTool.h b/Engine/Plugins/Experimental/MeshModelingToolsetExp/Source/MeshModelingToolsExp/Public/BakeMeshAttributeMapsTool.h
--- a/Engine/Plugins/Experimental/MeshModelingToolsetExp/Source/MeshModelingToolsExp/Public/BakeMeshAttributeMapsTool.h
+++ b/Engine/Plugins/Experimental/MeshModelingToolsetExp/Source/MeshModelingToolsExp/Public/BakeMeshAttributeMapsTool.h
@@ -3,7 +3,6 @@
 #pragma once
 
 #include "CoreMinimal.h"
-#include "MultiSelectionTool.h"
 #include "InteractiveToolBuilder.h"
 #include "DynamicMesh/DynamicMesh3.h"
 #include "DynamicMesh/DynamicMeshAABBTree3.h"
@@ -29,16 +28,16 @@ using UE::Geometry::FImageDimensions;
  */
 
 UCLASS()
-class MESHMODELINGTOOLSEXP_API UBakeMeshAttributeMapsToolBuilder : public UInteractiveToolBuilder
+class MESHMODELINGTOOLSEXP_API UBakeMeshAttributeMapsToolBuilder : public UMultiSelectionMeshEditingToolBuilder
 {
 	GENERATED_BODY()
 
 public:
 	virtual bool CanBuildTool(const FToolBuilderState& SceneState) const override;
-	virtual UInteractiveTool* BuildTool(const FToolBuilderState& SceneState) const override;
+	virtual UMultiSelectionMeshEditingTool* CreateNewTool(const FToolBuilderState& SceneState) const override;
 
 protected:
-	virtual const FToolTargetTypeRequirements& GetTargetRequirements() const;
+	virtual const FToolTargetTypeRequirements& GetTargetRequirements() const override;
 };
 
 
diff --git a/Engine/Plugins/Experimental/MeshModelingToolsetExp/Source/MeshModelingToolsExp/Public/BakeMeshAttributeTool.h b/Engine/Plugins/Experimental/MeshModelingToolsetExp/Source/MeshModelingToolsExp/Public/BakeMeshAttributeTool.h
--- a/Engine/Plugins/Experimental/MeshModelingToolsetExp/Source/MeshModelingToolsExp/Public/BakeMeshAttributeTool.h
+++ b/Engine/Plugins/Experimental/MeshModelingToolsetExp/Source/MeshModelingToolsExp/Public/BakeMeshAttributeTool.h
@@ -5,7 +5,7 @@
 #include "CoreMinimal.h"
 #include "Engine/Classes/Engine/Texture2D.h"
 #include "Materials/MaterialInstanceDynamic.h"
-#include "MultiSelectionTool.h"
+#include "BaseTools/MultiSelectionMeshEditingTool.h"
 #include "InteractiveToolQueryInterfaces.h" // for UInteractiveToolExclusiveToolAPI
 #include "DynamicMesh/DynamicMesh3.h"
 #include "PreviewMesh.h"
@@ -77,7 +77,7 @@ static constexpr EBakeMapType ALL_BAKE_MAP_TYPES[] =
  * Base Mesh Bake tool
  */
 UCLASS()
-class MESHMODELINGTOOLSEXP_API UBakeMeshAttributeTool : public UMultiSelectionTool, public IInteractiveToolExclusiveToolAPI
+class MESHMODELINGTOOLSEXP_API UBakeMeshAttributeTool : public UMultiSelectionMeshEditingTool, public IInteractiveToolExclusiveToolAPI
 {
 	GENERATED_BODY()
 
@@ -92,8 +92,6 @@ public:
 	virtual bool CanAccept() const override { return true; }
 	// End UInteractiveTool interface
 
-	void SetWorld(UWorld* World);
-
 protected:
 	//
 	// Preview materials
@@ -110,9 +108,6 @@ protected:
 	//
 	// Bake parameters
 	//
-	UPROPERTY()
-	TObjectPtr<UWorld> TargetWorld = nullptr;
-
 	EBakeOpState OpState = EBakeOpState::Evaluate;
 
 
diff --git a/Engine/Plugins/Experimental/MeshModelingToolsetExp/Source/MeshModelingToolsExp/Public/BakeMeshAttributeVertexTool.h b/Engine/Plugins/Experimental/MeshModelingToolsetExp/Source/MeshModelingToolsExp/Public/BakeMeshAttributeVertexTool.h
--- a/Engine/Plugins/Experimental/MeshModelingToolsetExp/Source/MeshModelingToolsExp/Public/BakeMeshAttributeVertexTool.h
+++ b/Engine/Plugins/Experimental/MeshModelingToolsetExp/Source/MeshModelingToolsExp/Public/BakeMeshAttributeVertexTool.h
@@ -22,16 +22,12 @@ class UMaterialInstanceDynamic;
  * Tool Builder
  */
 UCLASS()
-class MESHMODELINGTOOLSEXP_API UBakeMeshAttributeVertexToolBuilder : public UInteractiveToolBuilder
+class MESHMODELINGTOOLSEXP_API UBakeMeshAttributeVertexToolBuilder : public UMultiSelectionMeshEditingToolBuilder
 {
 	GENERATED_BODY()
-
 public:
 	virtual bool CanBuildTool(const FToolBuilderState& SceneState) const override;
-	virtual UInteractiveTool* BuildTool(const FToolBuilderState& SceneState) const override;
-
-protected:
-	virtual const FToolTargetTypeRequirements& GetTargetRequirements() const;
+	virtual UMultiSelectionMeshEditingTool* CreateNewTool(const FToolBuilderState& SceneState) const override;
 };
 
 
diff --git a/Engine/Plugins/Experimental/MeshModelingToolsetExp/Source/MeshModelingToolsExp/Public/BakeMultiMeshAttributeMapsTool.h b/Engine/Plugins/Experimental/MeshModelingToolsetExp/Source/MeshModelingToolsExp/Public/BakeMultiMeshAttributeMapsTool.h
--- a/Engine/Plugins/Experimental/MeshModelingToolsetExp/Source/MeshModelingToolsExp/Public/BakeMultiMeshAttributeMapsTool.h
+++ b/Engine/Plugins/Experimental/MeshModelingToolsetExp/Source/MeshModelingToolsExp/Public/BakeMultiMeshAttributeMapsTool.h
@@ -19,16 +19,16 @@
  * Tool Builder
  */
 UCLASS()
-class MESHMODELINGTOOLSEXP_API UBakeMultiMeshAttributeMapsToolBuilder : public UInteractiveToolBuilder
+class MESHMODELINGTOOLSEXP_API UBakeMultiMeshAttributeMapsToolBuilder : public UMultiSelectionMeshEditingToolBuilder
 {
 	GENERATED_BODY()
 
-	public:
+public:
 	virtual bool CanBuildTool(const FToolBuilderState& SceneState) const override;
-	virtual UInteractiveTool* BuildTool(const FToolBuilderState& SceneState) const override;
+	virtual UMultiSelectionMeshEditingTool* CreateNewTool(const FToolBuilderState& SceneState) const override;
 
-	protected:
-	virtual const FToolTargetTypeRequirements& GetTargetRequirements() const;
+protected:
+	virtual const FToolTargetTypeRequirements& GetTargetRequirements() const override;
 };
 
 
diff --git a/Engine/Plugins/Experimental/MeshModelingToolsetExp/Source/MeshModelingToolsExp/Public/BakeTransformTool.h b/Engine/Plugins/Experimental/MeshModelingToolsetExp/Source/MeshModelingToolsExp/Public/BakeTransformTool.h
--- a/Engine/Plugins/Experimental/MeshModelingToolsetExp/Source/MeshModelingToolsExp/Public/BakeTransformTool.h
+++ b/Engine/Plugins/Experimental/MeshModelingToolsetExp/Source/MeshModelingToolsExp/Public/BakeTransformTool.h
@@ -4,7 +4,7 @@
 
 #include "CoreMinimal.h"
 #include "UObject/NoExportTypes.h"
-#include "MultiSelectionTool.h"
+#include "BaseTools/MultiSelectionMeshEditingTool.h"
 #include "InteractiveToolBuilder.h"
 #include "MeshOpPreviewHelpers.h"
 #include "DynamicMesh/DynamicMesh3.h"
@@ -21,16 +21,11 @@ class UDynamicMeshComponent;
  *
  */
 UCLASS()
-class MESHMODELINGTOOLSEXP_API UBakeTransformToolBuilder : public UInteractiveToolBuilder
+class MESHMODELINGTOOLSEXP_API UBakeTransformToolBuilder : public UMultiSelectionMeshEditingToolBuilder
 {
 	GENERATED_BODY()
-
 public:
-	virtual bool CanBuildTool(const FToolBuilderState& SceneState) const override;
-	virtual UInteractiveTool* BuildTool(const FToolBuilderState& SceneState) const override;
-
-protected:
-	virtual const FToolTargetTypeRequirements& GetTargetRequirements() const;
+	virtual UMultiSelectionMeshEditingTool* CreateNewTool(const FToolBuilderState& SceneState) const override;
 };
 
 
@@ -76,7 +71,7 @@ public:
  * Simple tool to bake scene transform on meshes into the mesh assets
  */
 UCLASS()
-class MESHMODELINGTOOLSEXP_API UBakeTransformTool : public UMultiSelectionTool
+class MESHMODELINGTOOLSEXP_API UBakeTransformTool : public UMultiSelectionMeshEditingTool
 {
 	GENERATED_BODY()
 
@@ -86,8 +81,6 @@ public:
 	virtual void Setup() override;
 	virtual void Shutdown(EToolShutdownType ShutdownType) override;
 
-	virtual void SetWorld(UWorld* World);
-
 	virtual bool HasCancel() const override { return true; }
 	virtual bool HasAccept() const override { return true; }
 
@@ -97,8 +90,6 @@ protected:
 	TObjectPtr<UBakeTransformToolProperties> BasicProperties;
 
 protected:
-	UWorld* TargetWorld;
-
 	TArray<int> MapToFirstOccurrences;
 
 	void UpdateAssets();
diff --git a/Engine/Plugins/Experimental/MeshModelingToolsetExp/Source/MeshModelingToolsExp/Public/ConvertMeshesTool.h b/Engine/Plugins/Experimental/MeshModelingToolsetExp/Source/MeshModelingToolsExp/Public/ConvertMeshesTool.h
--- a/Engine/Plugins/Experimental/MeshModelingToolsetExp/Source/MeshModelingToolsExp/Public/ConvertMeshesTool.h
+++ b/Engine/Plugins/Experimental/MeshModelingToolsetExp/Source/MeshModelingToolsExp/Public/ConvertMeshesTool.h
@@ -3,7 +3,7 @@
 #pragma once
 
 #include "CoreMinimal.h"
-#include "MultiSelectionTool.h"
+#include "BaseTools/MultiSelectionMeshEditingTool.h"
 #include "InteractiveToolBuilder.h"
 #include "PropertySets/CreateMeshObjectTypeProperties.h"
 #include "ConvertMeshesTool.generated.h"
@@ -12,15 +12,14 @@
  *
  */
 UCLASS()
-class MESHMODELINGTOOLSEXP_API UConvertMeshesToolBuilder : public UInteractiveToolBuilder
+class MESHMODELINGTOOLSEXP_API UConvertMeshesToolBuilder : public UMultiSelectionMeshEditingToolBuilder
 {
 	GENERATED_BODY()
 public:
-	virtual bool CanBuildTool(const FToolBuilderState& SceneState) const override;
-	virtual UInteractiveTool* BuildTool(const FToolBuilderState& SceneState) const override;
+	virtual UMultiSelectionMeshEditingTool* CreateNewTool(const FToolBuilderState& SceneState) const override;
 
 protected:
-	virtual const FToolTargetTypeRequirements& GetTargetRequirements() const;
+	virtual const FToolTargetTypeRequirements& GetTargetRequirements() const override;
 };
 
 /**
@@ -38,13 +37,11 @@ public:
 
 
 UCLASS()
-class MESHMODELINGTOOLSEXP_API UConvertMeshesTool : public UMultiSelectionTool
+class MESHMODELINGTOOLSEXP_API UConvertMeshesTool : public UMultiSelectionMeshEditingTool
 {
 	GENERATED_BODY()
 
 public:
-	virtual void SetWorld(UWorld* World);
-
 	virtual void Setup() override;
 	virtual void Shutdown(EToolShutdownType ShutdownType) override;
 
@@ -57,7 +54,4 @@ public:
 
 	UPROPERTY()
 	TObjectPtr<UCreateMeshObjectTypeProperties> OutputTypeProperties;
-
-protected:
-	UWorld* TargetWorld;
 };
diff --git a/Engine/Plugins/Experimental/MeshModelingToolsetExp/Source/MeshModelingToolsExp/Public/EditNormalsTool.h b/Engine/Plugins/Experimental/MeshModelingToolsetExp/Source/MeshModelingToolsExp/Public/EditNormalsTool.h
--- a/Engine/Plugins/Experimental/MeshModelingToolsetExp/Source/MeshModelingToolsExp/Public/EditNormalsTool.h
+++ b/Engine/Plugins/Experimental/MeshModelingToolsetExp/Source/MeshModelingToolsExp/Public/EditNormalsTool.h
@@ -4,7 +4,7 @@
 
 #include "CoreMinimal.h"
 #include "UObject/NoExportTypes.h"
-#include "MultiSelectionTool.h"
+#include "BaseTools/MultiSelectionMeshEditingTool.h"
 #include "InteractiveToolBuilder.h"
 #include "MeshOpPreviewHelpers.h"
 #include "CleaningOps/EditNormalsOp.h"
@@ -16,21 +16,18 @@
 // predeclarations
 struct FMeshDescription;
 class UDynamicMeshComponent;
+class UEditNormalsTool;
 
 /**
  *
  */
 UCLASS()
-class MESHMODELINGTOOLSEXP_API UEditNormalsToolBuilder : public UInteractiveToolBuilder
+class MESHMODELINGTOOLSEXP_API UEditNormalsToolBuilder : public UMultiSelectionMeshEditingToolBuilder
 {
 	GENERATED_BODY()
 
 public:
-	virtual bool CanBuildTool(const FToolBuilderState& SceneState) const override;
-	virtual UInteractiveTool* BuildTool(const FToolBuilderState& SceneState) const override;
-
-protected:
-	virtual const FToolTargetTypeRequirements& GetTargetRequirements() const;
+	virtual UMultiSelectionMeshEditingTool* CreateNewTool(const FToolBuilderState& SceneState) const override;
 };
 
 
@@ -122,7 +119,7 @@ public:
  * Simple Mesh Normal Updating Tool
  */
 UCLASS()
-class MESHMODELINGTOOLSEXP_API UEditNormalsTool : public UMultiSelectionTool
+class MESHMODELINGTOOLSEXP_API UEditNormalsTool : public UMultiSelectionMeshEditingTool
 {
 	GENERATED_BODY()
 
@@ -135,8 +132,6 @@ public:
 	virtual void Setup() override;
 	virtual void Shutdown(EToolShutdownType ShutdownType) override;
 
-	virtual void SetWorld(UWorld* World);
-
 	virtual void OnTick(float DeltaTime) override;
 	virtual void Render(IToolsContextRenderAPI* RenderAPI) override;
 
@@ -166,8 +161,6 @@ protected:
 protected:
 	TArray<TSharedPtr<UE::Geometry::FDynamicMesh3, ESPMode::ThreadSafe>> OriginalDynamicMeshes;
 
-	UWorld* TargetWorld;
-
 	FViewCameraState CameraState;
 
 	void UpdateNumPreviews();
diff --git a/Engine/Plugins/Experimental/MeshModelingToolsetExp/Source/MeshModelingToolsExp/Public/EditPivotTool.h b/Engine/Plugins/Experimental/MeshModelingToolsetExp/Source/MeshModelingToolsExp/Public/EditPivotTool.h
--- a/Engine/Plugins/Experimental/MeshModelingToolsetExp/Source/MeshModelingToolsExp/Public/EditPivotTool.h
+++ b/Engine/Plugins/Experimental/MeshModelingToolsetExp/Source/MeshModelingToolsExp/Public/EditPivotTool.h
@@ -3,7 +3,7 @@
 #pragma once
 
 #include "CoreMinimal.h"
-#include "MultiSelectionTool.h"
+#include "BaseTools/MultiSelectionMeshEditingTool.h"
 #include "InteractiveToolBuilder.h"
 #include "BaseBehaviors/BehaviorTargetInterfaces.h"
 #include "Changes/TransformChange.h"
@@ -16,21 +16,17 @@ class UBaseAxisTranslationGizmo;
 class UAxisAngleGizmo;
 class UCombinedTransformGizmo;
 class UTransformProxy;
+class UEditPivotTool;
 
 /**
  *
  */
 UCLASS()
-class MESHMODELINGTOOLSEXP_API UEditPivotToolBuilder : public UInteractiveToolBuilder
+class MESHMODELINGTOOLSEXP_API UEditPivotToolBuilder : public UMultiSelectionMeshEditingToolBuilder
 {
 	GENERATED_BODY()
-
 public:
-	virtual bool CanBuildTool(const FToolBuilderState& SceneState) const override;
-	virtual UInteractiveTool* BuildTool(const FToolBuilderState& SceneState) const override;
-
-protected:
-	virtual const FToolTargetTypeRequirements& GetTargetRequirements() const;
+	virtual UMultiSelectionMeshEditingTool* CreateNewTool(const FToolBuilderState& SceneState) const override;
 };
 
 
@@ -151,7 +147,7 @@ public:
  *
  */
 UCLASS()
-class MESHMODELINGTOOLSEXP_API UEditPivotTool : public UMultiSelectionTool, public IClickDragBehaviorTarget
+class MESHMODELINGTOOLSEXP_API UEditPivotTool : public UMultiSelectionMeshEditingTool, public IClickDragBehaviorTarget
 {
 	GENERATED_BODY()
 
@@ -160,8 +156,6 @@ public:
 
 	virtual void RegisterActions(FInteractiveToolActionSet& ActionSet) override;
 
-	virtual void SetWorld(UWorld* World, UInteractiveGizmoManager* GizmoManager);
-
 	virtual void Setup() override;
 	virtual void Shutdown(EToolShutdownType ShutdownType) override;
 
@@ -191,9 +185,6 @@ public:
 	virtual void RequestAction(EEditPivotToolActions ActionType);
 
 protected:
-	UWorld* TargetWorld;
-	UInteractiveGizmoManager* GizmoManager;
-
 	TArray<int> MapToFirstOccurrences;
 
 	UE::Geometry::FTransform3d Transform;
diff --git a/Engine/Plugins/Experimental/MeshModelingToolsetExp/Source/MeshModelingToolsExp/Public/MirrorTool.h b/Engine/Plugins/Experimental/MeshModelingToolsetExp/Source/MeshModelingToolsExp/Public/MirrorTool.h
--- a/Engine/Plugins/Experimental/MeshModelingToolsetExp/Source/MeshModelingToolsExp/Public/MirrorTool.h
+++ b/Engine/Plugins/Experimental/MeshModelingToolsetExp/Source/MeshModelingToolsExp/Public/MirrorTool.h
@@ -9,22 +9,18 @@
 #include "InteractiveToolBuilder.h"
 #include "Mechanics/ConstructionPlaneMechanic.h"
 #include "MeshOpPreviewHelpers.h"
-#include "MultiSelectionTool.h"
+#include "BaseTools/MultiSelectionMeshEditingTool.h"
 #include "Selection/SelectClickedAction.h"
 #include "ToolContextInterfaces.h"
 
 #include "MirrorTool.generated.h"
 
 UCLASS()
-class MESHMODELINGTOOLSEXP_API UMirrorToolBuilder : public UInteractiveToolBuilder
+class MESHMODELINGTOOLSEXP_API UMirrorToolBuilder : public UMultiSelectionMeshEditingToolBuilder
 {
 	GENERATED_BODY()
 public:
-	virtual bool CanBuildTool(const FToolBuilderState& SceneState) const override;
-	virtual UInteractiveTool* BuildTool(const FToolBuilderState& SceneState) const override;
-
-protected:
-	virtual const FToolTargetTypeRequirements& GetTargetRequirements() const;
+	virtual UMultiSelectionMeshEditingTool* CreateNewTool(const FToolBuilderState& SceneState) const override;
 };
 
 UENUM()
@@ -182,7 +178,7 @@ public:
 
 /** Tool for mirroring one or more meshes across a plane. */
 UCLASS()
-class MESHMODELINGTOOLSEXP_API UMirrorTool : public UMultiSelectionTool, public IModifierToggleBehaviorTarget
+class MESHMODELINGTOOLSEXP_API UMirrorTool : public UMultiSelectionMeshEditingTool, public IModifierToggleBehaviorTarget
 {
 	GENERATED_BODY()
 public:
@@ -194,8 +190,6 @@ public:
 	virtual void Setup() override;
 	virtual void Shutdown(EToolShutdownType ShutdownType) override;
 
-	virtual void SetWorld(UWorld* World);
-
 	virtual void OnTick(float DeltaTime) override;
 	virtual void Render(IToolsContextRenderAPI* RenderAPI) override;
 
@@ -226,8 +220,6 @@ protected:
 	FVector3d MirrorPlaneOrigin = FVector3d::Zero();
 	FVector3d MirrorPlaneNormal = FVector3d::UnitZ();
 
-	UWorld* TargetWorld;
-
 	UPROPERTY()
 	TObjectPtr<UConstructionPlaneMechanic> PlaneMechanic;
 
diff --git a/Engine/Plugins/Experimental/MeshModelingToolsetExp/Source/MeshModelingToolsExp/Public/Physics/PhysicsInspectorTool.h b/Engine/Plugins/Experimental/MeshModelingToolsetExp/Source/MeshModelingToolsExp/Public/Physics/PhysicsInspectorTool.h
--- a/Engine/Plugins/Experimental/MeshModelingToolsetExp/Source/MeshModelingToolsExp/Public/Physics/PhysicsInspectorTool.h
+++ b/Engine/Plugins/Experimental/MeshModelingToolsetExp/Source/MeshModelingToolsExp/Public/Physics/PhysicsInspectorTool.h
@@ -3,7 +3,7 @@
 #pragma once
 
 #include "CoreMinimal.h"
-#include "MultiSelectionTool.h"
+#include "BaseTools/MultiSelectionMeshEditingTool.h"
 #include "InteractiveToolBuilder.h"
 #include "Physics/CollisionPropertySets.h"
 #include "PhysicsInspectorTool.generated.h"
@@ -11,16 +11,15 @@
 class UPreviewGeometry;
 
 UCLASS()
-class MESHMODELINGTOOLSEXP_API UPhysicsInspectorToolBuilder : public UInteractiveToolBuilder
+class MESHMODELINGTOOLSEXP_API UPhysicsInspectorToolBuilder : public UMultiSelectionMeshEditingToolBuilder
 {
 	GENERATED_BODY()
 
 public:
-	virtual bool CanBuildTool(const FToolBuilderState& SceneState) const override;
-	virtual UInteractiveTool* BuildTool(const FToolBuilderState& SceneState) const override;
+	virtual UMultiSelectionMeshEditingTool* CreateNewTool(const FToolBuilderState& SceneState) const override;
 
 protected:
-	virtual const FToolTargetTypeRequirements& GetTargetRequirements() const;
+	virtual const FToolTargetTypeRequirements& GetTargetRequirements() const override;
 };
 
 
@@ -30,7 +29,7 @@ protected:
  * Mesh Inspector Tool for visualizing mesh information
  */
 UCLASS()
-class MESHMODELINGTOOLSEXP_API UPhysicsInspectorTool : public UMultiSelectionTool
+class MESHMODELINGTOOLSEXP_API UPhysicsInspectorTool : public UMultiSelectionMeshEditingTool
 {
 	GENERATED_BODY()
 public:
diff --git a/Engine/Plugins/Experimental/MeshModelingToolsetExp/Source/MeshModelingToolsExp/Public/Physics/SetCollisionGeometryTool.h b/Engine/Plugins/Experimental/MeshModelingToolsetExp/Source/MeshModelingToolsExp/Public/Physics/SetCollisionGeometryTool.h
--- a/Engine/Plugins/Experimental/MeshModelingToolsetExp/Source/MeshModelingToolsExp/Public/Physics/SetCollisionGeometryTool.h
+++ b/Engine/Plugins/Experimental/MeshModelingToolsetExp/Source/MeshModelingToolsExp/Public/Physics/SetCollisionGeometryTool.h
@@ -3,7 +3,7 @@
 #pragma once
 
 #include "CoreMinimal.h"
-#include "MultiSelectionTool.h"
+#include "BaseTools/MultiSelectionMeshEditingTool.h"
 #include "InteractiveToolBuilder.h"
 #include "DynamicMesh/DynamicMesh3.h"
 #include "SphereTypes.h"
@@ -19,16 +19,16 @@ PREDECLARE_GEOMETRY(class FMeshSimpleShapeApproximation)
 PREDECLARE_USE_GEOMETRY_CLASS(FDynamicMesh3);
 
 UCLASS()
-class MESHMODELINGTOOLSEXP_API USetCollisionGeometryToolBuilder : public UInteractiveToolBuilder
+class MESHMODELINGTOOLSEXP_API USetCollisionGeometryToolBuilder : public UMultiSelectionMeshEditingToolBuilder
 {
 	GENERATED_BODY()
 
 public:
 	virtual bool CanBuildTool(const FToolBuilderState& SceneState) const override;
-	virtual UInteractiveTool* BuildTool(const FToolBuilderState& SceneState) const override;
+	virtual UMultiSelectionMeshEditingTool* CreateNewTool(const FToolBuilderState& SceneState) const override;
 
 protected:
-	virtual const FToolTargetTypeRequirements& GetTargetRequirements() const;
+	virtual const FToolTargetTypeRequirements& GetTargetRequirements() const override;
 };
 
 
@@ -141,7 +141,7 @@ public:
  * Mesh Inspector Tool for visualizing mesh information
  */
 UCLASS()
-class MESHMODELINGTOOLSEXP_API USetCollisionGeometryTool : public UMultiSelectionTool
+class MESHMODELINGTOOLSEXP_API USetCollisionGeometryTool : public UMultiSelectionMeshEditingTool
 {
 	GENERATED_BODY()
 public:
diff --git a/Engine/Plugins/Experimental/MeshModelingToolsetExp/Source/MeshModelingToolsExp/Public/PlaneCutTool.h b/Engine/Plugins/Experimental/MeshModelingToolsetExp/Source/MeshModelingToolsExp/Public/PlaneCutTool.h
--- a/Engine/Plugins/Experimental/MeshModelingToolsetExp/Source/MeshModelingToolsExp/Public/PlaneCutTool.h
+++ b/Engine/Plugins/Experimental/MeshModelingToolsetExp/Source/MeshModelingToolsExp/Public/PlaneCutTool.h
@@ -4,7 +4,7 @@
 
 #include "CoreMinimal.h"
 #include "InteractiveGizmo.h"
-#include "MultiSelectionTool.h"
+#include "BaseTools/MultiSelectionMeshEditingTool.h"
 #include "InteractiveToolBuilder.h"
 #include "MeshOpPreviewHelpers.h"
 #include "DynamicMesh/DynamicMesh3.h"
@@ -13,21 +13,19 @@
 #include "Mechanics/ConstructionPlaneMechanic.h"
 #include "PlaneCutTool.generated.h"
 
+class UPlaneCutTool;
+
 
 /**
  *
  */
 UCLASS()
-class MESHMODELINGTOOLSEXP_API UPlaneCutToolBuilder : public UInteractiveToolBuilder
+class MESHMODELINGTOOLSEXP_API UPlaneCutToolBuilder : public UMultiSelectionMeshEditingToolBuilder
 {
 	GENERATED_BODY()
 
 public:
-	virtual bool CanBuildTool(const FToolBuilderState& SceneState) const override;
-	virtual UInteractiveTool* BuildTool(const FToolBuilderState& SceneState) const override;
-
-protected:
-	virtual const FToolTargetTypeRequirements& GetTargetRequirements() const;
+	virtual UMultiSelectionMeshEditingTool* CreateNewTool(const FToolBuilderState& SceneState) const override;
 };
 
 
@@ -116,7 +114,7 @@ public:
  * Simple Mesh Plane Cutting Tool
  */
 UCLASS()
-class MESHMODELINGTOOLSEXP_API UPlaneCutTool : public UMultiSelectionTool
+class MESHMODELINGTOOLSEXP_API UPlaneCutTool : public UMultiSelectionMeshEditingTool
 {
 	GENERATED_BODY()
 
@@ -129,8 +127,6 @@ public:
 	virtual void Setup() override;
 	virtual void Shutdown(EToolShutdownType ShutdownType) override;
 
-	virtual void SetWorld(UWorld* World);
-
 	virtual void RegisterActions(FInteractiveToolActionSet& ActionSet) override;
 
 	virtual void OnTick(float DeltaTime) override;
@@ -190,8 +186,6 @@ protected:
 	// UV Scale factor is cached based on the bounding box of the mesh before any cuts are performed, so you don't get inconsistent UVs if you multi-cut the object to smaller sizes
 	TArray<float> MeshUVScaleFactor;
 
-	UWorld* TargetWorld;
-
 	FViewCameraState CameraState;
 
 	EPlaneCutToolActions PendingAction = EPlaneCutToolActions::NoAction;
diff --git a/Engine/Plugins/Experimental/MeshModelingToolsetExp/Source/MeshModelingToolsExp/Public/ProjectToTargetTool.h b/Engine/Plugins/Experimental/MeshModelingToolsetExp/Source/MeshModelingToolsExp/Public/ProjectToTargetTool.h
--- a/Engine/Plugins/Experimental/MeshModelingToolsetExp/Source/MeshModelingToolsExp/Public/ProjectToTargetTool.h
+++ b/Engine/Plugins/Experimental/MeshModelingToolsetExp/Source/MeshModelingToolsExp/Public/ProjectToTargetTool.h
@@ -10,15 +10,12 @@
  * Determine if/how we can build UProjectToTargetTool. It requires two selected mesh components.
  */
 UCLASS()
-class MESHMODELINGTOOLSEXP_API UProjectToTargetToolBuilder : public UInteractiveToolBuilder
+class MESHMODELINGTOOLSEXP_API UProjectToTargetToolBuilder : public UMultiSelectionMeshEditingToolBuilder
 {
 	GENERATED_BODY()
 public:
 	virtual bool CanBuildTool(const FToolBuilderState& SceneState) const override;
-	virtual UInteractiveTool* BuildTool(const FToolBuilderState& SceneState) const override;
-
-protected:
-	virtual const FToolTargetTypeRequirements& GetTargetRequirements() const;
+	virtual UMultiSelectionMeshEditingTool* CreateNewTool(const FToolBuilderState& SceneState) const override;
 };
 
 
diff --git a/Engine/Plugins/Experimental/MeshModelingToolsetExp/Source/MeshModelingToolsExp/Public/RemeshMeshTool.h b/Engine/Plugins/Experimental/MeshModelingToolsetExp/Source/MeshModelingToolsExp/Public/RemeshMeshTool.h
--- a/Engine/Plugins/Experimental/MeshModelingToolsetExp/Source/MeshModelingToolsExp/Public/RemeshMeshTool.h
+++ b/Engine/Plugins/Experimental/MeshModelingToolsetExp/Source/MeshModelingToolsExp/Public/RemeshMeshTool.h
@@ -3,7 +3,7 @@
 #pragma once
 
 #include "CoreMinimal.h"
-#include "MultiSelectionTool.h"
+#include "BaseTools/MultiSelectionMeshEditingTool.h"
 #include "InteractiveToolBuilder.h"
 #include "MeshOpPreviewHelpers.h"
 #include "CleaningOps/RemeshMeshOp.h"
@@ -19,7 +19,7 @@ PREDECLARE_GEOMETRY(typedef TMeshAABBTree3<FDynamicMesh3> FDynamicMeshAABBTree3)
  *
  */
 UCLASS()
-class MESHMODELINGTOOLSEXP_API URemeshMeshToolBuilder : public UInteractiveToolBuilder
+class MESHMODELINGTOOLSEXP_API URemeshMeshToolBuilder : public UMultiSelectionMeshEditingToolBuilder
 {
 	GENERATED_BODY()
 
@@ -30,10 +30,7 @@ public:
 	 */
 	virtual bool CanBuildTool(const FToolBuilderState& SceneState) const override;
 
-	virtual UInteractiveTool* BuildTool(const FToolBuilderState& SceneState) const override;
-
-protected:
-	virtual const FToolTargetTypeRequirements& GetTargetRequirements() const;
+	virtual UMultiSelectionMeshEditingTool* CreateNewTool(const FToolBuilderState& SceneState) const override;
 };
 
 /**
@@ -102,7 +99,7 @@ public:
  * that subclasses of this class can work with multiple meshes (see, for example, UProjectToTargetTool.)
  */
 UCLASS()
-class MESHMODELINGTOOLSEXP_API URemeshMeshTool : public UMultiSelectionTool, public UE::Geometry::IDynamicMeshOperatorFactory
+class MESHMODELINGTOOLSEXP_API URemeshMeshTool : public UMultiSelectionMeshEditingTool, public UE::Geometry::IDynamicMeshOperatorFactory
 {
 	GENERATED_BODY()
 
@@ -110,8 +107,6 @@ public:
 
 	URemeshMeshTool(const FObjectInitializer&);
 
-	virtual void SetWorld(UWorld* World);
-
 	virtual void Setup() override;
 	virtual void Shutdown(EToolShutdownType ShutdownType) override;
 
@@ -139,9 +134,6 @@ public:
 	TObjectPtr<UMeshElementsVisualizer> MeshElementsDisplay;
 
 protected:
-
-	UWorld* TargetWorld;
-
 	TSharedPtr<UE::Geometry::FDynamicMesh3, ESPMode::ThreadSafe> OriginalMesh;
 	TSharedPtr<UE::Geometry::FDynamicMeshAABBTree3, ESPMode::ThreadSafe> OriginalMeshSpatial;
 	double InitialMeshArea;
diff --git a/Engine/Plugins/Experimental/MeshModelingToolsetExp/Source/MeshModelingToolsExp/Public/RemoveOccludedTrianglesTool.h b/Engine/Plugins/Experimental/MeshModelingToolsetExp/Source/MeshModelingToolsExp/Public/RemoveOccludedTrianglesTool.h
--- a/Engine/Plugins/Experimental/MeshModelingToolsetExp/Source/MeshModelingToolsExp/Public/RemoveOccludedTrianglesTool.h
+++ b/Engine/Plugins/Experimental/MeshModelingToolsetExp/Source/MeshModelingToolsExp/Public/RemoveOccludedTrianglesTool.h
@@ -4,8 +4,7 @@
 
 #include "CoreMinimal.h"
 #include "UObject/NoExportTypes.h"
-#include "MultiSelectionTool.h"
-#include "InteractiveToolBuilder.h"
+#include "BaseTools/MultiSelectionMeshEditingTool.h"
 #include "MeshOpPreviewHelpers.h"
 #include "CleaningOps/RemoveOccludedTrianglesOp.h"
 #include "DynamicMesh/DynamicMesh3.h"
@@ -17,20 +16,17 @@
 
 PREDECLARE_USE_GEOMETRY_CLASS(FDynamicMesh3);
 
+class URemoveOccludedTrianglesTool;
+
 /**
  *
  */
 UCLASS()
-class MESHMODELINGTOOLSEXP_API URemoveOccludedTrianglesToolBuilder : public UInteractiveToolBuilder
+class MESHMODELINGTOOLSEXP_API URemoveOccludedTrianglesToolBuilder : public UMultiSelectionMeshEditingToolBuilder
 {
 	GENERATED_BODY()
-
 public:
-	virtual bool CanBuildTool(const FToolBuilderState& SceneState) const override;
-	virtual UInteractiveTool* BuildTool(const FToolBuilderState& SceneState) const override;
-
-protected:
-	virtual const FToolTargetTypeRequirements& GetTargetRequirements() const;
+	virtual UMultiSelectionMeshEditingTool* CreateNewTool(const FToolBuilderState& SceneState) const override;
 };
 
 
@@ -158,7 +154,7 @@ public:
  * Simple Mesh Normal Updating Tool
  */
 UCLASS()
-class MESHMODELINGTOOLSEXP_API URemoveOccludedTrianglesTool : public UMultiSelectionTool
+class MESHMODELINGTOOLSEXP_API URemoveOccludedTrianglesTool : public UMultiSelectionMeshEditingTool
 {
 	GENERATED_BODY()
 
@@ -171,8 +167,6 @@ public:
 	virtual void Setup() override;
 	virtual void Shutdown(EToolShutdownType ShutdownType) override;
 
-	virtual void SetWorld(UWorld* World);
-
 	virtual void OnTick(float DeltaTime) override;
 
 	virtual bool HasCancel() const override { return true; }
@@ -222,8 +216,6 @@ protected:
 	// Selected layer indices for Group IDs, per mesh in the Previews array
 	TArray<int32> OccludedGroupLayers;
 
-	UWorld* TargetWorld;
-
 	FViewCameraState CameraState;
 
 	void SetupPreviews();
diff --git a/Engine/Plugins/Experimental/MeshModelingToolsetExp/Source/MeshModelingToolsExp/Public/SelfUnionMeshesTool.h b/Engine/Plugins/Experimental/MeshModelingToolsetExp/Source/MeshModelingToolsExp/Public/SelfUnionMeshesTool.h
--- a/Engine/Plugins/Experimental/MeshModelingToolsetExp/Source/MeshModelingToolsExp/Public/SelfUnionMeshesTool.h
+++ b/Engine/Plugins/Experimental/MeshModelingToolsetExp/Source/MeshModelingToolsExp/Public/SelfUnionMeshesTool.h
@@ -106,9 +106,9 @@ class MESHMODELINGTOOLSEXP_API USelfUnionMeshesToolBuilder : public UBaseCreateF
 	GENERATED_BODY()
 
 public:
-	virtual UBaseCreateFromSelectedTool* MakeNewToolInstance(UObject* Outer) const override
+	virtual UMultiSelectionMeshEditingTool* CreateNewTool(const FToolBuilderState& SceneState) const override
 	{
-		return NewObject<USelfUnionMeshesTool>(Outer);
+		return NewObject<USelfUnionMeshesTool>(SceneState.ToolManager);
 	}
 };
 
diff --git a/Engine/Plugins/Experimental/MeshModelingToolsetExp/Source/MeshModelingToolsExp/Public/SkinWeightsBindingTool.h b/Engine/Plugins/Experimental/MeshModelingToolsetExp/Source/MeshModelingToolsExp/Public/SkinWeightsBindingTool.h
--- a/Engine/Plugins/Experimental/MeshModelingToolsetExp/Source/MeshModelingToolsExp/Public/SkinWeightsBindingTool.h
+++ b/Engine/Plugins/Experimental/MeshModelingToolsetExp/Source/MeshModelingToolsExp/Public/SkinWeightsBindingTool.h
@@ -7,7 +7,7 @@
 #include "BoneContainer.h"
 #include "InteractiveToolBuilder.h"
 #include "ModelingOperators.h"
-#include "MultiSelectionTool.h"
+#include "BaseTools/MultiSelectionMeshEditingTool.h"
 #include "Interfaces/Interface_BoneReferenceSkeletonProvider.h"
 
 #include "SkinWeightsBindingTool.generated.h"
@@ -22,12 +22,12 @@ struct FOccupancyGrid;
  *
  */
 UCLASS()
-class MESHMODELINGTOOLSEXP_API USkinWeightsBindingToolBuilder : public UInteractiveToolBuilder
+class MESHMODELINGTOOLSEXP_API USkinWeightsBindingToolBuilder : public UMultiSelectionMeshEditingToolBuilder
 {
 	GENERATED_BODY()
 public:
-	bool CanBuildTool(const FToolBuilderState& SceneState) const override;
-	UInteractiveTool* BuildTool(const FToolBuilderState& SceneState) const override;
+	virtual bool CanBuildTool(const FToolBuilderState& SceneState) const override;
+	virtual UMultiSelectionMeshEditingTool* CreateNewTool(const FToolBuilderState& SceneState) const override;
 };
 
 UENUM()
@@ -79,7 +79,7 @@ public:
  */
 UCLASS()
 class MESHMODELINGTOOLSEXP_API USkinWeightsBindingTool :
-	public UMultiSelectionTool,
+	public UMultiSelectionMeshEditingTool,
 	public UE::Geometry::IDynamicMeshOperatorFactory
 {
 	GENERATED_BODY()
@@ -87,8 +87,6 @@ class MESHMODELINGTOOLSEXP_API USkinWeightsBindingTool :
 public:
 	USkinWeightsBindingTool();
 	~USkinWeightsBindingTool();
-
-	void SetWorld(UWorld* World);
 	
 	void Setup() override;
 	void Shutdown(EToolShutdownType ShutdownType) override;
@@ -112,8 +110,6 @@ public:
 	TObjectPtr<UMeshOpPreviewWithBackgroundCompute> Preview = nullptr;
 
 protected:
-	UWorld* TargetWorld = nullptr;
-
 	TSharedPtr<FOccupancyGrid> Occupancy;
 	
 	TSharedPtr<UE::Geometry::FDynamicMesh3, ESPMode::ThreadSafe> OriginalMesh;
diff --git a/Engine/Plugins/Experimental/MeshModelingToolsetExp/Source/MeshModelingToolsExp/Public/SplitMeshesTool.h b/Engine/Plugins/Experimental/MeshModelingToolsetExp/Source/MeshModelingToolsExp/Public/SplitMeshesTool.h
--- a/Engine/Plugins/Experimental/MeshModelingToolsetExp/Source/MeshModelingToolsExp/Public/SplitMeshesTool.h
+++ b/Engine/Plugins/Experimental/MeshModelingToolsetExp/Source/MeshModelingToolsExp/Public/SplitMeshesTool.h
@@ -3,7 +3,7 @@
 #pragma once
 
 #include "CoreMinimal.h"
-#include "MultiSelectionTool.h"
+#include "BaseTools/MultiSelectionMeshEditingTool.h"
 #include "InteractiveToolBuilder.h"
 #include "PropertySets/CreateMeshObjectTypeProperties.h"
 #include "SplitMeshesTool.generated.h"
@@ -12,15 +12,14 @@ class UMaterialInterface;
 
 
 UCLASS()
-class MESHMODELINGTOOLSEXP_API USplitMeshesToolBuilder : public UInteractiveToolBuilder
+class MESHMODELINGTOOLSEXP_API USplitMeshesToolBuilder : public UMultiSelectionMeshEditingToolBuilder
 {
 	GENERATED_BODY()
 public:
-	virtual bool CanBuildTool(const FToolBuilderState& SceneState) const override;
-	virtual UInteractiveTool* BuildTool(const FToolBuilderState& SceneState) const override;
+	virtual UMultiSelectionMeshEditingTool* CreateNewTool(const FToolBuilderState& SceneState) const override;
 
 protected:
-	virtual const FToolTargetTypeRequirements& GetTargetRequirements() const;
+	virtual const FToolTargetTypeRequirements& GetTargetRequirements() const override;
 };
 
 
@@ -37,13 +36,11 @@ public:
 
 
 UCLASS()
-class MESHMODELINGTOOLSEXP_API USplitMeshesTool : public UMultiSelectionTool
+class MESHMODELINGTOOLSEXP_API USplitMeshesTool : public UMultiSelectionMeshEditingTool
 {
 	GENERATED_BODY()
 
 public:
-	virtual void SetWorld(UWorld* World);
-
 	virtual void Setup() override;
 	virtual void Shutdown(EToolShutdownType ShutdownType) override;
 
@@ -58,8 +55,6 @@ public:
 	TObjectPtr<UCreateMeshObjectTypeProperties> OutputTypeProperties;
 
 protected:
-	UWorld* TargetWorld;
-
 	struct FSourceMeshInfo
 	{
 		UE::Geometry::FDynamicMesh3 Mesh;
diff --git a/Engine/Plugins/Experimental/MeshModelingToolsetExp/Source/MeshModelingToolsExp/Public/TransferMeshTool.h b/Engine/Plugins/Experimental/MeshModelingToolsetExp/Source/MeshModelingToolsExp/Public/TransferMeshTool.h
--- a/Engine/Plugins/Experimental/MeshModelingToolsetExp/Source/MeshModelingToolsExp/Public/TransferMeshTool.h
+++ b/Engine/Plugins/Experimental/MeshModelingToolsetExp/Source/MeshModelingToolsExp/Public/TransferMeshTool.h
@@ -3,7 +3,7 @@
 #pragma once
 
 #include "CoreMinimal.h"
-#include "MultiSelectionTool.h"
+#include "BaseTools/MultiSelectionMeshEditingTool.h"
 #include "InteractiveToolBuilder.h"
 #include "TransferMeshTool.generated.h"
 
@@ -11,16 +11,12 @@
  *
  */
 UCLASS()
-class MESHMODELINGTOOLSEXP_API UTransferMeshToolBuilder : public UInteractiveToolBuilder
+class MESHMODELINGTOOLSEXP_API UTransferMeshToolBuilder : public UMultiSelectionMeshEditingToolBuilder
 {
 	GENERATED_BODY()
-
 public:
 	virtual bool CanBuildTool(const FToolBuilderState& SceneState) const override;
-	virtual UInteractiveTool* BuildTool(const FToolBuilderState& SceneState) const override;
-
-protected:
-	virtual const FToolTargetTypeRequirements& GetTargetRequirements() const;
+	virtual UMultiSelectionMeshEditingTool* CreateNewTool(const FToolBuilderState& SceneState) const override;
 };
 
 /**
@@ -38,13 +34,11 @@ public:
 
 
 UCLASS()
-class MESHMODELINGTOOLSEXP_API UTransferMeshTool : public UMultiSelectionTool
+class MESHMODELINGTOOLSEXP_API UTransferMeshTool : public UMultiSelectionMeshEditingTool
 {
 	GENERATED_BODY()
 
 public:
-	virtual void SetWorld(UWorld* World);
-
 	virtual void Setup() override;
 	virtual void Shutdown(EToolShutdownType ShutdownType) override;
 
@@ -54,7 +48,4 @@ public:
 
 	UPROPERTY()
 	TObjectPtr<UTransferMeshToolProperties> BasicProperties;
-
-protected:
-	UWorld* TargetWorld;
 };
diff --git a/Engine/Plugins/Experimental/MeshModelingToolsetExp/Source/MeshModelingToolsExp/Public/TransformMeshesTool.h b/Engine/Plugins/Experimental/MeshModelingToolsetExp/Source/MeshModelingToolsExp/Public/TransformMeshesTool.h
--- a/Engine/Plugins/Experimental/MeshModelingToolsetExp/Source/MeshModelingToolsExp/Public/TransformMeshesTool.h
+++ b/Engine/Plugins/Experimental/MeshModelingToolsetExp/Source/MeshModelingToolsExp/Public/TransformMeshesTool.h
@@ -3,7 +3,7 @@
 #pragma once
 
 #include "CoreMinimal.h"
-#include "MultiSelectionTool.h"
+#include "BaseTools/MultiSelectionMeshEditingTool.h"
 #include "InteractiveToolBuilder.h"
 #include "BaseBehaviors/BehaviorTargetInterfaces.h"
 #include "Changes/TransformChange.h"
@@ -21,16 +21,15 @@ class UTransformProxy;
  *
  */
 UCLASS()
-class MESHMODELINGTOOLSEXP_API UTransformMeshesToolBuilder : public UInteractiveToolBuilder
+class MESHMODELINGTOOLSEXP_API UTransformMeshesToolBuilder : public UMultiSelectionMeshEditingToolBuilder
 {
 	GENERATED_BODY()
 
 public:
-	virtual bool CanBuildTool(const FToolBuilderState& SceneState) const override;
-	virtual UInteractiveTool* BuildTool(const FToolBuilderState& SceneState) const override;
+	virtual UMultiSelectionMeshEditingTool* CreateNewTool(const FToolBuilderState& SceneState) const override;
 
 protected:
-	virtual const FToolTargetTypeRequirements& GetTargetRequirements() const;
+	virtual const FToolTargetTypeRequirements& GetTargetRequirements() const override;
 };
 
 
@@ -149,7 +148,7 @@ struct FTransformMeshesTarget
  *
  */
 UCLASS()
-class MESHMODELINGTOOLSEXP_API UTransformMeshesTool : public UMultiSelectionTool, public IClickDragBehaviorTarget
+class MESHMODELINGTOOLSEXP_API UTransformMeshesTool : public UMultiSelectionMeshEditingTool, public IClickDragBehaviorTarget
 {
 	GENERATED_BODY()
 
@@ -158,8 +157,6 @@ public:
 
 	virtual void RegisterActions(FInteractiveToolActionSet& ActionSet) override;
 
-	virtual void SetWorld(UWorld* World, UInteractiveGizmoManager* GizmoManager);
-
 	virtual void Setup() override;
 	virtual void Shutdown(EToolShutdownType ShutdownType) override;
 
@@ -182,9 +179,6 @@ protected:
 	TObjectPtr<UTransformMeshesToolProperties> TransformProps;
 
 protected:
-	UWorld* TargetWorld;
-	UInteractiveGizmoManager* GizmoManager;
-
 	UPROPERTY()
 	TArray<FTransformMeshesTarget> ActiveGizmos;
 
diff --git a/Engine/Plugins/Experimental/MeshModelingToolsetExp/Source/MeshModelingToolsExp/Public/UVLayoutTool.h b/Engine/Plugins/Experimental/MeshModelingToolsetExp/Source/MeshModelingToolsExp/Public/UVLayoutTool.h
--- a/Engine/Plugins/Experimental/MeshModelingToolsetExp/Source/MeshModelingToolsExp/Public/UVLayoutTool.h
+++ b/Engine/Plugins/Experimental/MeshModelingToolsetExp/Source/MeshModelingToolsExp/Public/UVLayoutTool.h
@@ -4,7 +4,7 @@
 
 #include "CoreMinimal.h"
 #include "UObject/NoExportTypes.h"
-#include "MultiSelectionTool.h"
+#include "BaseTools/MultiSelectionMeshEditingTool.h"
 #include "GeometryBase.h"
 #include "InteractiveToolBuilder.h"
 #include "MeshOpPreviewHelpers.h"
@@ -28,16 +28,12 @@ PREDECLARE_GEOMETRY(class FDynamicMesh3);
  *
  */
 UCLASS()
-class MESHMODELINGTOOLSEXP_API UUVLayoutToolBuilder : public UInteractiveToolBuilder
+class MESHMODELINGTOOLSEXP_API UUVLayoutToolBuilder : public UMultiSelectionMeshEditingToolBuilder
 {
 	GENERATED_BODY()
 
 public:
-	virtual bool CanBuildTool(const FToolBuilderState& SceneState) const override;
-	virtual UInteractiveTool* BuildTool(const FToolBuilderState& SceneState) const override;
-
-protected:
-	virtual const FToolTargetTypeRequirements& GetTargetRequirements() const;
+	virtual UMultiSelectionMeshEditingTool* CreateNewTool(const FToolBuilderState& SceneState) const override;
 };
 
 
@@ -45,7 +41,7 @@ protected:
  * The level editor version of the UV layout tool.
  */
 UCLASS()
-class MESHMODELINGTOOLSEXP_API UUVLayoutTool : public UMultiSelectionTool
+class MESHMODELINGTOOLSEXP_API UUVLayoutTool : public UMultiSelectionMeshEditingTool
 {
 	GENERATED_BODY()
 
@@ -56,8 +52,6 @@ public:
 	virtual void Setup() override;
 	virtual void Shutdown(EToolShutdownType ShutdownType) override;
 
-	virtual void SetWorld(UWorld* World, UInteractiveGizmoManager* GizmoManagerIn);
-
 	virtual void OnTick(float DeltaTime) override;
 	virtual void Render(IToolsContextRenderAPI* RenderAPI) override;
 
@@ -88,8 +82,6 @@ protected:
 
 	TArray<TSharedPtr<UE::Geometry::FDynamicMesh3, ESPMode::ThreadSafe>> OriginalDynamicMeshes;
 
-	UWorld* TargetWorld = nullptr;
-
 	FViewCameraState CameraState;
 
 	void UpdateNumPreviews();
diff --git a/Engine/Plugins/Experimental/MeshModelingToolsetExp/Source/MeshModelingToolsExp/Public/VoxelBlendMeshesTool.h b/Engine/Plugins/Experimental/MeshModelingToolsetExp/Source/MeshModelingToolsExp/Public/VoxelBlendMeshesTool.h
--- a/Engine/Plugins/Experimental/MeshModelingToolsetExp/Source/MeshModelingToolsExp/Public/VoxelBlendMeshesTool.h
+++ b/Engine/Plugins/Experimental/MeshModelingToolsetExp/Source/MeshModelingToolsExp/Public/VoxelBlendMeshesTool.h
@@ -95,9 +95,9 @@ class MESHMODELINGTOOLSEXP_API UVoxelBlendMeshesToolBuilder : public UBaseCreate
 public:
 	virtual int32 MinComponentsSupported() const override { return 2; }
 
-	virtual UBaseCreateFromSelectedTool* MakeNewToolInstance(UObject* Outer) const override
+	virtual UMultiSelectionMeshEditingTool* CreateNewTool(const FToolBuilderState& SceneState) const override
 	{
-		return NewObject<UVoxelBlendMeshesTool>(Outer);
+		return NewObject<UVoxelBlendMeshesTool>(SceneState.ToolManager);
 	}
 };
 
diff --git a/Engine/Plugins/Experimental/MeshModelingToolsetExp/Source/MeshModelingToolsExp/Public/VoxelMorphologyMeshesTool.h b/Engine/Plugins/Experimental/MeshModelingToolsetExp/Source/MeshModelingToolsExp/Public/VoxelMorphologyMeshesTool.h
--- a/Engine/Plugins/Experimental/MeshModelingToolsetExp/Source/MeshModelingToolsExp/Public/VoxelMorphologyMeshesTool.h
+++ b/Engine/Plugins/Experimental/MeshModelingToolsetExp/Source/MeshModelingToolsExp/Public/VoxelMorphologyMeshesTool.h
@@ -79,9 +79,9 @@ class MESHMODELINGTOOLSEXP_API UVoxelMorphologyMeshesToolBuilder : public UBaseC
 	GENERATED_BODY()
 
 public:
-	virtual UBaseCreateFromSelectedTool* MakeNewToolInstance(UObject* Outer) const override
+	virtual UMultiSelectionMeshEditingTool* CreateNewTool(const FToolBuilderState& SceneState) const override
 	{
-		return NewObject<UVoxelMorphologyMeshesTool>(Outer);
+		return NewObject<UVoxelMorphologyMeshesTool>(SceneState.ToolManager);
 	}
 };
 
diff --git a/Engine/Plugins/Experimental/MeshModelingToolsetExp/Source/MeshModelingToolsExp/Public/VoxelSolidifyMeshesTool.h b/Engine/Plugins/Experimental/MeshModelingToolsetExp/Source/MeshModelingToolsExp/Public/VoxelSolidifyMeshesTool.h
--- a/Engine/Plugins/Experimental/MeshModelingToolsetExp/Source/MeshModelingToolsExp/Public/VoxelSolidifyMeshesTool.h
+++ b/Engine/Plugins/Experimental/MeshModelingToolsetExp/Source/MeshModelingToolsExp/Public/VoxelSolidifyMeshesTool.h
@@ -83,9 +83,9 @@ class MESHMODELINGTOOLSEXP_API UVoxelSolidifyMeshesToolBuilder : public UBaseCre
 	GENERATED_BODY()
 
 public:
-	virtual UBaseCreateFromSelectedTool* MakeNewToolInstance(UObject* Outer) const override
+	virtual UMultiSelectionMeshEditingTool* CreateNewTool(const FToolBuilderState& SceneState) const override
 	{
-		return NewObject<UVoxelSolidifyMeshesTool>(Outer);
+		return NewObject<UVoxelSolidifyMeshesTool>(SceneState.ToolManager);
 	}
 };
 
diff --git a/Engine/Plugins/Runtime/MeshModelingToolset/Source/MeshModelingTools/Private/CSGMeshesTool.cpp b/Engine/Plugins/Runtime/MeshModelingToolset/Source/MeshModelingTools/Private/CSGMeshesTool.cpp
--- a/Engine/Plugins/Runtime/MeshModelingToolset/Source/MeshModelingTools/Private/CSGMeshesTool.cpp
+++ b/Engine/Plugins/Runtime/MeshModelingToolset/Source/MeshModelingTools/Private/CSGMeshesTool.cpp
@@ -3,10 +3,14 @@
 #include "CSGMeshesTool.h"
 #include "MeshDescriptionToDynamicMesh.h"
 #include "ToolSetupUtil.h"
+#include "ModelingToolTargetUtil.h"
 #include "BaseGizmos/TransformProxy.h"
 #include "CompositionOps/BooleanMeshesOp.h"
 #include "DynamicMesh/DynamicMesh3.h"
 
+#include "TargetInterfaces/MeshDescriptionProvider.h"
+#include "TargetInterfaces/PrimitiveComponentBackedTarget.h"
+
 using namespace UE::Geometry;
 
 #define LOCTEXT_NAMESPACE "UCSGMeshesTool"
@@ -178,8 +182,7 @@ void UCSGMeshesTool::ConvertInputsAndSetPreviewMaterials(bool bSetPreviewMesh)
 	{
 		for (int ComponentIdx = 0; ComponentIdx < Targets.Num(); ComponentIdx++)
 		{
-			FComponentMaterialSet ComponentMaterialSet;
-			TargetMaterialInterface(ComponentIdx)->GetMaterialSet(ComponentMaterialSet);
+			const FComponentMaterialSet ComponentMaterialSet = UE::ToolTarget::GetMaterialSet(Targets[ComponentIdx]);
 			for (UMaterialInterface* Mat : ComponentMaterialSet.Materials)
 			{
 				int* FoundMatIdx = KnownMaterials.Find(Mat);
@@ -199,14 +202,14 @@ void UCSGMeshesTool::ConvertInputsAndSetPreviewMaterials(bool bSetPreviewMesh)
 	}
 	else
 	{
-		TargetMaterialInterface(0)->GetMaterialSet(AllMaterialSet);
+		AllMaterialSet = UE::ToolTarget::GetMaterialSet(Targets[0]);
 		for (int MatIdx = 0; MatIdx < AllMaterialSet.Materials.Num(); MatIdx++)
 		{
 			MaterialRemap[0].Add(MatIdx);
 		}
 		for (int ComponentIdx = 1; ComponentIdx < Targets.Num(); ComponentIdx++)
 		{
-			MaterialRemap[ComponentIdx].Init(0, TargetMaterialInterface(ComponentIdx)->GetNumMaterials());
+			MaterialRemap[ComponentIdx].Init(0, Cast<IMaterialProvider>(Targets[ComponentIdx])->GetNumMaterials());
 		}
 	}
 
@@ -215,7 +218,7 @@ void UCSGMeshesTool::ConvertInputsAndSetPreviewMaterials(bool bSetPreviewMesh)
 	{
 		OriginalDynamicMeshes[ComponentIdx] = MakeShared<FDynamicMesh3, ESPMode::ThreadSafe>();
 		FMeshDescriptionToDynamicMesh Converter;
-		Converter.Convert(TargetMeshProviderInterface(ComponentIdx)->GetMeshDescription(), *OriginalDynamicMeshes[ComponentIdx]);
+		Converter.Convert(UE::ToolTarget::GetMeshDescription(Targets[ComponentIdx]), *OriginalDynamicMeshes[ComponentIdx]);
 
 		// ensure materials and attributes are always enabled
 		OriginalDynamicMeshes[ComponentIdx]->EnableAttributes();
@@ -229,7 +232,7 @@ void UCSGMeshesTool::ConvertInputsAndSetPreviewMaterials(bool bSetPreviewMesh)
 		if (bSetPreviewMesh)
 		{
 			UPreviewMesh* OriginalMeshPreview = OriginalMeshPreviews.Add_GetRef(NewObject<UPreviewMesh>());
-			OriginalMeshPreview->CreateInWorld(TargetWorld, TargetComponentInterface(ComponentIdx)->GetWorldTransform());
+			OriginalMeshPreview->CreateInWorld(TargetWorld, (FTransform) UE::ToolTarget::GetLocalToWorldTransform(Targets[ComponentIdx]));
 			OriginalMeshPreview->UpdatePreview(OriginalDynamicMeshes[ComponentIdx].Get());
 
 			OriginalMeshPreview->SetMaterial(0, PreviewsGhostMaterial);
diff --git a/Engine/Plugins/Runtime/MeshModelingToolset/Source/MeshModelingTools/Private/CombineMeshesTool.cpp b/Engine/Plugins/Runtime/MeshModelingToolset/Source/MeshModelingTools/Private/CombineMeshesTool.cpp
--- a/Engine/Plugins/Runtime/MeshModelingToolset/Source/MeshModelingTools/Private/CombineMeshesTool.cpp
+++ b/Engine/Plugins/Runtime/MeshModelingToolset/Source/MeshModelingTools/Private/CombineMeshesTool.cpp
@@ -32,18 +32,6 @@ using namespace UE::Geometry;
  * ToolBuilder
  */
 
-
-const FToolTargetTypeRequirements& UCombineMeshesToolBuilder::GetTargetRequirements() const
-{
-	static FToolTargetTypeRequirements TypeRequirements({
-		UMeshDescriptionCommitter::StaticClass(),
-		UMeshDescriptionProvider::StaticClass(),
-		UPrimitiveComponentBackedTarget::StaticClass(),
-		UMaterialProvider::StaticClass()
-		});
-	return TypeRequirements;
-}
-
 bool UCombineMeshesToolBuilder::CanBuildTool(const FToolBuilderState& SceneState) const
 {
 	return (bIsDuplicateTool) ?
@@ -51,15 +39,10 @@ bool UCombineMeshesToolBuilder::CanBuildTool(const FToolBuilderState& SceneState
 		: (SceneState.TargetManager->CountSelectedAndTargetable(SceneState, GetTargetRequirements()) > 1);
 }
 
-UInteractiveTool* UCombineMeshesToolBuilder::BuildTool(const FToolBuilderState& SceneState) const
+UMultiSelectionMeshEditingTool* UCombineMeshesToolBuilder::CreateNewTool(const FToolBuilderState& SceneState) const
 {
 	UCombineMeshesTool* NewTool = NewObject<UCombineMeshesTool>(SceneState.ToolManager);
-
-	TArray<TObjectPtr<UToolTarget>> Targets = SceneState.TargetManager->BuildAllSelectedTargetable(SceneState, GetTargetRequirements());
-	NewTool->SetTargets(MoveTemp(Targets));
-	NewTool->SetWorld(SceneState.World);
 	NewTool->SetDuplicateMode(bIsDuplicateTool);
-
 	return NewTool;
 }
 
@@ -69,11 +52,6 @@ UInteractiveTool* UCombineMeshesToolBuilder::BuildTool(const FToolBuilderState&
  */
 
 
-void UCombineMeshesTool::SetWorld(UWorld* World)
-{
-	this->TargetWorld = World;
-}
-
 void UCombineMeshesTool::SetDuplicateMode(bool bDuplicateModeIn)
 {
 	this->bDuplicateMode = bDuplicateModeIn;
diff --git a/Engine/Plugins/Runtime/MeshModelingToolset/Source/MeshModelingTools/Private/CutMeshWithMeshTool.cpp b/Engine/Plugins/Runtime/MeshModelingToolset/Source/MeshModelingTools/Private/CutMeshWithMeshTool.cpp
--- a/Engine/Plugins/Runtime/MeshModelingToolset/Source/MeshModelingTools/Private/CutMeshWithMeshTool.cpp
+++ b/Engine/Plugins/Runtime/MeshModelingToolset/Source/MeshModelingTools/Private/CutMeshWithMeshTool.cpp
@@ -82,8 +82,7 @@ void UCutMeshWithMeshTool::ConvertInputsAndSetPreviewMaterials(bool bSetPreviewM
 		TMap<UMaterialInterface*, int> KnownMaterials;
 		for (int ComponentIdx = 0; ComponentIdx < Targets.Num(); ComponentIdx++)
 		{
-			FComponentMaterialSet ComponentMaterialSet;
-			TargetMaterialInterface(ComponentIdx)->GetMaterialSet(ComponentMaterialSet);
+			const FComponentMaterialSet ComponentMaterialSet = UE::ToolTarget::GetMaterialSet(Targets[ComponentIdx]);
 			for (UMaterialInterface* Mat : ComponentMaterialSet.Materials)
 			{
 				int* FoundMatIdx = KnownMaterials.Find(Mat);
@@ -103,14 +102,14 @@ void UCutMeshWithMeshTool::ConvertInputsAndSetPreviewMaterials(bool bSetPreviewM
 	}
 	else
 	{
-		TargetMaterialInterface(0)->GetMaterialSet(AllMaterialSet);
+		AllMaterialSet = UE::ToolTarget::GetMaterialSet(Targets[0]);
 		for (int MatIdx = 0; MatIdx < AllMaterialSet.Materials.Num(); MatIdx++)
 		{
 			MaterialRemap[0].Add(MatIdx);
 		}
 		for (int ComponentIdx = 1; ComponentIdx < Targets.Num(); ComponentIdx++)
 		{
-			MaterialRemap[ComponentIdx].Init(0, TargetMaterialInterface(ComponentIdx)->GetNumMaterials());
+			MaterialRemap[ComponentIdx].Init(0, Cast<IMaterialProvider>(Targets[ComponentIdx])->GetNumMaterials());
 		}
 	}
 
@@ -118,7 +117,7 @@ void UCutMeshWithMeshTool::ConvertInputsAndSetPreviewMaterials(bool bSetPreviewM
 	{
 		TSharedPtr<FDynamicMesh3, ESPMode::ThreadSafe> Mesh = MakeShared<FDynamicMesh3, ESPMode::ThreadSafe>();
 		FMeshDescriptionToDynamicMesh Converter;
-		Converter.Convert(TargetMeshProviderInterface(ComponentIdx)->GetMeshDescription(), *Mesh);
+		Converter.Convert(UE::ToolTarget::GetMeshDescription(Targets[ComponentIdx]), *Mesh);
 
 		// ensure materials and attributes are always enabled
 		Mesh->EnableAttributes();
@@ -338,7 +337,7 @@ void UCutMeshWithMeshTool::Shutdown(EToolShutdownType ShutdownType)
 	// Restore (unhide) the source meshes
 	for ( int32 ci = 0; ci < Targets.Num(); ++ci)
 	{
-		TargetComponentInterface(ci)->SetOwnerVisibility(true);
+		UE::ToolTarget::ShowSourceObject(Targets[ci]);
 	}
 
 	if (ShutdownType == EToolShutdownType::Accept)
@@ -351,22 +350,21 @@ void UCutMeshWithMeshTool::Shutdown(EToolShutdownType ShutdownType)
 		MaterialSet.Materials = GetOutputMaterials();
 
 		// update subtract asset
-		IPrimitiveComponentBackedTarget* UpdateTarget = TargetComponentInterface(0);
-		FTransform3d TargetToWorld = (FTransform3d)UpdateTarget->GetWorldTransform();
+		FTransform3d TargetToWorld = UE::ToolTarget::GetLocalToWorldTransform(Targets[0]);
 		{
 			if (OpResult.Mesh->TriangleCount() > 0)
 			{
 				MeshTransforms::ApplyTransform(*OpResult.Mesh, OpResult.Transform);
 				MeshTransforms::ApplyTransformInverse(*OpResult.Mesh, TargetToWorld);
-				TargetMeshCommitterInterface(0)->CommitMeshDescription([&](const IMeshDescriptionCommitter::FCommitterParams& CommitParams)
+				Cast<IMeshDescriptionCommitter>(Targets[0])->CommitMeshDescription([&](const IMeshDescriptionCommitter::FCommitterParams& CommitParams)
 				{
 					FDynamicMeshToMeshDescription Converter;
 					Converter.Convert(OpResult.Mesh.Get(), *CommitParams.MeshDescriptionOut);
 				});
-				TargetMaterialInterface(0)->CommitMaterialSetUpdate(MaterialSet, true);
+				Cast<IMaterialProvider>(Targets[0])->CommitMaterialSetUpdate(MaterialSet, true);
 			}
 		}
-		SelectActors.Add(UpdateTarget->GetOwnerActor());
+		SelectActors.Add(UE::ToolTarget::GetTargetActor(Targets[0]));
 
 		// create intersection asset
 		if ( IntersectionMesh.TriangleCount() > 0)
@@ -375,7 +373,7 @@ void UCutMeshWithMeshTool::Shutdown(EToolShutdownType ShutdownType)
 			MeshTransforms::ApplyTransformInverse(IntersectionMesh, TargetToWorld);
 			FTransform3d NewTransform = TargetToWorld;
 
-			FString CurName = UE::Modeling::GetComponentAssetBaseName(UpdateTarget->GetOwnerComponent());
+			FString CurName = UE::Modeling::GetComponentAssetBaseName(UE::ToolTarget::GetTargetComponent(Targets[0]));
 			FString UseBaseName = FString::Printf(TEXT("%s_%s"), *CurName, TEXT("CutPart") );
 
 			FCreateMeshObjectParams NewMeshObjectParams;
diff --git a/Engine/Plugins/Runtime/MeshModelingToolset/Source/MeshModelingTools/Public/CSGMeshesTool.h b/Engine/Plugins/Runtime/MeshModelingToolset/Source/MeshModelingTools/Public/CSGMeshesTool.h
--- a/Engine/Plugins/Runtime/MeshModelingToolset/Source/MeshModelingTools/Public/CSGMeshesTool.h
+++ b/Engine/Plugins/Runtime/MeshModelingToolset/Source/MeshModelingTools/Public/CSGMeshesTool.h
@@ -183,9 +183,9 @@ public:
 		return 2;
 	}
 
-	virtual UBaseCreateFromSelectedTool* MakeNewToolInstance(UObject* Outer) const override
+	virtual UMultiSelectionMeshEditingTool* CreateNewTool(const FToolBuilderState& SceneState) const override
 	{
-		UCSGMeshesTool* Tool = NewObject<UCSGMeshesTool>(Outer);
+		UCSGMeshesTool* Tool = NewObject<UCSGMeshesTool>(SceneState.ToolManager);
 		if (bTrimMode)
 		{
 			Tool->EnableTrimMode();
diff --git a/Engine/Plugins/Runtime/MeshModelingToolset/Source/MeshModelingTools/Public/CombineMeshesTool.h b/Engine/Plugins/Runtime/MeshModelingToolset/Source/MeshModelingTools/Public/CombineMeshesTool.h
--- a/Engine/Plugins/Runtime/MeshModelingToolset/Source/MeshModelingTools/Public/CombineMeshesTool.h
+++ b/Engine/Plugins/Runtime/MeshModelingToolset/Source/MeshModelingTools/Public/CombineMeshesTool.h
@@ -3,7 +3,7 @@
 #pragma once
 
 #include "CoreMinimal.h"
-#include "MultiSelectionTool.h"
+#include "BaseTools/MultiSelectionMeshEditingTool.h"
 #include "InteractiveToolBuilder.h"
 #include "MeshOpPreviewHelpers.h"
 #include "DynamicMesh/DynamicMesh3.h"
@@ -19,7 +19,7 @@ struct FMeshDescription;
  *
  */
 UCLASS()
-class MESHMODELINGTOOLS_API UCombineMeshesToolBuilder : public UInteractiveToolBuilder
+class MESHMODELINGTOOLS_API UCombineMeshesToolBuilder : public UMultiSelectionMeshEditingToolBuilder
 {
 	GENERATED_BODY()
 
@@ -27,10 +27,7 @@ public:
 	bool bIsDuplicateTool = false;
 
 	virtual bool CanBuildTool(const FToolBuilderState& SceneState) const override;
-	virtual UInteractiveTool* BuildTool(const FToolBuilderState& SceneState) const override;
-
-protected:
-	virtual const FToolTargetTypeRequirements& GetTargetRequirements() const;
+	virtual UMultiSelectionMeshEditingTool* CreateNewTool(const FToolBuilderState& SceneState) const override;
 };
 
 
@@ -67,7 +64,7 @@ public:
  * Simple tool to combine multiple meshes into a single mesh asset
  */
 UCLASS()
-class MESHMODELINGTOOLS_API UCombineMeshesTool : public UMultiSelectionTool
+class MESHMODELINGTOOLS_API UCombineMeshesTool : public UMultiSelectionMeshEditingTool
 {
 	GENERATED_BODY()
 
@@ -77,8 +74,6 @@ public:
 	virtual void Setup() override;
 	virtual void Shutdown(EToolShutdownType ShutdownType) override;
 
-	virtual void SetWorld(UWorld* World);
-
 	virtual bool HasCancel() const override { return true; }
 	virtual bool HasAccept() const override { return true; }
 
@@ -93,8 +88,6 @@ protected:
 	UPROPERTY()
 	TObjectPtr<UOnAcceptHandleSourcesPropertiesBase> HandleSourceProperties;
 
-	UWorld* TargetWorld;
-
 	bool bDuplicateMode;
 
 	void CreateNewAsset();
diff --git a/Engine/Plugins/Runtime/MeshModelingToolset/Source/MeshModelingTools/Public/CutMeshWithMeshTool.h b/Engine/Plugins/Runtime/MeshModelingToolset/Source/MeshModelingTools/Public/CutMeshWithMeshTool.h
--- a/Engine/Plugins/Runtime/MeshModelingToolset/Source/MeshModelingTools/Public/CutMeshWithMeshTool.h
+++ b/Engine/Plugins/Runtime/MeshModelingToolset/Source/MeshModelingTools/Public/CutMeshWithMeshTool.h
@@ -4,7 +4,6 @@
 
 #include "CoreMinimal.h"
 #include "UObject/NoExportTypes.h"
-#include "InteractiveToolBuilder.h"
 #include "Drawing/LineSetComponent.h"
 #include "BaseTools/BaseCreateFromSelectedTool.h"
 
@@ -118,9 +117,9 @@ public:
 		return 2;
 	}
 
-	virtual UBaseCreateFromSelectedTool* MakeNewToolInstance(UObject* Outer) const override
+	virtual UMultiSelectionMeshEditingTool* CreateNewTool(const FToolBuilderState& SceneState) const override
 	{
-		return NewObject<UCutMeshWithMeshTool>(Outer);
+		return NewObject<UCutMeshWithMeshTool>(SceneState.ToolManager);
 	}
 };
 
diff --git a/Engine/Plugins/Runtime/MeshModelingToolset/Source/ModelingComponents/Private/BaseTools/BaseCreateFromSelectedTool.cpp b/Engine/Plugins/Runtime/MeshModelingToolset/Source/ModelingComponents/Private/BaseTools/BaseCreateFromSelectedTool.cpp
--- a/Engine/Plugins/Runtime/MeshModelingToolset/Source/ModelingComponents/Private/BaseTools/BaseCreateFromSelectedTool.cpp
+++ b/Engine/Plugins/Runtime/MeshModelingToolset/Source/ModelingComponents/Private/BaseTools/BaseCreateFromSelectedTool.cpp
@@ -29,17 +29,6 @@ using namespace UE::Geometry;
  * ToolBuilder
  */
 
-const FToolTargetTypeRequirements& UBaseCreateFromSelectedToolBuilder::GetTargetRequirements() const
-{
-	static FToolTargetTypeRequirements TypeRequirements({
-		UMeshDescriptionCommitter::StaticClass(),
-		UMeshDescriptionProvider::StaticClass(),
-		UPrimitiveComponentBackedTarget::StaticClass(),
-		UMaterialProvider::StaticClass()
-		});
-	return TypeRequirements;
-}
-
 bool UBaseCreateFromSelectedToolBuilder::CanBuildTool(const FToolBuilderState& SceneState) const
 {
 	int32 ComponentCount = SceneState.TargetManager->CountSelectedAndTargetable(SceneState, GetTargetRequirements());
@@ -47,19 +36,6 @@ bool UBaseCreateFromSelectedToolBuilder::CanBuildTool(const FToolBuilderState& S
 }
 
 
-UInteractiveTool* UBaseCreateFromSelectedToolBuilder::BuildTool(const FToolBuilderState& SceneState) const
-{
-	UBaseCreateFromSelectedTool* NewTool = MakeNewToolInstance(SceneState.ToolManager);
-
-	TArray<TObjectPtr<UToolTarget>> Targets = SceneState.TargetManager->BuildAllSelectedTargetable(SceneState, GetTargetRequirements());
-	NewTool->SetTargets(MoveTemp(Targets));
-	NewTool->SetWorld(SceneState.World);
-
-	return NewTool;
-}
-
-
-
 
 /*
  * Tool
@@ -72,7 +48,7 @@ void UBaseCreateFromSelectedTool::Setup()
 	// hide input StaticMeshComponents
 	for (int32 ComponentIdx = 0; ComponentIdx < Targets.Num(); ComponentIdx++)
 	{
-		TargetComponentInterface(ComponentIdx)->SetOwnerVisibility(false);
+		UE::ToolTarget::HideSourceObject(Targets[ComponentIdx]);
 	}
 
 	// initialize our properties
@@ -122,10 +98,10 @@ void UBaseCreateFromSelectedTool::Setup()
 		else
 		{
 			int32 Index = (HandleSourcesProperties->OutputWriteTo == EBaseCreateFromSelectedTargetType::FirstInputObject) ? 0 : Targets.Num() - 1;
-			HandleSourcesProperties->OutputExistingName = UE::Modeling::GetComponentAssetBaseName(TargetComponentInterface(Index)->GetOwnerComponent(), false);
+			HandleSourcesProperties->OutputExistingName = UE::Modeling::GetComponentAssetBaseName(UE::ToolTarget::GetTargetComponent(Targets[Index]), false);
 
 			// Reset the hidden gizmo to its initial position
-			FTransform ComponentTransform = TargetComponentInterface(Index)->GetWorldTransform();
+			const FTransform ComponentTransform = (FTransform) UE::ToolTarget::GetLocalToWorldTransform(Targets[Index]);
 			TransformGizmos[Index]->SetNewGizmoTransform(ComponentTransform, true);
 			UpdateGizmoVisibility();
 		}
@@ -148,12 +124,6 @@ int32 UBaseCreateFromSelectedTool::GetHiddenGizmoIndex() const
 }
 
 
-void UBaseCreateFromSelectedTool::SetWorld(UWorld* World)
-{
-	this->TargetWorld = World;
-}
-
-
 void UBaseCreateFromSelectedTool::OnTick(float DeltaTime)
 {
 	Preview->Tick(DeltaTime);
@@ -180,7 +150,7 @@ void UBaseCreateFromSelectedTool::SetTransformGizmos()
 	{
 		UTransformProxy* Proxy = TransformProxies.Add_GetRef(NewObject<UTransformProxy>(this));
 		UCombinedTransformGizmo* Gizmo = TransformGizmos.Add_GetRef(UE::TransformGizmoUtil::Create3AxisTransformGizmo(GizmoManager, this));
-		Proxy->SetTransform(TargetComponentInterface(ComponentIdx)->GetWorldTransform());
+		Proxy->SetTransform((FTransform) UE::ToolTarget::GetLocalToWorldTransform(Targets[ComponentIdx]));
 		Gizmo->SetActiveTarget(Proxy, GetToolManager());
 		Proxy->OnTransformChanged.AddUObject(this, &UBaseCreateFromSelectedTool::TransformChanged);
 	}
@@ -213,9 +183,9 @@ void UBaseCreateFromSelectedTool::GenerateAsset(const FDynamicMeshOpResult& OpRe
 	FTransform3d NewTransform;
 	if (Targets.Num() == 1) // in the single-selection case, shove the result back into the original component space
 	{
-		FTransform3d ToSourceComponentSpace = (FTransform3d)TargetComponentInterface(0)->GetWorldTransform().Inverse();
+		FTransform3d ToSourceComponentSpace = UE::ToolTarget::GetLocalToWorldTransform(Targets[0]).Inverse();
 		MeshTransforms::ApplyTransform(*OpResult.Mesh, ToSourceComponentSpace);
-		NewTransform = (FTransform3d)TargetComponentInterface(0)->GetWorldTransform();
+		NewTransform = UE::ToolTarget::GetLocalToWorldTransform(Targets[0]);
 	}
 	else // in the multi-selection case, center the pivot for the combined result
 	{
@@ -263,17 +233,13 @@ void UBaseCreateFromSelectedTool::UpdateAsset(const FDynamicMeshOpResult& Result
 {
 	check(Result.Mesh.Get() != nullptr);
 
-	IPrimitiveComponentBackedTarget* TargetComponent = Cast<IPrimitiveComponentBackedTarget>(UpdateTarget);
-	IMeshDescriptionCommitter* TargetMeshCommitter = Cast<IMeshDescriptionCommitter>(UpdateTarget);
-	IMaterialProvider* TargetMaterial = Cast<IMaterialProvider>(UpdateTarget);
-
-	FTransform3d TargetToWorld = (FTransform3d)TargetComponent->GetWorldTransform();
+	FTransform3d TargetToWorld = UE::ToolTarget::GetLocalToWorldTransform(UpdateTarget);
 
 	FTransform3d ResultTransform = Result.Transform;
 	MeshTransforms::ApplyTransform(*Result.Mesh, ResultTransform);
 	MeshTransforms::ApplyTransformInverse(*Result.Mesh, TargetToWorld);
 
-	TargetMeshCommitter->CommitMeshDescription([&](const IMeshDescriptionCommitter::FCommitterParams& CommitParams)
+	Cast<IMeshDescriptionCommitter>(UpdateTarget)->CommitMeshDescription([&](const IMeshDescriptionCommitter::FCommitterParams& CommitParams)
 	{
 		FDynamicMeshToMeshDescription Converter;
 		Converter.Convert(Result.Mesh.Get(), *CommitParams.MeshDescriptionOut);
@@ -281,7 +247,7 @@ void UBaseCreateFromSelectedTool::UpdateAsset(const FDynamicMeshOpResult& Result
 
 	FComponentMaterialSet MaterialSet;
 	MaterialSet.Materials = GetOutputMaterials();
-	TargetMaterial->CommitMaterialSetUpdate(MaterialSet, true);
+	Cast<IMaterialProvider>(UpdateTarget)->CommitMaterialSetUpdate(MaterialSet, true);
 }
 
 
@@ -289,7 +255,7 @@ FString UBaseCreateFromSelectedTool::PrefixWithSourceNameIfSingleSelection(const
 {
 	if (Targets.Num() == 1)
 	{
-		FString CurName = UE::Modeling::GetComponentAssetBaseName(TargetComponentInterface(0)->GetOwnerComponent());
+		FString CurName = UE::Modeling::GetComponentAssetBaseName(UE::ToolTarget::GetTargetComponent(Targets[0]));
 		return FString::Printf(TEXT("%s_%s"), *CurName, *AssetName);
 	}
 	else
@@ -329,7 +295,7 @@ void UBaseCreateFromSelectedTool::Shutdown(EToolShutdownType ShutdownType)
 	// Restore (unhide) the source meshes
 	for (int32 ComponentIdx = 0; ComponentIdx < Targets.Num(); ComponentIdx++)
 	{
-		TargetComponentInterface(ComponentIdx)->SetOwnerVisibility(true);
+		UE::ToolTarget::ShowSourceObject(Targets[ComponentIdx]);
 	}
 
 	if (ShutdownType == EToolShutdownType::Accept)
@@ -345,7 +311,7 @@ void UBaseCreateFromSelectedTool::Shutdown(EToolShutdownType ShutdownType)
 		else
 		{
 			int32 TargetIndex = (HandleSourcesProperties->OutputWriteTo == EBaseCreateFromSelectedTargetType::FirstInputObject) ? 0 : (Targets.Num() - 1);
-			KeepActor = TargetComponentInterface(TargetIndex)->GetOwnerActor();
+			KeepActor = UE::ToolTarget::GetTargetActor(Targets[TargetIndex]);
 
 			UpdateAsset(Result, Targets[TargetIndex]);
 		}
@@ -353,7 +319,7 @@ void UBaseCreateFromSelectedTool::Shutdown(EToolShutdownType ShutdownType)
 		TArray<AActor*> Actors;
 		for (int32 ComponentIdx = 0; ComponentIdx < Targets.Num(); ComponentIdx++)
 		{
-			AActor* Actor = TargetComponentInterface(ComponentIdx)->GetOwnerActor();
+			AActor* Actor = UE::ToolTarget::GetTargetActor(Targets[ComponentIdx]);
 			Actors.Add(Actor);
 		}
 		HandleSourcesProperties->ApplyMethod(Actors, GetToolManager(), KeepActor);
diff --git a/Engine/Plugins/Runtime/MeshModelingToolset/Source/ModelingComponents/Private/BaseTools/BaseVoxelTool.cpp b/Engine/Plugins/Runtime/MeshModelingToolset/Source/ModelingComponents/Private/BaseTools/BaseVoxelTool.cpp
--- a/Engine/Plugins/Runtime/MeshModelingToolset/Source/ModelingComponents/Private/BaseTools/BaseVoxelTool.cpp
+++ b/Engine/Plugins/Runtime/MeshModelingToolset/Source/ModelingComponents/Private/BaseTools/BaseVoxelTool.cpp
@@ -3,9 +3,12 @@
 #include "BaseTools/BaseVoxelTool.h"
 #include "InteractiveToolManager.h"
 #include "ToolSetupUtil.h"
+#include "ModelingToolTargetUtil.h"
 #include "DynamicMesh/DynamicMesh3.h"
 #include "Materials/MaterialInstanceDynamic.h"
 
+#include "TargetInterfaces/MeshDescriptionProvider.h"
+
 #include "MeshDescriptionToDynamicMesh.h"
 
 
@@ -45,7 +48,7 @@ void UBaseVoxelTool::ConvertInputsAndSetPreviewMaterials(bool bSetPreviewMesh)
 	{
 		OriginalDynamicMeshes[ComponentIdx] = MakeShared<FDynamicMesh3, ESPMode::ThreadSafe>();
 		FMeshDescriptionToDynamicMesh Converter;
-		Converter.Convert(TargetMeshProviderInterface(ComponentIdx)->GetMeshDescription(), *OriginalDynamicMeshes[ComponentIdx]);
+		Converter.Convert(UE::ToolTarget::GetMeshDescription(Targets[ComponentIdx]), *OriginalDynamicMeshes[ComponentIdx]);
 	}
 
 	//if (bSetPreviewMesh)
diff --git a/Engine/Plugins/Runtime/MeshModelingToolset/Source/ModelingComponents/Private/BaseTools/MultiSelectionMeshEditingTool.cpp b/Engine/Plugins/Runtime/MeshModelingToolset/Source/ModelingComponents/Private/BaseTools/MultiSelectionMeshEditingTool.cpp
new file mode 100644
--- /dev/null
+++ b/Engine/Plugins/Runtime/MeshModelingToolset/Source/ModelingComponents/Private/BaseTools/MultiSelectionMeshEditingTool.cpp
@@ -0,0 +1,86 @@
+// Copyright Epic Games, Inc. All Rights Reserved.
+
+#include "BaseTools/MultiSelectionMeshEditingTool.h"
+
+#include "TargetInterfaces/MaterialProvider.h"
+#include "TargetInterfaces/MeshDescriptionCommitter.h"
+#include "TargetInterfaces/MeshDescriptionProvider.h"
+#include "TargetInterfaces/PrimitiveComponentBackedTarget.h"
+#include "TargetInterfaces/AssetBackedTarget.h"
+#include "ToolTargetManager.h"
+
+/*
+ * ToolBuilder
+ */
+const FToolTargetTypeRequirements& UMultiSelectionMeshEditingToolBuilder::GetTargetRequirements() const
+{
+	static FToolTargetTypeRequirements TypeRequirements({
+		UMaterialProvider::StaticClass(),
+		UMeshDescriptionCommitter::StaticClass(),
+		UMeshDescriptionProvider::StaticClass(),
+		UPrimitiveComponentBackedTarget::StaticClass()
+		});
+	return TypeRequirements;
+}
+
+bool UMultiSelectionMeshEditingToolBuilder::CanBuildTool(const FToolBuilderState& SceneState) const
+{
+	return SceneState.TargetManager->CountSelectedAndTargetable(SceneState, GetTargetRequirements()) > 0;
+}
+
+UInteractiveTool* UMultiSelectionMeshEditingToolBuilder::BuildTool(const FToolBuilderState& SceneState) const
+{
+	UMultiSelectionMeshEditingTool* NewTool = CreateNewTool(SceneState);
+	InitializeNewTool(NewTool, SceneState);
+	return NewTool;
+}
+
+void UMultiSelectionMeshEditingToolBuilder::InitializeNewTool(UMultiSelectionMeshEditingTool* NewTool, const FToolBuilderState& SceneState) const
+{
+	const TArray<TObjectPtr<UToolTarget>> Targets = SceneState.TargetManager->BuildAllSelectedTargetable(SceneState, GetTargetRequirements());
+	NewTool->SetTargets(Targets);
+	NewTool->SetWorld(SceneState.World);
+}
+
+
+/**
+ * Tool
+ */
+
+bool UMultiSelectionMeshEditingTool::GetMapToSharedSourceData(TArray<int32>& MapToFirstOccurrences)
+{
+	bool bSharesSources = false;
+	MapToFirstOccurrences.SetNumUninitialized(Targets.Num());
+	for (int32 ComponentIdx = 0; ComponentIdx < Targets.Num(); ComponentIdx++)
+	{
+		MapToFirstOccurrences[ComponentIdx] = -1;
+	}
+	for (int32 ComponentIdx = 0; ComponentIdx < Targets.Num(); ComponentIdx++)
+	{
+		if (MapToFirstOccurrences[ComponentIdx] >= 0) // already mapped
+		{
+			continue;
+		}
+
+		MapToFirstOccurrences[ComponentIdx] = ComponentIdx;
+
+		IAssetBackedTarget* Target = Cast<IAssetBackedTarget>(Targets[ComponentIdx]);
+		if (!Target)
+		{
+			continue;
+		}
+		for (int32 VsIdx = ComponentIdx + 1; VsIdx < Targets.Num(); VsIdx++)
+		{
+			IAssetBackedTarget* OtherTarget = Cast<IAssetBackedTarget>(Targets[VsIdx]);
+			if (OtherTarget && OtherTarget->GetSourceData() == Target->GetSourceData())
+			{
+				bSharesSources = true;
+				MapToFirstOccurrences[VsIdx] = ComponentIdx;
+			}
+		}
+	}
+	return bSharesSources;
+}
+
+
+
diff --git a/Engine/Plugins/Runtime/MeshModelingToolset/Source/ModelingComponents/Public/BaseTools/BaseCreateFromSelectedTool.h b/Engine/Plugins/Runtime/MeshModelingToolset/Source/ModelingComponents/Public/BaseTools/BaseCreateFromSelectedTool.h
--- a/Engine/Plugins/Runtime/MeshModelingToolset/Source/ModelingComponents/Public/BaseTools/BaseCreateFromSelectedTool.h
+++ b/Engine/Plugins/Runtime/MeshModelingToolset/Source/ModelingComponents/Public/BaseTools/BaseCreateFromSelectedTool.h
@@ -3,8 +3,9 @@
 #pragma once
 
 #include "CoreMinimal.h"
-#include "MultiSelectionTool.h"
+#include "BaseTools/MultiSelectionMeshEditingTool.h"
 #include "InteractiveToolBuilder.h"
+#include "InteractiveToolManager.h"
 #include "MeshOpPreviewHelpers.h"
 #include "PropertySets/OnAcceptProperties.h"
 #include "PropertySets/CreateMeshObjectTypeProperties.h"
@@ -12,28 +13,22 @@
 
 class UCombinedTransformGizmo;
 class UTransformProxy;
+class UBaseCreateFromSelectedTool;
 
 /**
  * ToolBuilder for UBaseCreateFromSelectedTool
  */
 UCLASS()
-class MODELINGCOMPONENTS_API UBaseCreateFromSelectedToolBuilder : public UInteractiveToolBuilder
+class MODELINGCOMPONENTS_API UBaseCreateFromSelectedToolBuilder : public UMultiSelectionMeshEditingToolBuilder
 {
 	GENERATED_BODY()
 
 public:
 	virtual bool CanBuildTool(const FToolBuilderState& SceneState) const override;
-	virtual UInteractiveTool* BuildTool(const FToolBuilderState& SceneState) const override;
 
 public:
 	virtual TOptional<int32> MaxComponentsSupported() const { return TOptional<int32>(); }
 	virtual int32 MinComponentsSupported() const { return 1; }
-	
-	// subclass must override!
-	virtual UBaseCreateFromSelectedTool* MakeNewToolInstance(UObject* Outer) const { check(false); return nullptr; }
-
-protected:
-	virtual const FToolTargetTypeRequirements& GetTargetRequirements() const;
 };
 
 
@@ -95,7 +90,7 @@ public:
  * that provides support for common functionality in tools that create a new mesh from a selection of one or more existing meshes
  */
 UCLASS()
-class MODELINGCOMPONENTS_API UBaseCreateFromSelectedTool : public UMultiSelectionTool, public UE::Geometry::IDynamicMeshOperatorFactory
+class MODELINGCOMPONENTS_API UBaseCreateFromSelectedTool : public UMultiSelectionMeshEditingTool, public UE::Geometry::IDynamicMeshOperatorFactory
 {
 	GENERATED_BODY()
 protected:
@@ -104,8 +99,6 @@ protected:
 public:
 	UBaseCreateFromSelectedTool() = default;
 
-	virtual void SetWorld(UWorld* World);
-
 	//
 	// InteractiveTool API - generally does not need to be modified by subclasses
 	//
@@ -205,7 +198,5 @@ protected:
 
 	UPROPERTY()
 	TArray<TObjectPtr<UCombinedTransformGizmo>> TransformGizmos;
-
-	UWorld* TargetWorld = nullptr;
 };
 
diff --git a/Engine/Plugins/Runtime/MeshModelingToolset/Source/ModelingComponents/Public/BaseTools/MultiSelectionMeshEditingTool.h b/Engine/Plugins/Runtime/MeshModelingToolset/Source/ModelingComponents/Public/BaseTools/MultiSelectionMeshEditingTool.h
new file mode 100644
--- /dev/null
+++ b/Engine/Plugins/Runtime/MeshModelingToolset/Source/ModelingComponents/Public/BaseTools/MultiSelectionMeshEditingTool.h
@@ -0,0 +1,61 @@
+// Copyright Epic Games, Inc. All Rights Reserved.
+
+#pragma once
+
+#include "MultiSelectionTool.h"
+#include "InteractiveToolBuilder.h"
+#include "MultiSelectionMeshEditingTool.generated.h"
+
+class UMultiSelectionMeshEditingTool;
+
+/**
+ * UMultiSelectionMeshEditingToolBuilder is a base tool builder for multi
+ * selection tools that define a common set of ToolTarget interfaces required
+ * for editing meshes.
+ */
+UCLASS(Transient, Abstract)
+class MODELINGCOMPONENTS_API UMultiSelectionMeshEditingToolBuilder : public UInteractiveToolBuilder
+{
+	GENERATED_BODY()
+
+public:
+	/** @return true if mesh sources can be found in the active selection */
+	virtual bool CanBuildTool(const FToolBuilderState& SceneState) const override;
+
+	/** @return new Tool instance initialized with selected mesh source(s) */
+	virtual UInteractiveTool* BuildTool(const FToolBuilderState& SceneState) const override;
+
+	/** @return new Tool instance. Override this in subclasses to build a different Tool class type */
+	virtual UMultiSelectionMeshEditingTool* CreateNewTool(const FToolBuilderState& SceneState) const PURE_VIRTUAL(USingleSelectionMeshEditingToolBuilder::CreateNewTool, return nullptr; );
+
+	/** Called by BuildTool to configure the Tool with the input mesh source(s) based on the SceneState */
+	virtual void InitializeNewTool(UMultiSelectionMeshEditingTool* Tool, const FToolBuilderState& SceneState) const;
+
+protected:
+	virtual const FToolTargetTypeRequirements& GetTargetRequirements() const;
+};
+
+
+/**
+ * Multi Selection Mesh Editing tool base class.
+ */
+UCLASS()
+class MODELINGCOMPONENTS_API UMultiSelectionMeshEditingTool : public UMultiSelectionTool
+{
+	GENERATED_BODY()
+public:
+	virtual void SetWorld(UWorld* World) { TargetWorld = World; }
+	virtual UWorld* GetTargetWorld() { return TargetWorld; }
+
+	/**
+	 * Helper to find which targets share source data.
+	 * Requires UAssetBackedTarget as a tool target requirement.
+	 *
+	 * @return Array of indices, 1:1 with Targets, indicating the first index where a component target sharing the same source data appeared.
+	 */
+	bool GetMapToSharedSourceData(TArray<int32>& MapToFirstOccurrences);
+
+protected:
+	UPROPERTY()
+	TObjectPtr<UWorld> TargetWorld = nullptr;
+};
diff --git a/Engine/Source/Runtime/InteractiveToolsFramework/Private/MultiSelectionTool.cpp b/Engine/Source/Runtime/InteractiveToolsFramework/Private/MultiSelectionTool.cpp
--- a/Engine/Source/Runtime/InteractiveToolsFramework/Private/MultiSelectionTool.cpp
+++ b/Engine/Source/Runtime/InteractiveToolsFramework/Private/MultiSelectionTool.cpp
@@ -1,45 +1,8 @@
 // Copyright Epic Games, Inc. All Rights Reserved.
 #include "MultiSelectionTool.h"
-#include "TargetInterfaces/AssetBackedTarget.h"
 #include "ToolTargets/PrimitiveComponentToolTarget.h"
 #include "Components/PrimitiveComponent.h"
 
-bool UMultiSelectionTool::GetMapToSharedSourceData(TArray<int32>& MapToFirstOccurrences)
-{
-	bool bSharesSources = false;
-	MapToFirstOccurrences.SetNumUninitialized(Targets.Num());
-	for (int32 ComponentIdx = 0; ComponentIdx < Targets.Num(); ComponentIdx++)
-	{
-		MapToFirstOccurrences[ComponentIdx] = -1;
-	}
-	for (int32 ComponentIdx = 0; ComponentIdx < Targets.Num(); ComponentIdx++)
-	{
-		if (MapToFirstOccurrences[ComponentIdx] >= 0) // already mapped
-		{
-			continue;
-		}
-
-		MapToFirstOccurrences[ComponentIdx] = ComponentIdx;
-
-		IAssetBackedTarget* Target = Cast<IAssetBackedTarget>(Targets[ComponentIdx]);
-		if (!Target)
-		{
-			continue;
-		}
-		for (int32 VsIdx = ComponentIdx + 1; VsIdx < Targets.Num(); VsIdx++)
-		{
-			IAssetBackedTarget* OtherTarget = Cast<IAssetBackedTarget>(Targets[VsIdx]);
-			if (OtherTarget && OtherTarget->GetSourceData() == Target->GetSourceData())
-			{
-				bSharesSources = true;
-				MapToFirstOccurrences[VsIdx] = ComponentIdx;
-			}
-		}
-	}
-	return bSharesSources;
-}
-
-
 
 bool UMultiSelectionTool::SupportsWorldSpaceFocusBox()
 {
diff --git a/Engine/Source/Runtime/InteractiveToolsFramework/Public/MultiSelectionTool.h b/Engine/Source/Runtime/InteractiveToolsFramework/Public/MultiSelectionTool.h
--- a/Engine/Source/Runtime/InteractiveToolsFramework/Public/MultiSelectionTool.h
+++ b/Engine/Source/Runtime/InteractiveToolsFramework/Public/MultiSelectionTool.h
@@ -7,11 +7,6 @@
 #include "ToolTargets/ToolTarget.h"
 #include "InteractiveToolQueryInterfaces.h"
 
-#include "TargetInterfaces/MaterialProvider.h"
-#include "TargetInterfaces/MeshDescriptionCommitter.h"
-#include "TargetInterfaces/MeshDescriptionProvider.h"
-#include "TargetInterfaces/PrimitiveComponentBackedTarget.h"
-
 #include "MultiSelectionTool.generated.h"
 
 UCLASS(Transient)
@@ -50,45 +45,6 @@ protected:
 	UPROPERTY()
 	TArray<TObjectPtr<UToolTarget>> Targets{};
 
-	/**
-	 * Helper to find which targets share source data.
-	 * Requires UAssetBackedTarget as a tool target requirement.
-	 *
-	 * @return Array of indices, 1:1 with Targets, indicating the first index where a component target sharing the same source data appeared.
-	 */
-	bool GetMapToSharedSourceData(TArray<int32>& MapToFirstOccurrences);
-
-	/**
-	 * Template helper to retrieve an interface from a ToolTarget.
-	 */
-	template <class T>
-	T* TargetInterface(int32 TargetIdx) const
-	{
-		T* Interface = Cast<T>(Targets[TargetIdx]);
-		check(Interface);
-		return Interface;
-	}
-
-	/**
-	 * Helper to cast a Target into the IPrimitiveComponentBackedTarget interface.
-	 */
-	IPrimitiveComponentBackedTarget* TargetComponentInterface(int32 ComponentIdx) const;
-
-	/**
-	 * Helper to cast a Target into the IMeshDescriptionCommitter interface.
-	 */
-	IMeshDescriptionCommitter* TargetMeshCommitterInterface(int32 ComponentIdx) const;
-
-	/**
-	 * Helper to cast a Target into the IMeshDescriptionProvider interface.
-	 */
-	IMeshDescriptionProvider* TargetMeshProviderInterface(int32 ComponentIdx) const;
-
-	/**
-	 * Helper to cast a Target into the IMaterialProvider interface.
-	 */
-	IMaterialProvider* TargetMaterialInterface(int32 ComponentIdx) const;
-
 public:
 	// IInteractiveToolCameraFocusAPI implementation
 	virtual bool SupportsWorldSpaceFocusBox() override;
@@ -97,23 +53,4 @@ public:
 	virtual bool GetWorldSpaceFocusPoint(const FRay& WorldRay, FVector& PointOut) override;
 };
 
-inline IPrimitiveComponentBackedTarget* UMultiSelectionTool::TargetComponentInterface(int32 TargetIdx) const
-{
-	return TargetInterface<IPrimitiveComponentBackedTarget>(TargetIdx);
-}
-
-inline IMeshDescriptionCommitter* UMultiSelectionTool::TargetMeshCommitterInterface(int32 TargetIdx) const
-{
-	return TargetInterface<IMeshDescriptionCommitter>(TargetIdx);
-}
-
-inline IMeshDescriptionProvider* UMultiSelectionTool::TargetMeshProviderInterface(int32 TargetIdx) const
-{
-	return TargetInterface<IMeshDescriptionProvider>(TargetIdx);
-}
-
-inline IMaterialProvider* UMultiSelectionTool::TargetMaterialInterface(int32 TargetIdx) const
-{
-	return TargetInterface<IMaterialProvider>(TargetIdx);
-}
 
