MeshModelingToolset: Smaller PolyExtrude and PathExtrude fixes

* fixed Flat extrude mode not generating a flat path
* the ExtrudeHeight property in PolyExt and the ExtrudeHeight and Width properties in PathExt are now reset to their original values after being modified by the respective interactive modes

diff --git a/Engine/Plugins/Experimental/MeshModelingToolsetExp/Source/MeshModelingToolsExp/Private/DrawPolyPathTool.cpp b/Engine/Plugins/Experimental/MeshModelingToolsetExp/Source/MeshModelingToolsExp/Private/DrawPolyPathTool.cpp
--- a/Engine/Plugins/Experimental/MeshModelingToolsetExp/Source/MeshModelingToolsExp/Private/DrawPolyPathTool.cpp
+++ b/Engine/Plugins/Experimental/MeshModelingToolsetExp/Source/MeshModelingToolsExp/Private/DrawPolyPathTool.cpp
@@ -227,6 +227,18 @@ void UDrawPolyPathTool::Setup()
 
 void UDrawPolyPathTool::Shutdown(EToolShutdownType ShutdownType)
 {
+	if (bHasSavedWidth)
+	{
+		TransformProps->Width = SavedWidth;
+		bHasSavedWidth = false;
+	}
+
+	if (bHasSavedExtrudeHeight)
+	{
+		TransformProps->ExtrudeHeight = SavedExtrudeHeight;
+		SavedExtrudeHeight = false;
+	}
+
 	PlaneMechanic->Shutdown();
 	PlaneMechanic = nullptr;
 
@@ -309,6 +321,7 @@ void UDrawPolyPathTool::OnClicked(const FInputDeviceRay& ClickPos)
 	}
 	else if (ExtrudeHeightMechanic != nullptr)
 	{
+		CurHeight = TransformProps->ExtrudeHeight;
 		OnCompleteExtrudeHeight();
 	}
 }
@@ -528,6 +541,9 @@ void UDrawPolyPathTool::OnCompleteSurfacePath()
 
 void UDrawPolyPathTool::BeginInteractiveOffsetDistance()
 {
+	bHasSavedWidth = true;
+	SavedWidth = TransformProps->Width;
+	
 	// begin setting offset distance
 	CurveDistMechanic = NewObject<USpatialCurveDistanceMechanic>(this);
 	CurveDistMechanic->Setup(this);
@@ -555,6 +571,7 @@ void UDrawPolyPathTool::OnCompleteOffsetDistance()
 
 	if (TransformProps->ExtrudeMode == EDrawPolyPathExtrudeMode::Flat)
 	{
+		CurHeight = 0.0;
 		OnCompleteExtrudeHeight();
 	}
 	else if (TransformProps->ExtrudeMode == EDrawPolyPathExtrudeMode::Fixed || TransformProps->ExtrudeMode == EDrawPolyPathExtrudeMode::RampFixed)
@@ -571,12 +588,11 @@ void UDrawPolyPathTool::OnCompleteOffsetDistance()
 
 void UDrawPolyPathTool::OnCompleteExtrudeHeight()
 {
-	CurHeight = TransformProps->ExtrudeHeight;
 	ExtrudeHeightMechanic = nullptr;
 
 	ClearPreview();
 
-	EmitNewObject(TransformProps->ExtrudeMode);
+	EmitNewObject();
 
 	InitializeNewSurfacePath();
 	CurrentCurveTimestamp++;
@@ -620,6 +636,9 @@ void UDrawPolyPathTool::GeneratePathMesh(FDynamicMesh3& Mesh)
 
 void UDrawPolyPathTool::BeginInteractiveExtrudeHeight()
 {
+	bHasSavedExtrudeHeight = true;
+	SavedExtrudeHeight = TransformProps->ExtrudeHeight;
+
 	// begin extrude
 	ExtrudeHeightMechanic = NewObject<UPlaneDistanceFromHitMechanic>(this);
 	ExtrudeHeightMechanic->Setup(this);
@@ -723,7 +742,7 @@ void UDrawPolyPathTool::GenerateExtrudeMesh(FDynamicMesh3& PathMesh)
 }
 
 
-void UDrawPolyPathTool::EmitNewObject(EDrawPolyPathExtrudeMode ExtrudeMode)
+void UDrawPolyPathTool::EmitNewObject()
 {
 	FDynamicMesh3 PathMesh;
 	GeneratePathMesh(PathMesh);
@@ -752,6 +771,18 @@ void UDrawPolyPathTool::EmitNewObject(EDrawPolyPathExtrudeMode ExtrudeMode)
 
 	GetToolManager()->EndUndoTransaction();
 
+	if (bHasSavedWidth)
+	{		
+		TransformProps->Width = SavedWidth;
+		bHasSavedWidth = false;
+	}
+
+	if (bHasSavedExtrudeHeight)
+	{
+		TransformProps->ExtrudeHeight = SavedExtrudeHeight;
+		bHasSavedExtrudeHeight = false;
+	}
+
 	CurPolyLoop.Reset();
 	SecondPolyLoop.Reset();
 }
diff --git a/Engine/Plugins/Experimental/MeshModelingToolsetExp/Source/MeshModelingToolsExp/Public/DrawPolyPathTool.h b/Engine/Plugins/Experimental/MeshModelingToolsetExp/Source/MeshModelingToolsExp/Public/DrawPolyPathTool.h
--- a/Engine/Plugins/Experimental/MeshModelingToolsetExp/Source/MeshModelingToolsExp/Public/DrawPolyPathTool.h
+++ b/Engine/Plugins/Experimental/MeshModelingToolsetExp/Source/MeshModelingToolsExp/Public/DrawPolyPathTool.h
@@ -201,6 +201,10 @@ protected:
 	double CurPathLength;
 	double CurOffsetDistance;
 	double CurHeight;
+	bool bHasSavedWidth = false;
+	float SavedWidth;
+	bool bHasSavedExtrudeHeight = false;
+	float SavedExtrudeHeight;
 	bool bPathIsClosed = false;		// If true, CurPathPoints are assumed to define a closed path
 
 	TArray<FVector3d> CurPolyLoop;
@@ -231,7 +235,7 @@ protected:
 	void ClearPreview();
 	void GeneratePathMesh(UE::Geometry::FDynamicMesh3& Mesh);
 	void GenerateExtrudeMesh(UE::Geometry::FDynamicMesh3& PathMesh);
-	void EmitNewObject(EDrawPolyPathExtrudeMode ExtrudeMode);
+	void EmitNewObject();
 
 	// user feedback messages
 	void ShowStartupMessage();
diff --git a/Engine/Plugins/Runtime/MeshModelingToolset/Source/MeshModelingTools/Private/DrawPolygonTool.cpp b/Engine/Plugins/Runtime/MeshModelingToolset/Source/MeshModelingTools/Private/DrawPolygonTool.cpp
--- a/Engine/Plugins/Runtime/MeshModelingToolset/Source/MeshModelingTools/Private/DrawPolygonTool.cpp
+++ b/Engine/Plugins/Runtime/MeshModelingToolset/Source/MeshModelingTools/Private/DrawPolygonTool.cpp
@@ -173,6 +173,12 @@ void UDrawPolygonTool::Setup()
 
 void UDrawPolygonTool::Shutdown(EToolShutdownType ShutdownType)
 {
+	if (bHasSavedExtrudeHeight)
+	{
+		PolygonProperties->ExtrudeHeight = SavedExtrudeHeight;
+		bHasSavedExtrudeHeight = false;
+	}
+
 	PreviewMesh->Disconnect();
 	PreviewMesh = nullptr;
 
@@ -950,6 +956,10 @@ void UDrawPolygonTool::GenerateFixedPolygon(const TArray<FVector3d>& FixedPoints
 void UDrawPolygonTool::BeginInteractiveExtrude()
 {
 	bInInteractiveExtrude = true;
+
+	bHasSavedExtrudeHeight = true;
+	SavedExtrudeHeight = PolygonProperties->ExtrudeHeight;
+
 	SnapEngine.ResetActiveSnap();
 
 	HeightMechanic = NewObject<UPlaneDistanceFromHitMechanic>(this);
@@ -1099,6 +1109,12 @@ void UDrawPolygonTool::EmitCurrentPolygon()
 
 	GetToolManager()->EndUndoTransaction();
 
+	if (bHasSavedExtrudeHeight)
+	{
+		PolygonProperties->ExtrudeHeight = SavedExtrudeHeight;
+		bHasSavedExtrudeHeight = false;
+	}
+
 	ResetPolygon();
 }
 
diff --git a/Engine/Plugins/Runtime/MeshModelingToolset/Source/MeshModelingTools/Public/DrawPolygonTool.h b/Engine/Plugins/Runtime/MeshModelingToolset/Source/MeshModelingTools/Public/DrawPolygonTool.h
--- a/Engine/Plugins/Runtime/MeshModelingToolset/Source/MeshModelingTools/Public/DrawPolygonTool.h
+++ b/Engine/Plugins/Runtime/MeshModelingToolset/Source/MeshModelingTools/Public/DrawPolygonTool.h
@@ -309,6 +309,8 @@ protected:
 	// extrusion control
 
 	bool bInInteractiveExtrude = false;
+	bool bHasSavedExtrudeHeight = false;
+	float SavedExtrudeHeight;
 
 	void UpdateLivePreview();
 	bool bPreviewUpdatePending;
