ModelingTools: Only show source mesh target in Bake tools if it is not baking to self.
Removed broken EditCondition from BakeAll property set.

#rnx

diff --git a/Engine/Plugins/Experimental/MeshModelingToolsetExp/Source/MeshModelingToolsExp/Private/BakeMeshAttributeMapsTool.cpp b/Engine/Plugins/Experimental/MeshModelingToolsetExp/Source/MeshModelingToolsExp/Private/BakeMeshAttributeMapsTool.cpp
--- a/Engine/Plugins/Experimental/MeshModelingToolsetExp/Source/MeshModelingToolsExp/Private/BakeMeshAttributeMapsTool.cpp
+++ b/Engine/Plugins/Experimental/MeshModelingToolsetExp/Source/MeshModelingToolsExp/Private/BakeMeshAttributeMapsTool.cpp
@@ -277,9 +277,9 @@ void UBakeMeshAttributeMapsTool::Setup()
 	MeshProps->TargetStaticMesh = GetStaticMeshTarget(Target);
 	MeshProps->TargetSkeletalMesh = GetSkeletalMeshTarget(Target);
 	MeshProps->TargetDynamicMesh = GetDynamicMeshTarget(Target);
-	MeshProps->SourceStaticMesh = GetStaticMeshTarget(DetailTarget);
-	MeshProps->SourceSkeletalMesh = GetSkeletalMeshTarget(DetailTarget);
-	MeshProps->SourceDynamicMesh = GetDynamicMeshTarget(DetailTarget);
+	MeshProps->SourceStaticMesh = !bIsBakeToSelf ? GetStaticMeshTarget(DetailTarget) : nullptr;
+	MeshProps->SourceSkeletalMesh = !bIsBakeToSelf ? GetSkeletalMeshTarget(DetailTarget) : nullptr;
+	MeshProps->SourceDynamicMesh = !bIsBakeToSelf ? GetDynamicMeshTarget(DetailTarget) : nullptr;
 	MeshProps->SourceNormalMap = nullptr;
 	UpdateUVLayerNames(MeshProps->TargetUVLayer, MeshProps->TargetUVLayerNamesList, BaseMesh);
 	MeshProps->WatchProperty(MeshProps->TargetUVLayer, [this](FString) { OpState |= EBakeOpState::Evaluate; });
diff --git a/Engine/Plugins/Experimental/MeshModelingToolsetExp/Source/MeshModelingToolsExp/Private/BakeMeshAttributeVertexTool.cpp b/Engine/Plugins/Experimental/MeshModelingToolsetExp/Source/MeshModelingToolsExp/Private/BakeMeshAttributeVertexTool.cpp
--- a/Engine/Plugins/Experimental/MeshModelingToolsetExp/Source/MeshModelingToolsExp/Private/BakeMeshAttributeVertexTool.cpp
+++ b/Engine/Plugins/Experimental/MeshModelingToolsetExp/Source/MeshModelingToolsExp/Private/BakeMeshAttributeVertexTool.cpp
@@ -283,9 +283,9 @@ void UBakeMeshAttributeVertexTool::Setup()
 	MeshProps->TargetStaticMesh = GetStaticMeshTarget(Target);
 	MeshProps->TargetSkeletalMesh = GetSkeletalMeshTarget(Target);
 	MeshProps->TargetDynamicMesh = GetDynamicMeshTarget(Target);
-	MeshProps->SourceStaticMesh = GetStaticMeshTarget(DetailTarget);
-	MeshProps->SourceSkeletalMesh = GetSkeletalMeshTarget(DetailTarget);
-	MeshProps->SourceDynamicMesh = GetDynamicMeshTarget(DetailTarget);
+	MeshProps->SourceStaticMesh = !bIsBakeToSelf ? GetStaticMeshTarget(DetailTarget) : nullptr;
+	MeshProps->SourceSkeletalMesh = !bIsBakeToSelf ? GetSkeletalMeshTarget(DetailTarget) : nullptr;
+	MeshProps->SourceDynamicMesh = !bIsBakeToSelf ? GetDynamicMeshTarget(DetailTarget) : nullptr;
 	MeshProps->SourceNormalMap = nullptr;
 	MeshProps->WatchProperty(MeshProps->ProjectionDistance, [this](float) { OpState |= EBakeOpState::Evaluate; });
 	MeshProps->WatchProperty(MeshProps->bProjectionInWorldSpace, [this](bool) { OpState |= EBakeOpState::EvaluateDetailMesh; });
diff --git a/Engine/Plugins/Experimental/MeshModelingToolsetExp/Source/MeshModelingToolsExp/Public/BakeMultiMeshAttributeMapsTool.h b/Engine/Plugins/Experimental/MeshModelingToolsetExp/Source/MeshModelingToolsExp/Public/BakeMultiMeshAttributeMapsTool.h
--- a/Engine/Plugins/Experimental/MeshModelingToolsetExp/Source/MeshModelingToolsExp/Public/BakeMultiMeshAttributeMapsTool.h
+++ b/Engine/Plugins/Experimental/MeshModelingToolsetExp/Source/MeshModelingToolsExp/Public/BakeMultiMeshAttributeMapsTool.h
@@ -117,8 +117,7 @@ public:
 
 	/** Maximum allowed distance for the projection from target mesh to source mesh for the sample to be considered valid.
 	 * This is only relevant if a separate source mesh is provided. */
-	UPROPERTY(EditAnywhere, AdvancedDisplay, Category = BakeInput, meta = (ClampMin = "0.001",
-		EditCondition = "SourceStaticMesh != nullptr", HideEditConditionToggle))
+	UPROPERTY(EditAnywhere, AdvancedDisplay, Category = BakeInput, meta = (ClampMin = "0.001"))
 	float ProjectionDistance = 3.0;
 
 	UFUNCTION()
