ModelingTools: Use "UV #" format in the UI for Bake tools UV channel selection.

#rnx

diff --git a/Engine/Plugins/Experimental/MeshModelingToolsetExp/Source/MeshModelingToolsExp/Private/BakeMeshAttributeMapsTool.cpp b/Engine/Plugins/Experimental/MeshModelingToolsetExp/Source/MeshModelingToolsExp/Private/BakeMeshAttributeMapsTool.cpp
--- a/Engine/Plugins/Experimental/MeshModelingToolsetExp/Source/MeshModelingToolsExp/Private/BakeMeshAttributeMapsTool.cpp
+++ b/Engine/Plugins/Experimental/MeshModelingToolsetExp/Source/MeshModelingToolsExp/Private/BakeMeshAttributeMapsTool.cpp
@@ -535,7 +535,7 @@ void UBakeMeshAttributeMapsTool::UpdateResult()
 	FBakeCacheSettings BakeCacheSettings;
 	BakeCacheSettings.Dimensions = Dimensions;
 	BakeCacheSettings.BitDepth = Settings->BitDepth;
-	BakeCacheSettings.TargetUVLayer = FCString::Atoi(*MeshProps->TargetUVLayer);
+	BakeCacheSettings.TargetUVLayer = MeshProps->TargetUVLayerNamesList.IndexOfByKey(MeshProps->TargetUVLayer);
 	BakeCacheSettings.DetailTimestamp = this->DetailMeshTimestamp;
 	BakeCacheSettings.ProjectionDistance = MeshProps->ProjectionDistance;
 	BakeCacheSettings.SamplesPerPixel = (int32)Settings->SamplesPerPixel;
@@ -610,8 +610,9 @@ void UBakeMeshAttributeMapsTool::UpdateResult()
 EBakeOpState UBakeMeshAttributeMapsTool::UpdateResult_DetailNormalMap()
 {
 	EBakeOpState ResultState = EBakeOpState::Clean;
-	
-	const FDynamicMeshUVOverlay* UVOverlay = DetailMesh->Attributes()->GetUVLayer(FCString::Atoi(*MeshProps->SourceNormalMapUVLayer));
+
+	const int DetailUVLayer = MeshProps->SourceUVLayerNamesList.IndexOfByKey(MeshProps->SourceNormalMapUVLayer);
+	const FDynamicMeshUVOverlay* UVOverlay = DetailMesh->Attributes()->GetUVLayer(DetailUVLayer);
 	if (UVOverlay == nullptr)
 	{
 		GetToolManager()->DisplayMessage(LOCTEXT("InvalidUVWarning", "The Detail Mesh does not have the selected UV layer"), EToolMessageLevel::UserWarning);
@@ -634,7 +635,7 @@ EBakeOpState UBakeMeshAttributeMapsTool::UpdateResult_DetailNormalMap()
 	}
 
 	FDetailMeshSettings DetailMeshSettings;
-	DetailMeshSettings.UVLayer = FCString::Atoi(*MeshProps->SourceNormalMapUVLayer);
+	DetailMeshSettings.UVLayer = DetailUVLayer;
 
 	if (!(CachedDetailMeshSettings == DetailMeshSettings))
 	{
diff --git a/Engine/Plugins/Experimental/MeshModelingToolsetExp/Source/MeshModelingToolsExp/Private/BakeMultiMeshAttributeMapsTool.cpp b/Engine/Plugins/Experimental/MeshModelingToolsetExp/Source/MeshModelingToolsExp/Private/BakeMultiMeshAttributeMapsTool.cpp
--- a/Engine/Plugins/Experimental/MeshModelingToolsetExp/Source/MeshModelingToolsExp/Private/BakeMultiMeshAttributeMapsTool.cpp
+++ b/Engine/Plugins/Experimental/MeshModelingToolsetExp/Source/MeshModelingToolsExp/Private/BakeMultiMeshAttributeMapsTool.cpp
@@ -562,7 +562,7 @@ void UBakeMultiMeshAttributeMapsTool::UpdateResult()
 	FBakeCacheSettings BakeCacheSettings;
 	BakeCacheSettings.Dimensions = Dimensions;
 	BakeCacheSettings.BitDepth = Settings->BitDepth;
-	BakeCacheSettings.TargetUVLayer = FCString::Atoi(*InputMeshProps->TargetUVLayer);
+	BakeCacheSettings.TargetUVLayer = InputMeshProps->TargetUVLayerNamesList.IndexOfByKey(InputMeshProps->TargetUVLayer);
 	BakeCacheSettings.ProjectionDistance = InputMeshProps->ProjectionDistance;
 	BakeCacheSettings.bProjectionInWorldSpace = true;  // Always world space
 	BakeCacheSettings.SamplesPerPixel = static_cast<int32>(Settings->SamplesPerPixel);
diff --git a/Engine/Plugins/Experimental/MeshModelingToolsetExp/Source/MeshModelingToolsExp/Public/BakeMeshAttributeMapsToolBase.h b/Engine/Plugins/Experimental/MeshModelingToolsetExp/Source/MeshModelingToolsExp/Public/BakeMeshAttributeMapsToolBase.h
--- a/Engine/Plugins/Experimental/MeshModelingToolsetExp/Source/MeshModelingToolsExp/Public/BakeMeshAttributeMapsToolBase.h
+++ b/Engine/Plugins/Experimental/MeshModelingToolsetExp/Source/MeshModelingToolsExp/Public/BakeMeshAttributeMapsToolBase.h
@@ -422,7 +422,7 @@ inline void UBakeMeshAttributeMapsToolBase::UpdateUVLayerNames(FString& UVLayer,
 	int32 FoundIndex = -1;
 	for (int32 k = 0; k < Mesh.Attributes()->NumUVLayers(); ++k)
 	{
-		UVLayerNamesList.Add(FString::FromInt(k));
+		UVLayerNamesList.Add(FString::Printf(TEXT("UV %d"), k));
 		if (UVLayer == UVLayerNamesList.Last())
 		{
 			FoundIndex = k;
