ModelingTools: Exposed UVChannel parameter in BakePreviewMaterial and BakeBentNormalPreviewMaterial to preview bake texture results on specific UV channels.

#rnx

diff --git a/Engine/Plugins/Experimental/MeshModelingToolsetExp/Source/MeshModelingToolsExp/Private/BakeMeshAttributeMapsToolBase.cpp b/Engine/Plugins/Experimental/MeshModelingToolsetExp/Source/MeshModelingToolsExp/Private/BakeMeshAttributeMapsToolBase.cpp
--- a/Engine/Plugins/Experimental/MeshModelingToolsetExp/Source/MeshModelingToolsExp/Private/BakeMeshAttributeMapsToolBase.cpp
+++ b/Engine/Plugins/Experimental/MeshModelingToolsetExp/Source/MeshModelingToolsExp/Private/BakeMeshAttributeMapsToolBase.cpp
@@ -225,6 +225,8 @@ void UBakeMeshAttributeMapsToolBase::UpdatePreview(const EBakeMapType PreviewMap
 		PreviewMaterial->SetTextureParameterValue(TEXT("ColorMap"), PreviewMap);
 		break;
 	}
+	PreviewMaterial->SetScalarParameterValue(TEXT("UVChannel"), CachedBakeCacheSettings.TargetUVLayer);
+	BentNormalPreviewMaterial->SetScalarParameterValue(TEXT("UVChannel"), CachedBakeCacheSettings.TargetUVLayer);
 }
 
 
