ModelingTools: Fix issue with source normal maps not being considered by tangent space normal map bake.

#rnx

diff --git a/Engine/Plugins/Experimental/MeshModelingToolsetExp/Source/MeshModelingToolsExp/Private/BakeMeshAttributeMapsTool.cpp b/Engine/Plugins/Experimental/MeshModelingToolsetExp/Source/MeshModelingToolsExp/Private/BakeMeshAttributeMapsTool.cpp
--- a/Engine/Plugins/Experimental/MeshModelingToolsetExp/Source/MeshModelingToolsExp/Private/BakeMeshAttributeMapsTool.cpp
+++ b/Engine/Plugins/Experimental/MeshModelingToolsetExp/Source/MeshModelingToolsExp/Private/BakeMeshAttributeMapsTool.cpp
@@ -293,7 +293,7 @@ void UBakeMeshAttributeMapsTool::Setup()
 		{
 			OpState |= EBakeOpState::EvaluateDetailMesh;
 		}
-		OpState = EBakeOpState::Evaluate;
+		OpState |= EBakeOpState::Evaluate;
 	});
 
 	
