ModelingTools: Minor refactor and renaming of bake tool members.
Fixed Texture2D bake not respecting selected UV layer.
Fixed MultiTexture2D bake overriding Texture2D bake UV layer.

#rnx

diff --git a/Engine/Plugins/Experimental/MeshModelingToolsetExp/Source/MeshModelingToolsExp/Private/BakeMeshAttributeMapsTool.cpp b/Engine/Plugins/Experimental/MeshModelingToolsetExp/Source/MeshModelingToolsExp/Private/BakeMeshAttributeMapsTool.cpp
--- a/Engine/Plugins/Experimental/MeshModelingToolsetExp/Source/MeshModelingToolsExp/Private/BakeMeshAttributeMapsTool.cpp
+++ b/Engine/Plugins/Experimental/MeshModelingToolsetExp/Source/MeshModelingToolsExp/Private/BakeMeshAttributeMapsTool.cpp
@@ -101,7 +101,7 @@ public:
 	TSharedPtr<UE::Geometry::TMeshTangents<double>, ESPMode::ThreadSafe> DetailMeshTangents;
 	UE::Geometry::FDynamicMesh3* BaseMesh;
 	TUniquePtr<UE::Geometry::FMeshMapBaker> Baker;
-	UBakeMeshAttributeMapsTool::FBakeCacheSettings BakeCacheSettings;
+	UBakeMeshAttributeMapsTool::FBakeSettings BakeSettings;
 	TSharedPtr<UE::Geometry::TMeshTangents<double>, ESPMode::ThreadSafe> BaseMeshTangents;
 
 	// Map Type settings
@@ -110,7 +110,8 @@ public:
 	FOcclusionMapSettings OcclusionSettings;
 	FCurvatureMapSettings CurvatureSettings;
 	FMeshPropertyMapSettings PropertySettings;
-	FTexture2DImageSettings TextureSettings;
+	FTexture2DSettings TextureSettings;
+	FTexture2DSettings MultiTextureSettings;
 
 	// NormalMap settings
 	ImagePtr DetailMeshNormalMap;
@@ -128,10 +129,10 @@ public:
 			return Progress && Progress->Cancelled();
 		};
 		Baker->SetTargetMesh(BaseMesh);
-		Baker->SetTargetMeshUVLayer(BakeCacheSettings.TargetUVLayer);
-		Baker->SetDimensions(BakeCacheSettings.Dimensions);
-		Baker->SetProjectionDistance(BakeCacheSettings.ProjectionDistance);
-		Baker->SetSamplesPerPixel(BakeCacheSettings.SamplesPerPixel);
+		Baker->SetTargetMeshUVLayer(BakeSettings.TargetUVLayer);
+		Baker->SetDimensions(BakeSettings.Dimensions);
+		Baker->SetProjectionDistance(BakeSettings.ProjectionDistance);
+		Baker->SetSamplesPerPixel(BakeSettings.SamplesPerPixel);
 		Baker->SetTargetMeshTangents(BaseMeshTangents);
 		
 		FMeshBakerDynamicMeshSampler DetailSampler(DetailMesh.Get(), DetailSpatial.Get(), DetailMeshTangents.Get());
@@ -139,7 +140,7 @@ public:
 
 		for (const EBakeMapType MapType : ALL_BAKE_MAP_TYPES)
 		{
-			switch (BakeCacheSettings.BakeMapTypes & MapType)
+			switch (BakeSettings.BakeMapTypes & MapType)
 			{
 			case EBakeMapType::TangentSpaceNormal:
 			{
@@ -154,11 +155,11 @@ public:
 			{
 				TSharedPtr<FMeshOcclusionMapEvaluator> OcclusionEval = MakeShared<FMeshOcclusionMapEvaluator>();
 				OcclusionEval->OcclusionType = EMeshOcclusionMapType::None;
-				if ((bool)(BakeCacheSettings.BakeMapTypes & EBakeMapType::AmbientOcclusion))
+				if ((bool)(BakeSettings.BakeMapTypes & EBakeMapType::AmbientOcclusion))
 				{
 					OcclusionEval->OcclusionType |= EMeshOcclusionMapType::AmbientOcclusion;
 				}
-				if ((bool)(BakeCacheSettings.BakeMapTypes & EBakeMapType::BentNormal))
+				if ((bool)(BakeSettings.BakeMapTypes & EBakeMapType::BentNormal))
 				{
 					OcclusionEval->OcclusionType |= EMeshOcclusionMapType::BentNormal;
 				}
@@ -227,7 +228,7 @@ public:
 			case EBakeMapType::MultiTexture:
 			{
 				TSharedPtr<FMeshMultiResampleImageEvaluator, ESPMode::ThreadSafe> TextureEval = MakeShared<FMeshMultiResampleImageEvaluator, ESPMode::ThreadSafe>();
-				TextureEval->DetailUVLayer = TextureSettings.UVLayer;
+				TextureEval->DetailUVLayer = MultiTextureSettings.UVLayer;
 				TextureEval->MultiTextures = MaterialIDTextures;
 				Baker->AddEvaluator(TextureEval);
 				break;
@@ -257,10 +258,10 @@ void UBakeMeshAttributeMapsTool::Setup()
 
 	PreviewMesh->ProcessMesh([this](const FDynamicMesh3& Mesh)
 	{
-		BaseMesh.Copy(Mesh);
-		BaseSpatial.SetMesh(&BaseMesh, true);
-		BaseMeshTangents = MakeShared<FMeshTangentsd, ESPMode::ThreadSafe>(&BaseMesh);
-		BaseMeshTangents->CopyTriVertexTangents(Mesh);
+		TargetMesh.Copy(Mesh);
+		TargetSpatial.SetMesh(&TargetMesh, true);
+		TargetMeshTangents = MakeShared<FMeshTangentsd, ESPMode::ThreadSafe>(&TargetMesh);
+		TargetMeshTangents->CopyTriVertexTangents(Mesh);
 	});
 
 	UToolTarget* Target = Targets[0];
@@ -281,7 +282,7 @@ void UBakeMeshAttributeMapsTool::Setup()
 	MeshProps->SourceSkeletalMesh = !bIsBakeToSelf ? GetSkeletalMeshTarget(DetailTarget) : nullptr;
 	MeshProps->SourceDynamicMesh = !bIsBakeToSelf ? GetDynamicMeshTarget(DetailTarget) : nullptr;
 	MeshProps->SourceNormalMap = nullptr;
-	UpdateUVLayerNames(MeshProps->TargetUVLayer, MeshProps->TargetUVLayerNamesList, BaseMesh);
+	UpdateUVLayerNames(MeshProps->TargetUVLayer, MeshProps->TargetUVLayerNamesList, TargetMesh);
 	MeshProps->WatchProperty(MeshProps->TargetUVLayer, [this](FString) { OpState |= EBakeOpState::Evaluate; });
 	MeshProps->WatchProperty(MeshProps->ProjectionDistance, [this](float) { OpState |= EBakeOpState::Evaluate; });
 	MeshProps->WatchProperty(MeshProps->bProjectionInWorldSpace, [this](bool) { OpState |= EBakeOpState::EvaluateDetailMesh; });
@@ -308,44 +309,44 @@ void UBakeMeshAttributeMapsTool::Setup()
 	Settings->WatchProperty(Settings->SamplesPerPixel, [this](EBakeTextureSamplesPerPixel) { OpState |= EBakeOpState::Evaluate; });
 	
 
-	OcclusionMapProps = NewObject<UBakedOcclusionMapToolProperties>(this);
-	OcclusionMapProps->RestoreProperties(this);
-	AddToolPropertySource(OcclusionMapProps);
-	SetToolPropertySourceEnabled(OcclusionMapProps, false);
-	OcclusionMapProps->WatchProperty(OcclusionMapProps->OcclusionRays, [this](int32) { OpState |= EBakeOpState::Evaluate; });
-	OcclusionMapProps->WatchProperty(OcclusionMapProps->MaxDistance, [this](float) { OpState |= EBakeOpState::Evaluate; });
-	OcclusionMapProps->WatchProperty(OcclusionMapProps->SpreadAngle, [this](float) { OpState |= EBakeOpState::Evaluate; });
-	OcclusionMapProps->WatchProperty(OcclusionMapProps->BiasAngle, [this](float) { OpState |= EBakeOpState::Evaluate; });
-
-
-	CurvatureMapProps = NewObject<UBakedCurvatureMapToolProperties>(this);
-	CurvatureMapProps->RestoreProperties(this);
-	AddToolPropertySource(CurvatureMapProps);
-	SetToolPropertySourceEnabled(CurvatureMapProps, false);
-	CurvatureMapProps->WatchProperty(CurvatureMapProps->ColorRangeMultiplier, [this](float) { OpState |= EBakeOpState::Evaluate; });
-	CurvatureMapProps->WatchProperty(CurvatureMapProps->MinRangeMultiplier, [this](float) { OpState |= EBakeOpState::Evaluate; });
-	CurvatureMapProps->WatchProperty(CurvatureMapProps->CurvatureType, [this](EBakedCurvatureTypeMode) { OpState |= EBakeOpState::Evaluate; });
-	CurvatureMapProps->WatchProperty(CurvatureMapProps->ColorMapping, [this](EBakedCurvatureColorMode) { OpState |= EBakeOpState::Evaluate; });
-	CurvatureMapProps->WatchProperty(CurvatureMapProps->Clamping, [this](EBakedCurvatureClampMode) { OpState |= EBakeOpState::Evaluate; });
-
-
-	Texture2DProps = NewObject<UBakedTexture2DImageProperties>(this);
-	Texture2DProps->RestoreProperties(this);
-	AddToolPropertySource(Texture2DProps);
-	SetToolPropertySourceEnabled(Texture2DProps, false);
-	Texture2DProps->WatchProperty(Texture2DProps->UVLayer, [this](float) { OpState |= EBakeOpState::Evaluate; });
-	Texture2DProps->WatchProperty(Texture2DProps->SourceTexture, [this](UTexture2D*) { OpState |= EBakeOpState::Evaluate; });
-
-	MultiTextureProps = NewObject<UBakedMultiTexture2DImageProperties>(this);
-	MultiTextureProps->RestoreProperties(this);
-	AddToolPropertySource(MultiTextureProps);
-	SetToolPropertySourceEnabled(MultiTextureProps, false);
-
-	auto SetDirtyCallback = [this](decltype(MultiTextureProps->MaterialIDSourceTextures)) { OpState |= EBakeOpState::Evaluate; };
-	auto NotEqualsCallback = [](const decltype(MultiTextureProps->MaterialIDSourceTextures)& A, const decltype(MultiTextureProps->MaterialIDSourceTextures)& B) -> bool { return A != B; };
-	MultiTextureProps->WatchProperty(MultiTextureProps->MaterialIDSourceTextures, SetDirtyCallback, NotEqualsCallback);
-	MultiTextureProps->WatchProperty(MultiTextureProps->UVLayer, [this](float) { OpState |= EBakeOpState::Evaluate; });
-	UpdateMultiTextureMaterialIDs(DetailTarget, MultiTextureProps->AllSourceTextures, MultiTextureProps->MaterialIDSourceTextures);
+	OcclusionSettings = NewObject<UBakeOcclusionMapToolProperties>(this);
+	OcclusionSettings->RestoreProperties(this);
+	AddToolPropertySource(OcclusionSettings);
+	SetToolPropertySourceEnabled(OcclusionSettings, false);
+	OcclusionSettings->WatchProperty(OcclusionSettings->OcclusionRays, [this](int32) { OpState |= EBakeOpState::Evaluate; });
+	OcclusionSettings->WatchProperty(OcclusionSettings->MaxDistance, [this](float) { OpState |= EBakeOpState::Evaluate; });
+	OcclusionSettings->WatchProperty(OcclusionSettings->SpreadAngle, [this](float) { OpState |= EBakeOpState::Evaluate; });
+	OcclusionSettings->WatchProperty(OcclusionSettings->BiasAngle, [this](float) { OpState |= EBakeOpState::Evaluate; });
+
+
+	CurvatureSettings = NewObject<UBakeCurvatureMapToolProperties>(this);
+	CurvatureSettings->RestoreProperties(this);
+	AddToolPropertySource(CurvatureSettings);
+	SetToolPropertySourceEnabled(CurvatureSettings, false);
+	CurvatureSettings->WatchProperty(CurvatureSettings->ColorRangeMultiplier, [this](float) { OpState |= EBakeOpState::Evaluate; });
+	CurvatureSettings->WatchProperty(CurvatureSettings->MinRangeMultiplier, [this](float) { OpState |= EBakeOpState::Evaluate; });
+	CurvatureSettings->WatchProperty(CurvatureSettings->CurvatureType, [this](EBakeCurvatureTypeMode) { OpState |= EBakeOpState::Evaluate; });
+	CurvatureSettings->WatchProperty(CurvatureSettings->ColorMapping, [this](EBakeCurvatureColorMode) { OpState |= EBakeOpState::Evaluate; });
+	CurvatureSettings->WatchProperty(CurvatureSettings->Clamping, [this](EBakeCurvatureClampMode) { OpState |= EBakeOpState::Evaluate; });
+
+
+	TextureSettings = NewObject<UBakeTexture2DProperties>(this);
+	TextureSettings->RestoreProperties(this);
+	AddToolPropertySource(TextureSettings);
+	SetToolPropertySourceEnabled(TextureSettings, false);
+	TextureSettings->WatchProperty(TextureSettings->UVLayer, [this](float) { OpState |= EBakeOpState::Evaluate; });
+	TextureSettings->WatchProperty(TextureSettings->SourceTexture, [this](UTexture2D*) { OpState |= EBakeOpState::Evaluate; });
+
+	MultiTextureSettings = NewObject<UBakeMultiTexture2DProperties>(this);
+	MultiTextureSettings->RestoreProperties(this);
+	AddToolPropertySource(MultiTextureSettings);
+	SetToolPropertySourceEnabled(MultiTextureSettings, false);
+
+	auto SetDirtyCallback = [this](decltype(MultiTextureSettings->MaterialIDSourceTextures)) { OpState |= EBakeOpState::Evaluate; };
+	auto NotEqualsCallback = [](const decltype(MultiTextureSettings->MaterialIDSourceTextures)& A, const decltype(MultiTextureSettings->MaterialIDSourceTextures)& B) -> bool { return A != B; };
+	MultiTextureSettings->WatchProperty(MultiTextureSettings->MaterialIDSourceTextures, SetDirtyCallback, NotEqualsCallback);
+	MultiTextureSettings->WatchProperty(MultiTextureSettings->UVLayer, [this](float) { OpState |= EBakeOpState::Evaluate; });
+	UpdateMultiTextureMaterialIDs(DetailTarget, MultiTextureSettings->AllSourceTextures, MultiTextureSettings->MaterialIDSourceTextures);
 
 	UpdateOnModeChange();
 
@@ -383,13 +384,13 @@ TUniquePtr<UE::Geometry::TGenericDataOperator<FMeshMapBaker>> UBakeMeshAttribute
 	TUniquePtr<FMeshMapBakerOp> Op = MakeUnique<FMeshMapBakerOp>();
 	Op->DetailMesh = DetailMesh;
 	Op->DetailSpatial = DetailSpatial;
-	Op->BaseMesh = &BaseMesh;
-	Op->BakeCacheSettings = CachedBakeCacheSettings;
+	Op->BaseMesh = &TargetMesh;
+	Op->BakeSettings = CachedBakeSettings;
 
 	constexpr EBakeMapType RequiresTangents = EBakeMapType::TangentSpaceNormal | EBakeMapType::BentNormal;
-	if ((bool)(CachedBakeCacheSettings.BakeMapTypes & RequiresTangents))
+	if ((bool)(CachedBakeSettings.BakeMapTypes & RequiresTangents))
 	{
-		Op->BaseMeshTangents = BaseMeshTangents;
+		Op->BaseMeshTangents = TargetMeshTangents;
 	}
 
 	if (CachedDetailNormalMap)
@@ -399,40 +400,40 @@ TUniquePtr<UE::Geometry::TGenericDataOperator<FMeshMapBaker>> UBakeMeshAttribute
 		Op->DetailMeshNormalUVLayer = CachedDetailMeshSettings.UVLayer;
 	}
 
-	if ((bool)(CachedBakeCacheSettings.BakeMapTypes & EBakeMapType::TangentSpaceNormal))
+	if ((bool)(CachedBakeSettings.BakeMapTypes & EBakeMapType::TangentSpaceNormal))
 	{
 		Op->NormalSettings = CachedNormalMapSettings;
 	}
 
-	if ((bool)(CachedBakeCacheSettings.BakeMapTypes & EBakeMapType::AmbientOcclusion) ||
-		(bool)(CachedBakeCacheSettings.BakeMapTypes & EBakeMapType::BentNormal))
+	if ((bool)(CachedBakeSettings.BakeMapTypes & EBakeMapType::AmbientOcclusion) ||
+		(bool)(CachedBakeSettings.BakeMapTypes & EBakeMapType::BentNormal))
 	{
 		Op->OcclusionSettings = CachedOcclusionMapSettings;
 	}
 
-	if ((bool)(CachedBakeCacheSettings.BakeMapTypes & EBakeMapType::Curvature))
+	if ((bool)(CachedBakeSettings.BakeMapTypes & EBakeMapType::Curvature))
 	{
 		Op->CurvatureSettings = CachedCurvatureMapSettings;
 	}
 
-	if ((bool)(CachedBakeCacheSettings.BakeMapTypes & EBakeMapType::ObjectSpaceNormal) ||
-		(bool)(CachedBakeCacheSettings.BakeMapTypes & EBakeMapType::FaceNormal) ||
-		(bool)(CachedBakeCacheSettings.BakeMapTypes & EBakeMapType::Position) ||
-		(bool)(CachedBakeCacheSettings.BakeMapTypes & EBakeMapType::MaterialID) ||
-		(bool)(CachedBakeCacheSettings.BakeMapTypes & EBakeMapType::VertexColor))
+	if ((bool)(CachedBakeSettings.BakeMapTypes & EBakeMapType::ObjectSpaceNormal) ||
+		(bool)(CachedBakeSettings.BakeMapTypes & EBakeMapType::FaceNormal) ||
+		(bool)(CachedBakeSettings.BakeMapTypes & EBakeMapType::Position) ||
+		(bool)(CachedBakeSettings.BakeMapTypes & EBakeMapType::MaterialID) ||
+		(bool)(CachedBakeSettings.BakeMapTypes & EBakeMapType::VertexColor))
 	{
 		Op->PropertySettings = CachedMeshPropertyMapSettings;
 	}
 
-	if ((bool)(CachedBakeCacheSettings.BakeMapTypes & EBakeMapType::Texture))
+	if ((bool)(CachedBakeSettings.BakeMapTypes & EBakeMapType::Texture))
 	{
-		Op->TextureSettings = CachedTexture2DImageSettings;
+		Op->TextureSettings = CachedTexture2DSettings;
 		Op->TextureImage = CachedTextureImage;
 	}
 
-	if ((bool)(CachedBakeCacheSettings.BakeMapTypes & EBakeMapType::MultiTexture))
+	if ((bool)(CachedBakeSettings.BakeMapTypes & EBakeMapType::MultiTexture))
 	{
-		Op->TextureSettings = CachedTexture2DImageSettings;
+		Op->TextureSettings = CachedMultiTexture2DSettings;
 		Op->MaterialIDTextures = CachedMultiTextures;
 	}
 
@@ -448,10 +449,10 @@ void UBakeMeshAttributeMapsTool::Shutdown(EToolShutdownType ShutdownType)
 	
 	Settings->SaveProperties(this);
 	MeshProps->SaveProperties(this);
-	OcclusionMapProps->SaveProperties(this);
-	CurvatureMapProps->SaveProperties(this);
-	Texture2DProps->SaveProperties(this);
-	MultiTextureProps->SaveProperties(this);
+	OcclusionSettings->SaveProperties(this);
+	CurvatureSettings->SaveProperties(this);
+	TextureSettings->SaveProperties(this);
+	MultiTextureSettings->SaveProperties(this);
 
 	if (Compute)
 	{
@@ -509,7 +510,7 @@ void UBakeMeshAttributeMapsTool::UpdateDetailMesh()
 	// Clear detail mesh evaluation flag and mark evaluation.
 	OpState &= ~EBakeOpState::EvaluateDetailMesh;
 	OpState |= EBakeOpState::Evaluate;
-	CachedBakeCacheSettings = FBakeCacheSettings();
+	CachedBakeSettings = FBakeSettings();
 	DetailMeshTimestamp++;
 }
 
@@ -532,30 +533,31 @@ void UBakeMeshAttributeMapsTool::UpdateResult()
 	int32 ImageSize = (int32)Settings->Resolution;
 	FImageDimensions Dimensions(ImageSize, ImageSize);
 
-	FBakeCacheSettings BakeCacheSettings;
-	BakeCacheSettings.Dimensions = Dimensions;
-	BakeCacheSettings.BitDepth = Settings->BitDepth;
-	BakeCacheSettings.TargetUVLayer = MeshProps->TargetUVLayerNamesList.IndexOfByKey(MeshProps->TargetUVLayer);
-	BakeCacheSettings.DetailTimestamp = this->DetailMeshTimestamp;
-	BakeCacheSettings.ProjectionDistance = MeshProps->ProjectionDistance;
-	BakeCacheSettings.SamplesPerPixel = (int32)Settings->SamplesPerPixel;
-	BakeCacheSettings.bProjectionInWorldSpace = MeshProps->bProjectionInWorldSpace;
+	FBakeSettings BakeSettings;
+	BakeSettings.Dimensions = Dimensions;
+	BakeSettings.BitDepth = Settings->BitDepth;
+	BakeSettings.TargetUVLayer = MeshProps->TargetUVLayerNamesList.IndexOfByKey(MeshProps->TargetUVLayer);
+	BakeSettings.DetailTimestamp = this->DetailMeshTimestamp;
+	BakeSettings.ProjectionDistance = MeshProps->ProjectionDistance;
+	BakeSettings.SamplesPerPixel = (int32)Settings->SamplesPerPixel;
+	BakeSettings.bProjectionInWorldSpace = MeshProps->bProjectionInWorldSpace;
 
 	// Record the original map types and process the raw bitfield which may add
 	// additional targets.
-	BakeCacheSettings.SourceBakeMapTypes = static_cast<EBakeMapType>(Settings->MapTypes);
-	BakeCacheSettings.BakeMapTypes = GetMapTypes(Settings->MapTypes);
+	BakeSettings.SourceBakeMapTypes = static_cast<EBakeMapType>(Settings->MapTypes);
+	BakeSettings.BakeMapTypes = GetMapTypes(Settings->MapTypes);
 
 	// update bake cache settings
-	if (!(CachedBakeCacheSettings == BakeCacheSettings))
+	if (!(CachedBakeSettings == BakeSettings))
 	{
-		CachedBakeCacheSettings = BakeCacheSettings;
+		CachedBakeSettings = BakeSettings;
 
 		CachedNormalMapSettings = FNormalMapSettings();
 		CachedOcclusionMapSettings = FOcclusionMapSettings();
 		CachedCurvatureMapSettings = FCurvatureMapSettings();
 		CachedMeshPropertyMapSettings = FMeshPropertyMapSettings();
-		CachedTexture2DImageSettings = FTexture2DImageSettings();
+		CachedTexture2DSettings = FTexture2DSettings();
+		CachedMultiTexture2DSettings = FTexture2DSettings();
 	}
 
 	// Clear our invalid bitflag to check again for valid inputs.
@@ -564,34 +566,34 @@ void UBakeMeshAttributeMapsTool::UpdateResult()
 	OpState |= UpdateResult_DetailNormalMap();
 
 	// Update map type settings
-	if ((bool)(CachedBakeCacheSettings.BakeMapTypes & EBakeMapType::TangentSpaceNormal))
+	if ((bool)(CachedBakeSettings.BakeMapTypes & EBakeMapType::TangentSpaceNormal))
 	{
-		OpState |= UpdateResult_Normal();
+		OpState |= UpdateResult_Normal(CachedBakeSettings.Dimensions);
 	}
-	if ((bool)(CachedBakeCacheSettings.BakeMapTypes & EBakeMapType::AmbientOcclusion) ||
-		(bool)(CachedBakeCacheSettings.BakeMapTypes & EBakeMapType::BentNormal))
+	if ((bool)(CachedBakeSettings.BakeMapTypes & EBakeMapType::AmbientOcclusion) ||
+		(bool)(CachedBakeSettings.BakeMapTypes & EBakeMapType::BentNormal))
 	{
-		OpState |= UpdateResult_Occlusion();
+		OpState |= UpdateResult_Occlusion(CachedBakeSettings.Dimensions);
 	}
-	if ((bool)(CachedBakeCacheSettings.BakeMapTypes & EBakeMapType::Curvature))
+	if ((bool)(CachedBakeSettings.BakeMapTypes & EBakeMapType::Curvature))
 	{
-		OpState |= UpdateResult_Curvature();
+		OpState |= UpdateResult_Curvature(CachedBakeSettings.Dimensions);
 	}
-	if ((bool)(CachedBakeCacheSettings.BakeMapTypes & EBakeMapType::ObjectSpaceNormal) ||
-		(bool)(CachedBakeCacheSettings.BakeMapTypes & EBakeMapType::FaceNormal) ||
-		(bool)(CachedBakeCacheSettings.BakeMapTypes & EBakeMapType::Position) ||
-		(bool)(CachedBakeCacheSettings.BakeMapTypes & EBakeMapType::MaterialID) ||
-		(bool)(CachedBakeCacheSettings.BakeMapTypes & EBakeMapType::VertexColor))
+	if ((bool)(CachedBakeSettings.BakeMapTypes & EBakeMapType::ObjectSpaceNormal) ||
+		(bool)(CachedBakeSettings.BakeMapTypes & EBakeMapType::FaceNormal) ||
+		(bool)(CachedBakeSettings.BakeMapTypes & EBakeMapType::Position) ||
+		(bool)(CachedBakeSettings.BakeMapTypes & EBakeMapType::MaterialID) ||
+		(bool)(CachedBakeSettings.BakeMapTypes & EBakeMapType::VertexColor))
 	{
-		OpState |= UpdateResult_MeshProperty();
+		OpState |= UpdateResult_MeshProperty(CachedBakeSettings.Dimensions);
 	}
-	if ((bool)(CachedBakeCacheSettings.BakeMapTypes & EBakeMapType::Texture))
+	if ((bool)(CachedBakeSettings.BakeMapTypes & EBakeMapType::Texture))
 	{
-		OpState |= UpdateResult_Texture2DImage();
+		OpState |= UpdateResult_Texture2DImage(CachedBakeSettings.Dimensions, DetailMesh.Get());
 	}
-	if ((bool)(CachedBakeCacheSettings.BakeMapTypes & EBakeMapType::MultiTexture))
+	if ((bool)(CachedBakeSettings.BakeMapTypes & EBakeMapType::MultiTexture))
 	{
-		OpState |= UpdateResult_MultiTexture();
+		OpState |= UpdateResult_MultiTexture(CachedBakeSettings.Dimensions, DetailMesh.Get());
 	}
 
 	// Early exit if op input parameters are invalid.
@@ -646,222 +648,6 @@ EBakeOpState UBakeMeshAttributeMapsTool::UpdateResult_DetailNormalMap()
 }
 
 
-EBakeOpState UBakeMeshAttributeMapsTool::UpdateResult_Normal()
-{
-	EBakeOpState ResultState = EBakeOpState::Clean;
-
-	const int32 ImageSize = (int32)Settings->Resolution;
-	const FImageDimensions Dimensions(ImageSize, ImageSize);
-
-	FNormalMapSettings NormalMapSettings;
-	NormalMapSettings.Dimensions = Dimensions;
-
-	if (!(CachedNormalMapSettings == NormalMapSettings))
-	{
-		CachedNormalMapSettings = NormalMapSettings;
-		ResultState = EBakeOpState::Evaluate;
-	}
-	return ResultState;
-}
-
-
-EBakeOpState UBakeMeshAttributeMapsTool::UpdateResult_Occlusion()
-{
-	EBakeOpState ResultState = EBakeOpState::Clean;
-
-	const int32 ImageSize = (int32)Settings->Resolution;
-	const FImageDimensions Dimensions(ImageSize, ImageSize);
-
-	FOcclusionMapSettings OcclusionMapSettings;
-	OcclusionMapSettings.Dimensions = Dimensions;
-	OcclusionMapSettings.MaxDistance = (OcclusionMapProps->MaxDistance == 0) ? TNumericLimits<float>::Max() : OcclusionMapProps->MaxDistance;
-	OcclusionMapSettings.OcclusionRays = OcclusionMapProps->OcclusionRays;
-	OcclusionMapSettings.SpreadAngle = OcclusionMapProps->SpreadAngle;
-	OcclusionMapSettings.BiasAngle = OcclusionMapProps->BiasAngle;
-
-	if ( !(CachedOcclusionMapSettings == OcclusionMapSettings) )
-	{
-		CachedOcclusionMapSettings = OcclusionMapSettings;
-		ResultState = EBakeOpState::Evaluate;
-	}
-	return ResultState;
-}
-
-
-EBakeOpState UBakeMeshAttributeMapsTool::UpdateResult_Curvature()
-{
-	EBakeOpState ResultState = EBakeOpState::Clean;
-
-	const int32 ImageSize = (int32)Settings->Resolution;
-	const FImageDimensions Dimensions(ImageSize, ImageSize);
-
-	FCurvatureMapSettings CurvatureMapSettings;
-	CurvatureMapSettings.Dimensions = Dimensions;
-	CurvatureMapSettings.RangeMultiplier = CurvatureMapProps->ColorRangeMultiplier;
-	CurvatureMapSettings.MinRangeMultiplier = CurvatureMapProps->MinRangeMultiplier;
-	switch (CurvatureMapProps->CurvatureType)
-	{
-	default:
-	case EBakedCurvatureTypeMode::MeanAverage:
-		CurvatureMapSettings.CurvatureType = (int32)FMeshCurvatureMapEvaluator::ECurvatureType::Mean;
-		break;
-	case EBakedCurvatureTypeMode::Gaussian:
-		CurvatureMapSettings.CurvatureType = (int32)FMeshCurvatureMapEvaluator::ECurvatureType::Gaussian;
-		break;
-	case EBakedCurvatureTypeMode::Max:
-		CurvatureMapSettings.CurvatureType = (int32)FMeshCurvatureMapEvaluator::ECurvatureType::MaxPrincipal;
-		break;
-	case EBakedCurvatureTypeMode::Min:
-		CurvatureMapSettings.CurvatureType = (int32)FMeshCurvatureMapEvaluator::ECurvatureType::MinPrincipal;
-		break;
-	}
-	switch (CurvatureMapProps->ColorMapping)
-	{
-	default:
-	case EBakedCurvatureColorMode::Grayscale:
-		CurvatureMapSettings.ColorMode = (int32)FMeshCurvatureMapEvaluator::EColorMode::BlackGrayWhite;
-		break;
-	case EBakedCurvatureColorMode::RedBlue:
-		CurvatureMapSettings.ColorMode = (int32)FMeshCurvatureMapEvaluator::EColorMode::RedBlue;
-		break;
-	case EBakedCurvatureColorMode::RedGreenBlue:
-		CurvatureMapSettings.ColorMode = (int32)FMeshCurvatureMapEvaluator::EColorMode::RedGreenBlue;
-		break;
-	}
-	switch (CurvatureMapProps->Clamping)
-	{
-	default:
-	case EBakedCurvatureClampMode::None:
-		CurvatureMapSettings.ClampMode = (int32)FMeshCurvatureMapEvaluator::EClampMode::FullRange;
-		break;
-	case EBakedCurvatureClampMode::OnlyPositive:
-		CurvatureMapSettings.ClampMode = (int32)FMeshCurvatureMapEvaluator::EClampMode::Positive;
-		break;
-	case EBakedCurvatureClampMode::OnlyNegative:
-		CurvatureMapSettings.ClampMode = (int32)FMeshCurvatureMapEvaluator::EClampMode::Negative;
-		break;
-	}
-
-	if (!(CachedCurvatureMapSettings == CurvatureMapSettings))
-	{
-		CachedCurvatureMapSettings = CurvatureMapSettings;
-		ResultState = EBakeOpState::Evaluate;
-	}
-	return ResultState;
-}
-
-
-EBakeOpState UBakeMeshAttributeMapsTool::UpdateResult_MeshProperty()
-{
-	EBakeOpState ResultState = EBakeOpState::Clean;
-
-	const int32 ImageSize = (int32)Settings->Resolution;
-	const FImageDimensions Dimensions(ImageSize, ImageSize);
-
-	FMeshPropertyMapSettings MeshPropertyMapSettings;
-	MeshPropertyMapSettings.Dimensions = Dimensions;
-
-	if (!(CachedMeshPropertyMapSettings == MeshPropertyMapSettings))
-	{
-		CachedMeshPropertyMapSettings = MeshPropertyMapSettings;
-		ResultState = EBakeOpState::Evaluate;
-	}
-	return ResultState;
-}
-
-
-EBakeOpState UBakeMeshAttributeMapsTool::UpdateResult_Texture2DImage()
-{
-	EBakeOpState ResultState = EBakeOpState::Clean;
-
-	const int32 ImageSize = (int32)Settings->Resolution;
-	const FImageDimensions Dimensions(ImageSize, ImageSize);
-
-	FTexture2DImageSettings NewSettings;
-	NewSettings.Dimensions = Dimensions;
-	NewSettings.UVLayer = 0;
-
-	const FDynamicMeshUVOverlay* UVOverlay = DetailMesh->Attributes()->GetUVLayer(NewSettings.UVLayer);
-	if (UVOverlay == nullptr)
-	{
-		GetToolManager()->DisplayMessage(LOCTEXT("InvalidUVWarning", "The Source Mesh does not have the selected UV layer"), EToolMessageLevel::UserWarning);
-		return EBakeOpState::Invalid;
-	}
-	
-	if (Texture2DProps->SourceTexture == nullptr)
-	{
-		GetToolManager()->DisplayMessage(LOCTEXT("InvalidTextureWarning", "The Source Texture is not valid"), EToolMessageLevel::UserWarning);
-		return EBakeOpState::Invalid;
-	}
-
-	{
-		CachedTextureImage = MakeShared<UE::Geometry::TImageBuilder<FVector4f>, ESPMode::ThreadSafe>();
-		if (!UE::AssetUtils::ReadTexture(Texture2DProps->SourceTexture, *CachedTextureImage, bPreferPlatformData))
-		{
-			GetToolManager()->DisplayMessage(LOCTEXT("CannotReadTextureWarning", "Cannot read from the source texture"), EToolMessageLevel::UserWarning);
-			return EBakeOpState::Invalid;
-		}
-	}
-
-	if (!(CachedTexture2DImageSettings == NewSettings))
-	{
-		CachedTexture2DImageSettings = NewSettings;
-		ResultState = EBakeOpState::Evaluate;
-	}
-	return ResultState;
-}
-
-
-EBakeOpState UBakeMeshAttributeMapsTool::UpdateResult_MultiTexture()
-{
-	EBakeOpState ResultState = EBakeOpState::Clean;
-
-	const int32 ImageSize = (int32)Settings->Resolution;
-	const FImageDimensions Dimensions(ImageSize, ImageSize);
-
-	FTexture2DImageSettings NewSettings;
-	NewSettings.Dimensions = Dimensions;
-	NewSettings.UVLayer = MultiTextureProps->UVLayer;
-
-	const FDynamicMeshUVOverlay* UVOverlay = DetailMesh->Attributes()->GetUVLayer(NewSettings.UVLayer);
-	if (UVOverlay == nullptr)
-	{
-		GetToolManager()->DisplayMessage(LOCTEXT("InvalidUVWarning", "The Source Mesh does not have the selected UV layer"), EToolMessageLevel::UserWarning);
-		return EBakeOpState::Invalid;
-	}
-
-	const int NumMaterialIDs = MultiTextureProps->MaterialIDSourceTextures.Num();
-	CachedMultiTextures.Reset();
-	CachedMultiTextures.SetNum(NumMaterialIDs);
-	int NumValidTextures = 0;
-	for ( int MaterialID = 0; MaterialID < NumMaterialIDs; ++MaterialID)
-	{
-		if (UTexture2D* Texture = MultiTextureProps->MaterialIDSourceTextures[MaterialID])
-		{
-			CachedMultiTextures[MaterialID] = MakeShared<TImageBuilder<FVector4f>, ESPMode::ThreadSafe>();
-			if (!UE::AssetUtils::ReadTexture(Texture, *CachedMultiTextures[MaterialID], bPreferPlatformData))
-			{
-				GetToolManager()->DisplayMessage(LOCTEXT("CannotReadTextureWarning", "Cannot read from the source texture"), EToolMessageLevel::UserWarning);
-				return EBakeOpState::Invalid;
-			}
-			++NumValidTextures;
-		}
-	}
-	if (NumValidTextures == 0)
-	{
-		GetToolManager()->DisplayMessage(LOCTEXT("InvalidTextureWarning", "The Source Texture is not valid"), EToolMessageLevel::UserWarning);
-		return EBakeOpState::Invalid;
-	}
-
-	if (!(CachedTexture2DImageSettings == NewSettings))
-	{
-		CachedTexture2DImageSettings = NewSettings;
-		ResultState = EBakeOpState::Evaluate;
-	}
-	return ResultState;
-}
-
-
 void UBakeMeshAttributeMapsTool::UpdateVisualization()
 {
 	PreviewMesh->SetOverrideRenderMaterial(PreviewMaterial);
@@ -884,10 +670,10 @@ void UBakeMeshAttributeMapsTool::UpdateOnModeChange()
 	OnMapTypesUpdated(Settings);
 
 	// Update tool property sets.
-	SetToolPropertySourceEnabled(OcclusionMapProps, false);
-	SetToolPropertySourceEnabled(CurvatureMapProps, false);
-	SetToolPropertySourceEnabled(Texture2DProps, false);
-	SetToolPropertySourceEnabled(MultiTextureProps, false);
+	SetToolPropertySourceEnabled(OcclusionSettings, false);
+	SetToolPropertySourceEnabled(CurvatureSettings, false);
+	SetToolPropertySourceEnabled(TextureSettings, false);
+	SetToolPropertySourceEnabled(MultiTextureSettings, false);
 
 	for (const EBakeMapType MapType : ALL_BAKE_MAP_TYPES)
 	{
@@ -898,10 +684,10 @@ void UBakeMeshAttributeMapsTool::UpdateOnModeChange()
 		case EBakeMapType::AmbientOcclusion:
 		case EBakeMapType::BentNormal:
 		case EBakeMapType::Occlusion:
-			SetToolPropertySourceEnabled(OcclusionMapProps, true);
+			SetToolPropertySourceEnabled(OcclusionSettings, true);
 			break;
 		case EBakeMapType::Curvature:
-			SetToolPropertySourceEnabled(CurvatureMapProps, true);
+			SetToolPropertySourceEnabled(CurvatureSettings, true);
 			break;
 		case EBakeMapType::ObjectSpaceNormal:
 		case EBakeMapType::FaceNormal:
@@ -910,10 +696,10 @@ void UBakeMeshAttributeMapsTool::UpdateOnModeChange()
 		case EBakeMapType::VertexColor:
 			break;
 		case EBakeMapType::Texture:
-			SetToolPropertySourceEnabled(Texture2DProps, true);
+			SetToolPropertySourceEnabled(TextureSettings, true);
 			break;
 		case EBakeMapType::MultiTexture:
-			SetToolPropertySourceEnabled(MultiTextureProps, true);
+			SetToolPropertySourceEnabled(MultiTextureSettings, true);
 			break;
 		default:
 			break;
@@ -935,7 +721,7 @@ void UBakeMeshAttributeMapsTool::GatherAnalytics(FBakeAnalytics::FMeshSettings&
 {
 	if (FEngineAnalytics::IsAvailable())
 	{
-		Data.NumTargetMeshTris = BaseMesh.TriangleCount();
+		Data.NumTargetMeshTris = TargetMesh.TriangleCount();
 		Data.NumDetailMesh = 1;
 		Data.NumDetailMeshTris = DetailMesh->TriangleCount();
 	}
diff --git a/Engine/Plugins/Experimental/MeshModelingToolsetExp/Source/MeshModelingToolsExp/Private/BakeMeshAttributeMapsToolBase.cpp b/Engine/Plugins/Experimental/MeshModelingToolsetExp/Source/MeshModelingToolsExp/Private/BakeMeshAttributeMapsToolBase.cpp
--- a/Engine/Plugins/Experimental/MeshModelingToolsetExp/Source/MeshModelingToolsExp/Private/BakeMeshAttributeMapsToolBase.cpp
+++ b/Engine/Plugins/Experimental/MeshModelingToolsetExp/Source/MeshModelingToolsExp/Private/BakeMeshAttributeMapsToolBase.cpp
@@ -71,7 +71,7 @@ void UBakeMeshAttributeMapsToolBase::Setup()
 
 void UBakeMeshAttributeMapsToolBase::PostSetup()
 {
-	VisualizationProps = NewObject<UBakedOcclusionMapVisualizationProperties>(this);
+	VisualizationProps = NewObject<UBakeVisualizationProperties>(this);
 	VisualizationProps->RestoreProperties(this);
 	AddToolPropertySource(VisualizationProps);
 
@@ -225,15 +225,15 @@ void UBakeMeshAttributeMapsToolBase::UpdatePreview(const EBakeMapType PreviewMap
 		PreviewMaterial->SetTextureParameterValue(TEXT("ColorMap"), PreviewMap);
 		break;
 	}
-	PreviewMaterial->SetScalarParameterValue(TEXT("UVChannel"), CachedBakeCacheSettings.TargetUVLayer);
-	BentNormalPreviewMaterial->SetScalarParameterValue(TEXT("UVChannel"), CachedBakeCacheSettings.TargetUVLayer);
+	PreviewMaterial->SetScalarParameterValue(TEXT("UVChannel"), CachedBakeSettings.TargetUVLayer);
+	BentNormalPreviewMaterial->SetScalarParameterValue(TEXT("UVChannel"), CachedBakeSettings.TargetUVLayer);
 }
 
 
 void UBakeMeshAttributeMapsToolBase::OnMapsUpdated(const TUniquePtr<UE::Geometry::FMeshMapBaker>& NewResult)
 {
 	const FImageDimensions BakeDimensions = NewResult->GetDimensions();
-	const EBakeTextureBitDepth Format = CachedBakeCacheSettings.BitDepth;
+	const EBakeTextureBitDepth Format = CachedBakeSettings.BitDepth;
 	const int32 NumEval = NewResult->NumEvaluators();
 	for (int32 EvalIdx = 0; EvalIdx < NumEval; ++EvalIdx)
 	{
@@ -344,7 +344,7 @@ void UBakeMeshAttributeMapsToolBase::OnMapsUpdated(const TUniquePtr<UE::Geometry
 		}
 	}
 
-	GatherAnalytics(*NewResult, CachedBakeCacheSettings, BakeAnalytics);
+	GatherAnalytics(*NewResult, CachedBakeSettings, BakeAnalytics);
 	UpdateVisualization();
 	GetToolManager()->PostInvalidation();
 }
@@ -472,7 +472,7 @@ void UBakeMeshAttributeMapsToolBase::GatherAnalytics(FBakeAnalytics::FMeshSettin
 }
 
 
-void UBakeMeshAttributeMapsToolBase::GatherAnalytics(const FMeshMapBaker& Result, const FBakeCacheSettings& Settings, FBakeAnalytics& Data)
+void UBakeMeshAttributeMapsToolBase::GatherAnalytics(const FMeshMapBaker& Result, const FBakeSettings& Settings, FBakeAnalytics& Data)
 {
 	if (!FEngineAnalytics::IsAvailable())
 	{
diff --git a/Engine/Plugins/Experimental/MeshModelingToolsetExp/Source/MeshModelingToolsExp/Private/BakeMeshAttributeTool.cpp b/Engine/Plugins/Experimental/MeshModelingToolsetExp/Source/MeshModelingToolsExp/Private/BakeMeshAttributeTool.cpp
--- a/Engine/Plugins/Experimental/MeshModelingToolsetExp/Source/MeshModelingToolsExp/Private/BakeMeshAttributeTool.cpp
+++ b/Engine/Plugins/Experimental/MeshModelingToolsetExp/Source/MeshModelingToolsExp/Private/BakeMeshAttributeTool.cpp
@@ -4,6 +4,9 @@
 #include "InteractiveToolManager.h"
 #include "ModelingToolTargetUtil.h"
 #include "ToolSetupUtil.h"
+#include "AssetUtils/Texture2DUtil.h"
+
+#include "Sampling/MeshCurvatureMapEvaluator.h"
 
 using namespace UE::Geometry;
 
@@ -20,6 +23,220 @@ void UBakeMeshAttributeTool::Setup()
 }
 
 
+EBakeOpState UBakeMeshAttributeTool::UpdateResult_Normal(const FImageDimensions& Dimensions)
+{
+	EBakeOpState ResultState = EBakeOpState::Clean;
+
+	FNormalMapSettings NormalMapSettings;
+	NormalMapSettings.Dimensions = Dimensions;
+
+	if (!(CachedNormalMapSettings == NormalMapSettings))
+	{
+		CachedNormalMapSettings = NormalMapSettings;
+		ResultState = EBakeOpState::Evaluate;
+	}
+	return ResultState;
+}
+
+
+EBakeOpState UBakeMeshAttributeTool::UpdateResult_Occlusion(const FImageDimensions& Dimensions)
+{
+	EBakeOpState ResultState = EBakeOpState::Clean;
+	if (!OcclusionSettings)
+	{
+		return EBakeOpState::Invalid;
+	}
+
+	FOcclusionMapSettings OcclusionMapSettings;
+	OcclusionMapSettings.Dimensions = Dimensions;
+	OcclusionMapSettings.MaxDistance = (OcclusionSettings->MaxDistance == 0) ? TNumericLimits<float>::Max() : OcclusionSettings->MaxDistance;
+	OcclusionMapSettings.OcclusionRays = OcclusionSettings->OcclusionRays;
+	OcclusionMapSettings.SpreadAngle = OcclusionSettings->SpreadAngle;
+	OcclusionMapSettings.BiasAngle = OcclusionSettings->BiasAngle;
+
+	if ( !(CachedOcclusionMapSettings == OcclusionMapSettings) )
+	{
+		CachedOcclusionMapSettings = OcclusionMapSettings;
+		ResultState = EBakeOpState::Evaluate;
+	}
+	return ResultState;
+}
+
+
+EBakeOpState UBakeMeshAttributeTool::UpdateResult_Curvature(const FImageDimensions& Dimensions)
+{
+	EBakeOpState ResultState = EBakeOpState::Clean;
+	if (!CurvatureSettings)
+	{
+		return EBakeOpState::Invalid;
+	}
+
+	FCurvatureMapSettings CurvatureMapSettings;
+	CurvatureMapSettings.Dimensions = Dimensions;
+	CurvatureMapSettings.RangeMultiplier = CurvatureSettings->ColorRangeMultiplier;
+	CurvatureMapSettings.MinRangeMultiplier = CurvatureSettings->MinRangeMultiplier;
+	switch (CurvatureSettings->CurvatureType)
+	{
+	default:
+	case EBakeCurvatureTypeMode::MeanAverage:
+		CurvatureMapSettings.CurvatureType = (int32)FMeshCurvatureMapEvaluator::ECurvatureType::Mean;
+		break;
+	case EBakeCurvatureTypeMode::Gaussian:
+		CurvatureMapSettings.CurvatureType = (int32)FMeshCurvatureMapEvaluator::ECurvatureType::Gaussian;
+		break;
+	case EBakeCurvatureTypeMode::Max:
+		CurvatureMapSettings.CurvatureType = (int32)FMeshCurvatureMapEvaluator::ECurvatureType::MaxPrincipal;
+		break;
+	case EBakeCurvatureTypeMode::Min:
+		CurvatureMapSettings.CurvatureType = (int32)FMeshCurvatureMapEvaluator::ECurvatureType::MinPrincipal;
+		break;
+	}
+	switch (CurvatureSettings->ColorMapping)
+	{
+	default:
+	case EBakeCurvatureColorMode::Grayscale:
+		CurvatureMapSettings.ColorMode = (int32)FMeshCurvatureMapEvaluator::EColorMode::BlackGrayWhite;
+		break;
+	case EBakeCurvatureColorMode::RedBlue:
+		CurvatureMapSettings.ColorMode = (int32)FMeshCurvatureMapEvaluator::EColorMode::RedBlue;
+		break;
+	case EBakeCurvatureColorMode::RedGreenBlue:
+		CurvatureMapSettings.ColorMode = (int32)FMeshCurvatureMapEvaluator::EColorMode::RedGreenBlue;
+		break;
+	}
+	switch (CurvatureSettings->Clamping)
+	{
+	default:
+	case EBakeCurvatureClampMode::None:
+		CurvatureMapSettings.ClampMode = (int32)FMeshCurvatureMapEvaluator::EClampMode::FullRange;
+		break;
+	case EBakeCurvatureClampMode::OnlyPositive:
+		CurvatureMapSettings.ClampMode = (int32)FMeshCurvatureMapEvaluator::EClampMode::Positive;
+		break;
+	case EBakeCurvatureClampMode::OnlyNegative:
+		CurvatureMapSettings.ClampMode = (int32)FMeshCurvatureMapEvaluator::EClampMode::Negative;
+		break;
+	}
+
+	if (!(CachedCurvatureMapSettings == CurvatureMapSettings))
+	{
+		CachedCurvatureMapSettings = CurvatureMapSettings;
+		ResultState = EBakeOpState::Evaluate;
+	}
+	return ResultState;
+}
+
+
+EBakeOpState UBakeMeshAttributeTool::UpdateResult_MeshProperty(const FImageDimensions& Dimensions)
+{
+	EBakeOpState ResultState = EBakeOpState::Clean;
+
+	FMeshPropertyMapSettings MeshPropertyMapSettings;
+	MeshPropertyMapSettings.Dimensions = Dimensions;
+
+	if (!(CachedMeshPropertyMapSettings == MeshPropertyMapSettings))
+	{
+		CachedMeshPropertyMapSettings = MeshPropertyMapSettings;
+		ResultState = EBakeOpState::Evaluate;
+	}
+	return ResultState;
+}
+
+
+EBakeOpState UBakeMeshAttributeTool::UpdateResult_Texture2DImage(const FImageDimensions& Dimensions, const FDynamicMesh3* DetailMesh)
+{
+	EBakeOpState ResultState = EBakeOpState::Clean;
+	if (!TextureSettings || !DetailMesh)
+	{
+		return EBakeOpState::Invalid;
+	}
+
+	FTexture2DSettings NewSettings;
+	NewSettings.Dimensions = Dimensions;
+	NewSettings.UVLayer = TextureSettings->UVLayer;
+
+	const FDynamicMeshUVOverlay* UVOverlay = DetailMesh->Attributes()->GetUVLayer(NewSettings.UVLayer);
+	if (UVOverlay == nullptr)
+	{
+		GetToolManager()->DisplayMessage(LOCTEXT("InvalidUVWarning", "The Source Mesh does not have the selected UV layer"), EToolMessageLevel::UserWarning);
+		return EBakeOpState::Invalid;
+	}
+	
+	if (TextureSettings->SourceTexture == nullptr)
+	{
+		GetToolManager()->DisplayMessage(LOCTEXT("InvalidTextureWarning", "The Source Texture is not valid"), EToolMessageLevel::UserWarning);
+		return EBakeOpState::Invalid;
+	}
+
+	{
+		CachedTextureImage = MakeShared<UE::Geometry::TImageBuilder<FVector4f>, ESPMode::ThreadSafe>();
+		if (!UE::AssetUtils::ReadTexture(TextureSettings->SourceTexture, *CachedTextureImage, bPreferPlatformData))
+		{
+			GetToolManager()->DisplayMessage(LOCTEXT("CannotReadTextureWarning", "Cannot read from the source texture"), EToolMessageLevel::UserWarning);
+			return EBakeOpState::Invalid;
+		}
+	}
+
+	if (!(CachedTexture2DSettings == NewSettings))
+	{
+		CachedTexture2DSettings = NewSettings;
+		ResultState = EBakeOpState::Evaluate;
+	}
+	return ResultState;
+}
+
+
+EBakeOpState UBakeMeshAttributeTool::UpdateResult_MultiTexture(const FImageDimensions& Dimensions, const FDynamicMesh3* DetailMesh)
+{
+	EBakeOpState ResultState = EBakeOpState::Clean;
+	if (!MultiTextureSettings || !DetailMesh)
+	{
+		return EBakeOpState::Invalid;
+	}
+
+	FTexture2DSettings NewSettings;
+	NewSettings.Dimensions = Dimensions;
+	NewSettings.UVLayer = MultiTextureSettings->UVLayer;
+
+	const FDynamicMeshUVOverlay* UVOverlay = DetailMesh->Attributes()->GetUVLayer(NewSettings.UVLayer);
+	if (UVOverlay == nullptr)
+	{
+		GetToolManager()->DisplayMessage(LOCTEXT("InvalidUVWarning", "The Source Mesh does not have the selected UV layer"), EToolMessageLevel::UserWarning);
+		return EBakeOpState::Invalid;
+	}
+
+	const int NumMaterialIDs = MultiTextureSettings->MaterialIDSourceTextures.Num();
+	CachedMultiTextures.Reset();
+	CachedMultiTextures.SetNum(NumMaterialIDs);
+	int NumValidTextures = 0;
+	for ( int MaterialID = 0; MaterialID < NumMaterialIDs; ++MaterialID)
+	{
+		if (UTexture2D* Texture = MultiTextureSettings->MaterialIDSourceTextures[MaterialID])
+		{
+			CachedMultiTextures[MaterialID] = MakeShared<TImageBuilder<FVector4f>, ESPMode::ThreadSafe>();
+			if (!UE::AssetUtils::ReadTexture(Texture, *CachedMultiTextures[MaterialID], bPreferPlatformData))
+			{
+				GetToolManager()->DisplayMessage(LOCTEXT("CannotReadTextureWarning", "Cannot read from the source texture"), EToolMessageLevel::UserWarning);
+				return EBakeOpState::Invalid;
+			}
+			++NumValidTextures;
+		}
+	}
+	if (NumValidTextures == 0)
+	{
+		GetToolManager()->DisplayMessage(LOCTEXT("InvalidTextureWarning", "The Source Texture is not valid"), EToolMessageLevel::UserWarning);
+		return EBakeOpState::Invalid;
+	}
+
+	if (!(CachedMultiTexture2DSettings == NewSettings))
+	{
+		CachedMultiTexture2DSettings = NewSettings;
+		ResultState = EBakeOpState::Evaluate;
+	}
+	return ResultState;
+}
+
+
 int UBakeMeshAttributeTool::SelectColorTextureToBake(const TArray<UTexture*>& Textures)
 {
 	TArray<int> TextureVotes;
@@ -98,5 +315,24 @@ void UBakeMeshAttributeTool::UpdateMultiTextureMaterialIDs(
 	});
 }
 
+
+void UBakeMeshAttributeTool::UpdateUVLayerNames(FString& UVLayer, TArray<FString>& UVLayerNamesList, const FDynamicMesh3& Mesh)
+{
+	UVLayerNamesList.Reset();
+	int32 FoundIndex = -1;
+	for (int32 k = 0; k < Mesh.Attributes()->NumUVLayers(); ++k)
+	{
+		UVLayerNamesList.Add(FString::Printf(TEXT("UV %d"), k));
+		if (UVLayer == UVLayerNamesList.Last())
+		{
+			FoundIndex = k;
+		}
+	}
+	if (FoundIndex == -1)
+	{
+		UVLayer = UVLayerNamesList[0];
+	}
+}
+
 #undef LOCTEXT_NAMESPACE
 
diff --git a/Engine/Plugins/Experimental/MeshModelingToolsetExp/Source/MeshModelingToolsExp/Private/BakeMeshAttributeVertexTool.cpp b/Engine/Plugins/Experimental/MeshModelingToolsetExp/Source/MeshModelingToolsExp/Private/BakeMeshAttributeVertexTool.cpp
--- a/Engine/Plugins/Experimental/MeshModelingToolsetExp/Source/MeshModelingToolsExp/Private/BakeMeshAttributeVertexTool.cpp
+++ b/Engine/Plugins/Experimental/MeshModelingToolsetExp/Source/MeshModelingToolsExp/Private/BakeMeshAttributeVertexTool.cpp
@@ -66,11 +66,11 @@ public:
 	UBakeMeshAttributeVertexTool::FBakeSettings BakeSettings;
 	FOcclusionMapSettings OcclusionSettings;
 	FCurvatureMapSettings CurvatureSettings;
-	FTexture2DImageSettings TextureSettings;
+	FTexture2DSettings TextureSettings;
+	FTexture2DSettings MultiTextureSettings;
 
 	// Texture2DImage & MultiTexture settings
 	using ImagePtr = TSharedPtr<UE::Geometry::TImageBuilder<FVector4f>, ESPMode::ThreadSafe>;
-	const FDynamicMeshUVOverlay* UVOverlay = nullptr;
 	ImagePtr TextureImage;
 	TArray<ImagePtr> MaterialIDTextures;
 
@@ -205,7 +205,7 @@ public:
 			case EBakeMapType::MultiTexture:
 			{
 				TSharedPtr<FMeshMultiResampleImageEvaluator, ESPMode::ThreadSafe> TextureEval = MakeShared<FMeshMultiResampleImageEvaluator, ESPMode::ThreadSafe>();
-				TextureEval->DetailUVLayer = TextureSettings.UVLayer;
+				TextureEval->DetailUVLayer = MultiTextureSettings.UVLayer;
 				TextureEval->MultiTextures = MaterialIDTextures;
 				Baker->ColorEvaluator = TextureEval;
 				break;
@@ -263,10 +263,10 @@ void UBakeMeshAttributeVertexTool::Setup()
 
 	PreviewMesh->ProcessMesh([this](const FDynamicMesh3& Mesh)
 	{
-		BaseMesh.Copy(Mesh);
-		BaseSpatial.SetMesh(&BaseMesh, true);
-		BaseMeshTangents = MakeShared<FMeshTangentsd, ESPMode::ThreadSafe>(&BaseMesh);
-		BaseMeshTangents->CopyTriVertexTangents(Mesh);
+		TargetMesh.Copy(Mesh);
+		TargetSpatial.SetMesh(&TargetMesh, true);
+		TargetMeshTangents = MakeShared<FMeshTangentsd, ESPMode::ThreadSafe>(&TargetMesh);
+		TargetMeshTangents->CopyTriVertexTangents(Mesh);
 	});
 
 	UToolTarget* Target = Targets[0];
@@ -304,7 +304,7 @@ void UBakeMeshAttributeVertexTool::Setup()
 	Settings->WatchProperty(Settings->bSplitAtNormalSeams, [this](bool) { bColorTopologyValid = false; OpState |= EBakeOpState::Evaluate; });
 	Settings->WatchProperty(Settings->bSplitAtUVSeams, [this](bool) { bColorTopologyValid = false; OpState |= EBakeOpState::Evaluate; });
 
-	OcclusionSettings = NewObject<UBakedOcclusionMapToolProperties>(this);
+	OcclusionSettings = NewObject<UBakeOcclusionMapToolProperties>(this);
 	OcclusionSettings->RestoreProperties(this);
 	AddToolPropertySource(OcclusionSettings);
 	SetToolPropertySourceEnabled(OcclusionSettings, false);
@@ -313,24 +313,24 @@ void UBakeMeshAttributeVertexTool::Setup()
 	OcclusionSettings->WatchProperty(OcclusionSettings->SpreadAngle, [this](float) { OpState |= EBakeOpState::Evaluate; });
 	OcclusionSettings->WatchProperty(OcclusionSettings->BiasAngle, [this](float) { OpState |= EBakeOpState::Evaluate; });
 
-	CurvatureSettings = NewObject<UBakedCurvatureMapToolProperties>(this);
+	CurvatureSettings = NewObject<UBakeCurvatureMapToolProperties>(this);
 	CurvatureSettings->RestoreProperties(this);
 	AddToolPropertySource(CurvatureSettings);
 	SetToolPropertySourceEnabled(CurvatureSettings, false);
 	CurvatureSettings->WatchProperty(CurvatureSettings->ColorRangeMultiplier, [this](float) { OpState |= EBakeOpState::Evaluate; });
 	CurvatureSettings->WatchProperty(CurvatureSettings->MinRangeMultiplier, [this](float) { OpState |= EBakeOpState::Evaluate; });
-	CurvatureSettings->WatchProperty(CurvatureSettings->CurvatureType, [this](EBakedCurvatureTypeMode) { OpState |= EBakeOpState::Evaluate; });
-	CurvatureSettings->WatchProperty(CurvatureSettings->ColorMapping, [this](EBakedCurvatureColorMode) { OpState |= EBakeOpState::Evaluate; });
-	CurvatureSettings->WatchProperty(CurvatureSettings->Clamping, [this](EBakedCurvatureClampMode) { OpState |= EBakeOpState::Evaluate; });
+	CurvatureSettings->WatchProperty(CurvatureSettings->CurvatureType, [this](EBakeCurvatureTypeMode) { OpState |= EBakeOpState::Evaluate; });
+	CurvatureSettings->WatchProperty(CurvatureSettings->ColorMapping, [this](EBakeCurvatureColorMode) { OpState |= EBakeOpState::Evaluate; });
+	CurvatureSettings->WatchProperty(CurvatureSettings->Clamping, [this](EBakeCurvatureClampMode) { OpState |= EBakeOpState::Evaluate; });
 
-	TextureSettings = NewObject<UBakedTexture2DImageProperties>(this);
+	TextureSettings = NewObject<UBakeTexture2DProperties>(this);
 	TextureSettings->RestoreProperties(this);
 	AddToolPropertySource(TextureSettings);
 	SetToolPropertySourceEnabled(TextureSettings, false);
 	TextureSettings->WatchProperty(TextureSettings->UVLayer, [this](float) { OpState |= EBakeOpState::Evaluate; });
 	TextureSettings->WatchProperty(TextureSettings->SourceTexture, [this](UTexture2D*) { OpState |= EBakeOpState::Evaluate; });
 
-	MultiTextureSettings = NewObject<UBakedMultiTexture2DImageProperties>(this);
+	MultiTextureSettings = NewObject<UBakeMultiTexture2DProperties>(this);
 	MultiTextureSettings->RestoreProperties(this);
 	AddToolPropertySource(MultiTextureSettings);
 	SetToolPropertySourceEnabled(MultiTextureSettings, false);
@@ -436,17 +436,17 @@ TUniquePtr<UE::Geometry::TGenericDataOperator<FMeshVertexBaker>> UBakeMeshAttrib
 	TUniquePtr<FMeshVertexBakerOp> Op = MakeUnique<FMeshVertexBakerOp>();
 	Op->DetailMesh = DetailMesh;
 	Op->DetailSpatial = DetailSpatial;
-	Op->BaseMesh = &BaseMesh;
-	Op->BaseMeshTangents = BaseMeshTangents;
+	Op->BaseMesh = &TargetMesh;
+	Op->BaseMeshTangents = TargetMeshTangents;
 	Op->BakeSettings = CachedBakeSettings;
 	Op->OcclusionSettings = CachedOcclusionMapSettings;
 	Op->CurvatureSettings = CachedCurvatureMapSettings;
-	Op->TextureSettings = CachedTexture2DImageSettings;
+	Op->TextureSettings = CachedTexture2DSettings;
+	Op->MultiTextureSettings = CachedMultiTexture2DSettings;
 
 	// Texture2DImage & MultiTexture settings
 	Op->TextureImage = CachedTextureImage;
 	Op->MaterialIDTextures = CachedMultiTextures;
-	Op->UVOverlay = DetailMesh->Attributes()->GetUVLayer(CachedTexture2DImageSettings.UVLayer);
 	return Op;
 }
 
@@ -593,13 +593,14 @@ void UBakeMeshAttributeVertexTool::UpdateColorTopology()
 	});
 
 	// Update BaseMesh color topology.
-	BaseMesh.EnableAttributes();
-	BaseMesh.Attributes()->DisablePrimaryColors();
-	BaseMesh.Attributes()->EnablePrimaryColors();
+	TargetMesh.EnableAttributes();
+	TargetMesh.Attributes()->DisablePrimaryColors();
+	TargetMesh.Attributes()->EnablePrimaryColors();
 	PreviewMesh->ProcessMesh([this](const FDynamicMesh3& Mesh)
 	{
-		BaseMesh.Attributes()->PrimaryColors()->Copy(*Mesh.Attributes()->PrimaryColors());
+		TargetMesh.Attributes()->PrimaryColors()->Copy(*Mesh.Attributes()->PrimaryColors());
 	});
+	NumColorElements = TargetMesh.Attributes()->PrimaryColors()->ElementCount();
 
 	bColorTopologyValid = true;
 }
@@ -644,31 +645,32 @@ void UBakeMeshAttributeVertexTool::UpdateResult()
 	OpState &= ~EBakeOpState::Invalid;
 
 	// Validate bake inputs
+	const FImageDimensions Dimensions(NumColorElements, 1);
 	if (CachedBakeSettings.OutputMode == EBakeVertexOutput::RGBA)
 	{
 		switch(BakeSettings.OutputType)
 		{
 		case EBakeMapType::TangentSpaceNormal:
-			OpState |= UpdateResult_Normal();
+			OpState |= UpdateResult_Normal(Dimensions);
 			break;
 		case EBakeMapType::AmbientOcclusion:
 		case EBakeMapType::BentNormal:
-			OpState |= UpdateResult_Occlusion();
+			OpState |= UpdateResult_Occlusion(Dimensions);
 			break;
 		case EBakeMapType::Curvature:
-			OpState |= UpdateResult_Curvature();
+			OpState |= UpdateResult_Curvature(Dimensions);
 			break;
 		case EBakeMapType::ObjectSpaceNormal:
 		case EBakeMapType::FaceNormal:
 		case EBakeMapType::Position:
 		case EBakeMapType::MaterialID:
-			OpState |= UpdateResult_MeshProperty();
+			OpState |= UpdateResult_MeshProperty(Dimensions);
 			break;
 		case EBakeMapType::Texture:
-			OpState |= UpdateResult_Texture2DImage();
+			OpState |= UpdateResult_Texture2DImage(Dimensions, DetailMesh.Get());
 			break;
 		case EBakeMapType::MultiTexture:
-			OpState |= UpdateResult_MultiTexture();
+			OpState |= UpdateResult_MultiTexture(Dimensions, DetailMesh.Get());
 			break;
 		default:
 			break;
@@ -679,11 +681,11 @@ void UBakeMeshAttributeVertexTool::UpdateResult()
 		// The enabled state of these settings are precomputed in UpdateOnModeChange().
 		if (OcclusionSettings->IsPropertySetEnabled())
 		{
-			OpState |= UpdateResult_Occlusion();
+			OpState |= UpdateResult_Occlusion(Dimensions);
 		}
 		if (CurvatureSettings->IsPropertySetEnabled())
 		{
-			OpState |= UpdateResult_Curvature();
+			OpState |= UpdateResult_Curvature(Dimensions);
 		}
 	}
 
@@ -722,178 +724,6 @@ void UBakeMeshAttributeVertexTool::OnResultUpdated(const TUniquePtr<FMeshVertexB
 	GatherAnalytics(*NewResult, CachedBakeSettings, BakeAnalytics);
 }
 
-EBakeOpState UBakeMeshAttributeVertexTool::UpdateResult_Normal()
-{
-	// No settings to configure, always valid to evaluate.
-	return EBakeOpState::Evaluate;
-}
-
-EBakeOpState UBakeMeshAttributeVertexTool::UpdateResult_Occlusion()
-{
-	EBakeOpState ResultState = EBakeOpState::Clean;
-
-	FOcclusionMapSettings OcclusionMapSettings;
-	OcclusionMapSettings.MaxDistance = (OcclusionSettings->MaxDistance == 0) ? TNumericLimits<float>::Max() : OcclusionSettings->MaxDistance;
-	OcclusionMapSettings.OcclusionRays = OcclusionSettings->OcclusionRays;
-	OcclusionMapSettings.SpreadAngle = OcclusionSettings->SpreadAngle;
-	OcclusionMapSettings.BiasAngle = OcclusionSettings->BiasAngle;
-
-	if (!(CachedOcclusionMapSettings == OcclusionMapSettings))
-	{
-		CachedOcclusionMapSettings = OcclusionMapSettings;
-		ResultState = EBakeOpState::Evaluate;
-	}
-	return ResultState;
-}
-
-EBakeOpState UBakeMeshAttributeVertexTool::UpdateResult_Curvature()
-{
-	EBakeOpState ResultState = EBakeOpState::Clean;
-
-	FCurvatureMapSettings CurvatureMapSettings;
-	CurvatureMapSettings.RangeMultiplier = CurvatureSettings->ColorRangeMultiplier;
-	CurvatureMapSettings.MinRangeMultiplier = CurvatureSettings->MinRangeMultiplier;
-	switch (CurvatureSettings->CurvatureType)
-	{
-	default:
-	case EBakedCurvatureTypeMode::MeanAverage:
-		CurvatureMapSettings.CurvatureType = (int32)FMeshCurvatureMapEvaluator::ECurvatureType::Mean;
-		break;
-	case EBakedCurvatureTypeMode::Gaussian:
-		CurvatureMapSettings.CurvatureType = (int32)FMeshCurvatureMapEvaluator::ECurvatureType::Gaussian;
-		break;
-	case EBakedCurvatureTypeMode::Max:
-		CurvatureMapSettings.CurvatureType = (int32)FMeshCurvatureMapEvaluator::ECurvatureType::MaxPrincipal;
-		break;
-	case EBakedCurvatureTypeMode::Min:
-		CurvatureMapSettings.CurvatureType = (int32)FMeshCurvatureMapEvaluator::ECurvatureType::MinPrincipal;
-		break;
-	}
-	switch (CurvatureSettings->ColorMapping)
-	{
-	default:
-	case EBakedCurvatureColorMode::Grayscale:
-		CurvatureMapSettings.ColorMode = (int32)FMeshCurvatureMapEvaluator::EColorMode::BlackGrayWhite;
-		break;
-	case EBakedCurvatureColorMode::RedBlue:
-		CurvatureMapSettings.ColorMode = (int32)FMeshCurvatureMapEvaluator::EColorMode::RedBlue;
-		break;
-	case EBakedCurvatureColorMode::RedGreenBlue:
-		CurvatureMapSettings.ColorMode = (int32)FMeshCurvatureMapEvaluator::EColorMode::RedGreenBlue;
-		break;
-	}
-	switch (CurvatureSettings->Clamping)
-	{
-	default:
-	case EBakedCurvatureClampMode::None:
-		CurvatureMapSettings.ClampMode = (int32)FMeshCurvatureMapEvaluator::EClampMode::FullRange;
-		break;
-	case EBakedCurvatureClampMode::OnlyPositive:
-		CurvatureMapSettings.ClampMode = (int32)FMeshCurvatureMapEvaluator::EClampMode::Positive;
-		break;
-	case EBakedCurvatureClampMode::OnlyNegative:
-		CurvatureMapSettings.ClampMode = (int32)FMeshCurvatureMapEvaluator::EClampMode::Negative;
-		break;
-	}
-
-	if (!(CachedCurvatureMapSettings == CurvatureMapSettings))
-	{
-		CachedCurvatureMapSettings = CurvatureMapSettings;
-		ResultState = EBakeOpState::Evaluate;
-	}
-	return ResultState;
-}
-
-EBakeOpState UBakeMeshAttributeVertexTool::UpdateResult_MeshProperty()
-{
-	// No settings to configure, always valid to evaluate.
-	return EBakeOpState::Evaluate;
-}
-
-EBakeOpState UBakeMeshAttributeVertexTool::UpdateResult_Texture2DImage()
-{
-	EBakeOpState ResultState = EBakeOpState::Clean;
-
-	FTexture2DImageSettings NewSettings;
-	NewSettings.UVLayer = 0;
-
-	const FDynamicMeshUVOverlay* UVOverlay = DetailMesh->Attributes()->GetUVLayer(NewSettings.UVLayer);
-	if (UVOverlay == nullptr)
-	{
-		GetToolManager()->DisplayMessage(LOCTEXT("InvalidUVWarning", "The Source Mesh does not have the selected UV layer"), EToolMessageLevel::UserWarning);
-		return EBakeOpState::Invalid;
-	}
-	
-	if (TextureSettings->SourceTexture == nullptr)
-	{
-		GetToolManager()->DisplayMessage(LOCTEXT("InvalidTextureWarning", "The Source Texture is not valid"), EToolMessageLevel::UserWarning);
-		return EBakeOpState::Invalid;
-	}
-
-
-	{
-		CachedTextureImage = MakeShared<UE::Geometry::TImageBuilder<FVector4f>, ESPMode::ThreadSafe>();
-		if (!UE::AssetUtils::ReadTexture(TextureSettings->SourceTexture, *CachedTextureImage, bPreferPlatformData))
-		{
-			GetToolManager()->DisplayMessage(LOCTEXT("CannotReadTextureWarning", "Cannot read from the source texture"), EToolMessageLevel::UserWarning);
-			return EBakeOpState::Invalid;
-		}
-	}
-
-	if (!(CachedTexture2DImageSettings == NewSettings))
-	{
-		CachedTexture2DImageSettings = NewSettings;
-		ResultState = EBakeOpState::Evaluate;
-	}
-	return ResultState;
-}
-
-EBakeOpState UBakeMeshAttributeVertexTool::UpdateResult_MultiTexture()
-{
-	EBakeOpState ResultState = EBakeOpState::Clean;
-
-	FTexture2DImageSettings NewSettings;
-	NewSettings.UVLayer = MultiTextureSettings->UVLayer;
-
-	const FDynamicMeshUVOverlay* UVOverlay = DetailMesh->Attributes()->GetUVLayer(NewSettings.UVLayer);
-	if (UVOverlay == nullptr)
-	{
-		GetToolManager()->DisplayMessage(LOCTEXT("InvalidUVWarning", "The Source Mesh does not have the selected UV layer"), EToolMessageLevel::UserWarning);
-		return EBakeOpState::Invalid;
-	}
-
-	const int NumMaterialIDs = MultiTextureSettings->MaterialIDSourceTextures.Num();
-	CachedMultiTextures.Reset();
-	CachedMultiTextures.SetNum(NumMaterialIDs);
-	int NumValidTextures = 0;
-	for ( int MaterialID = 0; MaterialID < NumMaterialIDs; ++MaterialID)
-	{
-		if (UTexture2D* Texture = MultiTextureSettings->MaterialIDSourceTextures[MaterialID])
-		{
-			CachedMultiTextures[MaterialID] = MakeShared<TImageBuilder<FVector4f>, ESPMode::ThreadSafe>();
-			if (!UE::AssetUtils::ReadTexture(Texture, *CachedMultiTextures[MaterialID], bPreferPlatformData))
-			{
-				GetToolManager()->DisplayMessage(LOCTEXT("CannotReadTextureWarning", "Cannot read from the source texture"), EToolMessageLevel::UserWarning);
-				return EBakeOpState::Invalid;
-			}
-			++NumValidTextures;
-		}
-	}
-	if (NumValidTextures == 0)
-	{
-		GetToolManager()->DisplayMessage(LOCTEXT("InvalidTextureWarning", "The Source Texture is not valid"), EToolMessageLevel::UserWarning);
-		return EBakeOpState::Invalid;
-	}
-
-	if (!(CachedTexture2DImageSettings == NewSettings))
-	{
-		CachedTexture2DImageSettings = NewSettings;
-		ResultState = EBakeOpState::Evaluate;
-	}
-	return ResultState;
-}
-
-
 void UBakeMeshAttributeVertexTool::GatherAnalytics(FBakeAnalytics::FMeshSettings& Data)
 {
 	if (!FEngineAnalytics::IsAvailable())
@@ -901,8 +731,8 @@ void UBakeMeshAttributeVertexTool::GatherAnalytics(FBakeAnalytics::FMeshSettings
 		return;
 	}
 	
-	Data.NumTargetMeshVerts = BaseMesh.VertexCount();
-	Data.NumTargetMeshTris = BaseMesh.TriangleCount();
+	Data.NumTargetMeshVerts = TargetMesh.VertexCount();
+	Data.NumTargetMeshTris = TargetMesh.TriangleCount();
 	Data.NumDetailMesh = 1;
 	Data.NumDetailMeshTris = DetailMesh->TriangleCount();
 }
diff --git a/Engine/Plugins/Experimental/MeshModelingToolsetExp/Source/MeshModelingToolsExp/Private/BakeMultiMeshAttributeMapsTool.cpp b/Engine/Plugins/Experimental/MeshModelingToolsetExp/Source/MeshModelingToolsExp/Private/BakeMultiMeshAttributeMapsTool.cpp
--- a/Engine/Plugins/Experimental/MeshModelingToolsetExp/Source/MeshModelingToolsExp/Private/BakeMultiMeshAttributeMapsTool.cpp
+++ b/Engine/Plugins/Experimental/MeshModelingToolsetExp/Source/MeshModelingToolsExp/Private/BakeMultiMeshAttributeMapsTool.cpp
@@ -318,7 +318,7 @@ public:
 	UE::Geometry::FDynamicMesh3* BaseMesh;
 	TSharedPtr<UE::Geometry::FMeshTangentsd, ESPMode::ThreadSafe> BaseMeshTangents;
 	TUniquePtr<UE::Geometry::FMeshMapBaker> Baker;
-	UBakeMultiMeshAttributeMapsTool::FBakeCacheSettings BakeCacheSettings;
+	UBakeMultiMeshAttributeMapsTool::FBakeSettings BakeSettings;
 
 	// Detail bake data
 	TArray<TSharedPtr<UE::Geometry::TImageBuilder<FVector4f>>> CachedColorImages;
@@ -332,10 +332,10 @@ public:
 			return Progress && Progress->Cancelled();
 		};
 		Baker->SetTargetMesh(BaseMesh);
-		Baker->SetTargetMeshUVLayer(BakeCacheSettings.TargetUVLayer);
-		Baker->SetDimensions(BakeCacheSettings.Dimensions);
-		Baker->SetProjectionDistance(BakeCacheSettings.ProjectionDistance);
-		Baker->SetSamplesPerPixel(BakeCacheSettings.SamplesPerPixel);
+		Baker->SetTargetMeshUVLayer(BakeSettings.TargetUVLayer);
+		Baker->SetDimensions(BakeSettings.Dimensions);
+		Baker->SetProjectionDistance(BakeSettings.ProjectionDistance);
+		Baker->SetSamplesPerPixel(BakeSettings.SamplesPerPixel);
 		Baker->SetTargetMeshTangents(BaseMeshTangents);
 		
 		FMeshBakerMeshSceneSampler DetailSampler(DetailMeshScene);
@@ -343,7 +343,7 @@ public:
 
 		for (const EBakeMapType MapType : ALL_BAKE_MAP_TYPES)
 		{
-			switch (BakeCacheSettings.BakeMapTypes & MapType)
+			switch (BakeSettings.BakeMapTypes & MapType)
 			{
 			case EBakeMapType::TangentSpaceNormal:
 			{
@@ -383,13 +383,13 @@ void UBakeMultiMeshAttributeMapsTool::Setup()
 	const UE::Geometry::FTransform3d BaseToWorld = UE::ToolTarget::GetLocalToWorldTransform(Targets[0]);
 	PreviewMesh->ProcessMesh([this, BaseToWorld](const FDynamicMesh3& Mesh)
 	{
-		BaseMesh.Copy(Mesh);
-		BaseMeshTangents = MakeShared<FMeshTangentsd, ESPMode::ThreadSafe>(&BaseMesh);
-		BaseMeshTangents->CopyTriVertexTangents(Mesh);
+		TargetMesh.Copy(Mesh);
+		TargetMeshTangents = MakeShared<FMeshTangentsd, ESPMode::ThreadSafe>(&TargetMesh);
+		TargetMeshTangents->CopyTriVertexTangents(Mesh);
 		
 		// FMeshSceneAdapter operates in world space, so ensure our mesh transformed to world.
-		MeshTransforms::ApplyTransform(BaseMesh, BaseToWorld);
-		BaseSpatial.SetMesh(&BaseMesh, true);
+		MeshTransforms::ApplyTransform(TargetMesh, BaseToWorld);
+		TargetSpatial.SetMesh(&TargetMesh, true);
 	});
 
 	// Initialize detail sampler
@@ -414,7 +414,7 @@ void UBakeMultiMeshAttributeMapsTool::Setup()
 	InputMeshProps->RestoreProperties(this);
 	AddToolPropertySource(InputMeshProps);
 	InputMeshProps->TargetStaticMesh = GetStaticMeshTarget(Target);
-	UpdateUVLayerNames(InputMeshProps->TargetUVLayer, InputMeshProps->TargetUVLayerNamesList, BaseMesh);
+	UpdateUVLayerNames(InputMeshProps->TargetUVLayer, InputMeshProps->TargetUVLayerNamesList, TargetMesh);
 	InputMeshProps->WatchProperty(InputMeshProps->TargetUVLayer, [this](FString) { OpState |= EBakeOpState::Evaluate; });
 	InputMeshProps->WatchProperty(InputMeshProps->ProjectionDistance, [this](float) { OpState |= EBakeOpState::Evaluate; });
 	
@@ -497,16 +497,16 @@ TUniquePtr<UE::Geometry::TGenericDataOperator<FMeshMapBaker>> UBakeMultiMeshAttr
 {
 	TUniquePtr<FMultiMeshMapBakerOp> Op = MakeUnique<FMultiMeshMapBakerOp>();
 	Op->DetailMeshScene = &DetailMeshScene;
-	Op->BaseMesh = &BaseMesh;
-	Op->BakeCacheSettings = CachedBakeCacheSettings;
+	Op->BaseMesh = &TargetMesh;
+	Op->BakeSettings = CachedBakeSettings;
 
 	constexpr EBakeMapType RequiresTangents = EBakeMapType::TangentSpaceNormal | EBakeMapType::BentNormal;
-	if (static_cast<bool>(CachedBakeCacheSettings.BakeMapTypes & RequiresTangents))
+	if (static_cast<bool>(CachedBakeSettings.BakeMapTypes & RequiresTangents))
 	{
-		Op->BaseMeshTangents = BaseMeshTangents;
+		Op->BaseMeshTangents = TargetMeshTangents;
 	}
 
-	if (static_cast<bool>(CachedBakeCacheSettings.BakeMapTypes & EBakeMapType::Texture))
+	if (static_cast<bool>(CachedBakeSettings.BakeMapTypes & EBakeMapType::Texture))
 	{
 		Op->CachedColorImages = CachedColorImages;
 		Op->CachedMeshToColorImageMap = CachedMeshToColorImagesMap;
@@ -559,23 +559,23 @@ void UBakeMultiMeshAttributeMapsTool::UpdateResult()
 	const int32 ImageSize = static_cast<int32>(Settings->Resolution);
 	const FImageDimensions Dimensions(ImageSize, ImageSize);
 
-	FBakeCacheSettings BakeCacheSettings;
-	BakeCacheSettings.Dimensions = Dimensions;
-	BakeCacheSettings.BitDepth = Settings->BitDepth;
-	BakeCacheSettings.TargetUVLayer = InputMeshProps->TargetUVLayerNamesList.IndexOfByKey(InputMeshProps->TargetUVLayer);
-	BakeCacheSettings.ProjectionDistance = InputMeshProps->ProjectionDistance;
-	BakeCacheSettings.bProjectionInWorldSpace = true;  // Always world space
-	BakeCacheSettings.SamplesPerPixel = static_cast<int32>(Settings->SamplesPerPixel);
+	FBakeSettings BakeSettings;
+	BakeSettings.Dimensions = Dimensions;
+	BakeSettings.BitDepth = Settings->BitDepth;
+	BakeSettings.TargetUVLayer = InputMeshProps->TargetUVLayerNamesList.IndexOfByKey(InputMeshProps->TargetUVLayer);
+	BakeSettings.ProjectionDistance = InputMeshProps->ProjectionDistance;
+	BakeSettings.bProjectionInWorldSpace = true;  // Always world space
+	BakeSettings.SamplesPerPixel = static_cast<int32>(Settings->SamplesPerPixel);
 
 	// Record the original map types and process the raw bitfield which may add
 	// additional targets.
-	BakeCacheSettings.SourceBakeMapTypes = static_cast<EBakeMapType>(Settings->MapTypes);
-	BakeCacheSettings.BakeMapTypes = GetMapTypes(Settings->MapTypes);
+	BakeSettings.SourceBakeMapTypes = static_cast<EBakeMapType>(Settings->MapTypes);
+	BakeSettings.BakeMapTypes = GetMapTypes(Settings->MapTypes);
 
 	// update bake cache settings
-	if (!(CachedBakeCacheSettings == BakeCacheSettings))
+	if (!(CachedBakeSettings == BakeSettings))
 	{
-		CachedBakeCacheSettings = BakeCacheSettings;
+		CachedBakeSettings = BakeSettings;
 		CachedDetailSettings = FBakeMultiMeshDetailSettings();
 	}
 
@@ -585,6 +585,11 @@ void UBakeMultiMeshAttributeMapsTool::UpdateResult()
 	// Update map type settings
 	OpState |= UpdateResult_DetailMeshes();
 
+	if ((bool)(CachedBakeSettings.BakeMapTypes & EBakeMapType::TangentSpaceNormal))
+	{
+		OpState |= UpdateResult_Normal(CachedBakeSettings.Dimensions);
+	}
+
 	// Early exit if op input parameters are invalid.
 	if (static_cast<bool>(OpState & EBakeOpState::Invalid))
 	{
@@ -613,7 +618,7 @@ EBakeOpState UBakeMultiMeshAttributeMapsTool::UpdateResult_DetailMeshes()
 		ActorToDataMap.Emplace(ActorComponent, Idx);
 		
 		// Color map data
-		if (static_cast<bool>(CachedBakeCacheSettings.BakeMapTypes & EBakeMapType::Texture))
+		if (static_cast<bool>(CachedBakeSettings.BakeMapTypes & EBakeMapType::Texture))
 		{
 			UTexture2D* ColorMapSourceTexture = InputMeshProps->SourceMeshes[Idx].SourceTexture;
 			const int ColorMapUVLayer = InputMeshProps->SourceMeshes[Idx].SourceTextureUVLayer;
@@ -642,7 +647,7 @@ EBakeOpState UBakeMultiMeshAttributeMapsTool::UpdateResult_DetailMeshes()
 		{
 			if (const int* DataIndex = ActorToDataMap.Find(ChildMesh->SourceComponent))
 			{
-				if (static_cast<bool>(CachedBakeCacheSettings.BakeMapTypes & EBakeMapType::Texture))
+				if (static_cast<bool>(CachedBakeSettings.BakeMapTypes & EBakeMapType::Texture))
 				{
 					CachedMeshToColorImagesMap.Emplace(
 					ChildMesh->MeshSpatial,
@@ -694,7 +699,7 @@ void UBakeMultiMeshAttributeMapsTool::GatherAnalytics(FBakeAnalytics::FMeshSetti
 {
 	if (FEngineAnalytics::IsAvailable())
 	{
-		Data.NumTargetMeshTris = BaseMesh.TriangleCount();
+		Data.NumTargetMeshTris = TargetMesh.TriangleCount();
 		Data.NumDetailMesh = 0;
 		Data.NumDetailMeshTris = 0;
 		DetailMeshScene.ProcessActorChildMeshes([&Data](const FActorAdapter* ActorAdapter, const FActorChildMesh* ChildMesh)
diff --git a/Engine/Plugins/Experimental/MeshModelingToolsetExp/Source/MeshModelingToolsExp/Public/BakeMeshAttributeMapsTool.h b/Engine/Plugins/Experimental/MeshModelingToolsetExp/Source/MeshModelingToolsExp/Public/BakeMeshAttributeMapsTool.h
--- a/Engine/Plugins/Experimental/MeshModelingToolsetExp/Source/MeshModelingToolsExp/Public/BakeMeshAttributeMapsTool.h
+++ b/Engine/Plugins/Experimental/MeshModelingToolsetExp/Source/MeshModelingToolsExp/Public/BakeMeshAttributeMapsTool.h
@@ -119,17 +119,6 @@ protected:
 	UPROPERTY()
 	TObjectPtr<UBakeMeshAttributeMapsToolProperties> Settings;
 
-	UPROPERTY()
-	TObjectPtr<UBakedOcclusionMapToolProperties> OcclusionMapProps;
-
-	UPROPERTY()
-	TObjectPtr<UBakedCurvatureMapToolProperties> CurvatureMapProps;
-
-	UPROPERTY()
-	TObjectPtr<UBakedTexture2DImageProperties> Texture2DProps;
-
-	UPROPERTY()
-	TObjectPtr<UBakedMultiTexture2DImageProperties> MultiTextureProps;
 
 	// Begin UBakeMeshAttributeMapsToolBase interface
 	virtual void UpdateResult() override;
@@ -155,24 +144,5 @@ protected:
 	FDetailMeshSettings CachedDetailMeshSettings;
 	TSharedPtr<UE::Geometry::TImageBuilder<FVector4f>, ESPMode::ThreadSafe> CachedDetailNormalMap;
 	EBakeOpState UpdateResult_DetailNormalMap();
-
-	FNormalMapSettings CachedNormalMapSettings;
-	EBakeOpState UpdateResult_Normal();
-
-	FOcclusionMapSettings CachedOcclusionMapSettings;
-	EBakeOpState UpdateResult_Occlusion();
-
-	FCurvatureMapSettings CachedCurvatureMapSettings;
-	EBakeOpState UpdateResult_Curvature();
-
-	FMeshPropertyMapSettings CachedMeshPropertyMapSettings;
-	EBakeOpState UpdateResult_MeshProperty();
-
-	TSharedPtr<UE::Geometry::TImageBuilder<FVector4f>, ESPMode::ThreadSafe> CachedTextureImage;
-	FTexture2DImageSettings CachedTexture2DImageSettings;
-	EBakeOpState UpdateResult_Texture2DImage();
-
-	TArray<TSharedPtr<UE::Geometry::TImageBuilder<FVector4f>, ESPMode::ThreadSafe>> CachedMultiTextures;
-	EBakeOpState UpdateResult_MultiTexture();
 };
 
diff --git a/Engine/Plugins/Experimental/MeshModelingToolsetExp/Source/MeshModelingToolsExp/Public/BakeMeshAttributeMapsToolBase.h b/Engine/Plugins/Experimental/MeshModelingToolsetExp/Source/MeshModelingToolsExp/Public/BakeMeshAttributeMapsToolBase.h
--- a/Engine/Plugins/Experimental/MeshModelingToolsetExp/Source/MeshModelingToolsExp/Public/BakeMeshAttributeMapsToolBase.h
+++ b/Engine/Plugins/Experimental/MeshModelingToolsetExp/Source/MeshModelingToolsExp/Public/BakeMeshAttributeMapsToolBase.h
@@ -95,7 +95,7 @@ protected:
 	// Tool property sets
 	//
 	UPROPERTY()
-	TObjectPtr<UBakedOcclusionMapVisualizationProperties> VisualizationProps;
+	TObjectPtr<UBakeVisualizationProperties> VisualizationProps;
 
 	//
 	// Preview mesh and materials
@@ -144,11 +144,7 @@ protected:
 	//
 	// Bake parameters
 	//
-	UE::Geometry::FDynamicMesh3 BaseMesh;
-	TSharedPtr<UE::Geometry::TMeshTangents<double>, ESPMode::ThreadSafe> BaseMeshTangents;
-	UE::Geometry::FDynamicMeshAABBTree3 BaseSpatial;
-
-	struct FBakeCacheSettings
+	struct FBakeSettings
 	{
 		EBakeMapType SourceBakeMapTypes = EBakeMapType::None;
 		EBakeMapType BakeMapTypes = EBakeMapType::None;
@@ -160,7 +156,7 @@ protected:
 		int32 SamplesPerPixel = 1;
 		bool bProjectionInWorldSpace = false;
 
-		bool operator==(const FBakeCacheSettings& Other) const
+		bool operator==(const FBakeSettings& Other) const
 		{
 			return BakeMapTypes == Other.BakeMapTypes && Dimensions == Other.Dimensions &&
 				TargetUVLayer == Other.TargetUVLayer && DetailTimestamp == Other.DetailTimestamp &&
@@ -169,7 +165,7 @@ protected:
 				bProjectionInWorldSpace == Other.bProjectionInWorldSpace;
 		}
 	};
-	FBakeCacheSettings CachedBakeCacheSettings;
+	FBakeSettings CachedBakeSettings;
 
 	/**
 	 * To be invoked by client when bake map types change.
@@ -223,18 +219,6 @@ protected:
 	void UpdatePreviewNames(PropertySet& Properties);
 
 
-	/**
-	 * Updates a tool property set's UVLayerNamesList from the list of UV layers
-	 * on a given mesh. Also updates the UVLayer property if the current UV layer
-	 * is no longer available.
-	 *
-	 * @param UVLayer Selected UV Layer.
-	 * @param UVLayerNamesList List of available UV layers.
-	 * @param Mesh the mesh to query
-	 */
-	static void UpdateUVLayerNames(FString& UVLayer, TArray<FString>& UVLayerNamesList, const FDynamicMesh3& Mesh);
-
-
 	//
 	// Analytics
 	//
@@ -254,7 +238,7 @@ protected:
 		};
 		FMeshSettings MeshSettings;
 
-		FBakeCacheSettings BakeSettings;
+		FBakeSettings BakeSettings;
 		FOcclusionMapSettings OcclusionSettings;
 		FCurvatureMapSettings CurvatureSettings;
 	};
@@ -273,7 +257,7 @@ protected:
 	 * @param Data the output bake analytics struct.
 	 */
 	static void GatherAnalytics(const UE::Geometry::FMeshMapBaker& Result,
-								const FBakeCacheSettings& Settings,
+								const FBakeSettings& Settings,
 								FBakeAnalytics& Data);
 
 	/**
@@ -295,7 +279,6 @@ protected:
 	//
 	// Utilities
 	//
-	const bool bPreferPlatformData = false;
 	
 	/** @return the Texture2D type for a given map type */
 	static UE::Geometry::FTexture2DBuilder::ETextureType GetTextureType(EBakeMapType MapType, EBakeTextureBitDepth MapFormat);
@@ -416,25 +399,6 @@ void UBakeMeshAttributeMapsToolBase::UpdatePreviewNames(PropertySet& Properties)
 }
 
 
-inline void UBakeMeshAttributeMapsToolBase::UpdateUVLayerNames(FString& UVLayer, TArray<FString>& UVLayerNamesList, const FDynamicMesh3& Mesh)
-{
-	UVLayerNamesList.Reset();
-	int32 FoundIndex = -1;
-	for (int32 k = 0; k < Mesh.Attributes()->NumUVLayers(); ++k)
-	{
-		UVLayerNamesList.Add(FString::Printf(TEXT("UV %d"), k));
-		if (UVLayer == UVLayerNamesList.Last())
-		{
-			FoundIndex = k;
-		}
-	}
-	if (FoundIndex == -1)
-	{
-		UVLayer = UVLayerNamesList[0];
-	}
-}
-
-
 template <typename PropertySet>
 bool UBakeMeshAttributeMapsToolBase::IsRequestedMapType(PropertySet& Properties, EBakeMapType MapType)
 {
diff --git a/Engine/Plugins/Experimental/MeshModelingToolsetExp/Source/MeshModelingToolsExp/Public/BakeMeshAttributeTool.h b/Engine/Plugins/Experimental/MeshModelingToolsetExp/Source/MeshModelingToolsExp/Public/BakeMeshAttributeTool.h
--- a/Engine/Plugins/Experimental/MeshModelingToolsetExp/Source/MeshModelingToolsExp/Public/BakeMeshAttributeTool.h
+++ b/Engine/Plugins/Experimental/MeshModelingToolsetExp/Source/MeshModelingToolsExp/Public/BakeMeshAttributeTool.h
@@ -93,6 +93,22 @@ public:
 	// End UInteractiveTool interface
 
 protected:
+	//
+	// Bake tool property sets
+	//
+	UPROPERTY()
+	TObjectPtr<UBakeOcclusionMapToolProperties> OcclusionSettings;
+
+	UPROPERTY()
+	TObjectPtr<UBakeCurvatureMapToolProperties> CurvatureSettings;
+
+	UPROPERTY()
+	TObjectPtr<UBakeTexture2DProperties> TextureSettings;
+
+	UPROPERTY()
+	TObjectPtr<UBakeMultiTexture2DProperties> MultiTextureSettings;
+	
+	
 	//
 	// Preview materials
 	//
@@ -110,8 +126,40 @@ protected:
 	//
 	EBakeOpState OpState = EBakeOpState::Evaluate;
 
+	UE::Geometry::FDynamicMesh3 TargetMesh;
+	TSharedPtr<UE::Geometry::TMeshTangents<double>, ESPMode::ThreadSafe> TargetMeshTangents;
+	UE::Geometry::FDynamicMeshAABBTree3 TargetSpatial;
+
+	//
+	// Bake results update management
+	//
+	const bool bPreferPlatformData = false;
+	
+	EBakeOpState UpdateResult_Normal(const FImageDimensions& Dimensions);
+	FNormalMapSettings CachedNormalMapSettings;
+
+	EBakeOpState UpdateResult_Occlusion(const FImageDimensions& Dimensions);
+	FOcclusionMapSettings CachedOcclusionMapSettings;
+
+	EBakeOpState UpdateResult_Curvature(const FImageDimensions& Dimensions);
+	FCurvatureMapSettings CachedCurvatureMapSettings;
+
+	EBakeOpState UpdateResult_MeshProperty(const FImageDimensions& Dimensions);
+	FMeshPropertyMapSettings CachedMeshPropertyMapSettings;
+
+	EBakeOpState UpdateResult_Texture2DImage(const FImageDimensions& Dimensions, const FDynamicMesh3* DetailMesh);
+	FTexture2DSettings CachedTexture2DSettings;
+	TSharedPtr<UE::Geometry::TImageBuilder<FVector4f>, ESPMode::ThreadSafe> CachedTextureImage;
+
+	EBakeOpState UpdateResult_MultiTexture(const FImageDimensions& Dimensions, const FDynamicMesh3* DetailMesh);
+	FTexture2DSettings CachedMultiTexture2DSettings;
+	TArray<TSharedPtr<UE::Geometry::TImageBuilder<FVector4f>, ESPMode::ThreadSafe>> CachedMultiTextures;
 
 protected:
+	//
+	// Utilities
+	//
+	
 	/** @return StaticMesh from a tool target */ 
 	static UStaticMesh* GetStaticMeshTarget(UToolTarget* Target)
 	{
@@ -163,6 +211,17 @@ protected:
 		UToolTarget* Target,
 		TArray<TObjectPtr<UTexture2D>>& AllSourceTextures,
 		TArray<TObjectPtr<UTexture2D>>& MaterialIDTextures);
+
+	/**
+	 * Updates a tool property set's UVLayerNamesList from the list of UV layers
+	 * on a given mesh. Also updates the UVLayer property if the current UV layer
+	 * is no longer available.
+	 *
+	 * @param UVLayer Selected UV Layer.
+	 * @param UVLayerNamesList List of available UV layers.
+	 * @param Mesh the mesh to query
+	 */
+	static void UpdateUVLayerNames(FString& UVLayer, TArray<FString>& UVLayerNamesList, const FDynamicMesh3& Mesh);
 };
 
 
diff --git a/Engine/Plugins/Experimental/MeshModelingToolsetExp/Source/MeshModelingToolsExp/Public/BakeMeshAttributeToolCommon.h b/Engine/Plugins/Experimental/MeshModelingToolsetExp/Source/MeshModelingToolsExp/Public/BakeMeshAttributeToolCommon.h
--- a/Engine/Plugins/Experimental/MeshModelingToolsetExp/Source/MeshModelingToolsExp/Public/BakeMeshAttributeToolCommon.h
+++ b/Engine/Plugins/Experimental/MeshModelingToolsetExp/Source/MeshModelingToolsExp/Public/BakeMeshAttributeToolCommon.h
@@ -114,7 +114,7 @@ public:
 
 
 UCLASS()
-class MESHMODELINGTOOLSEXP_API UBakedNormalMapToolProperties : public UInteractiveToolPropertySet
+class MESHMODELINGTOOLSEXP_API UBakeNormalMapToolProperties : public UInteractiveToolPropertySet
 {
 	GENERATED_BODY()
 public:
@@ -123,7 +123,7 @@ public:
 
 
 UCLASS()
-class MESHMODELINGTOOLSEXP_API UBakedOcclusionMapToolProperties : public UInteractiveToolPropertySet
+class MESHMODELINGTOOLSEXP_API UBakeOcclusionMapToolProperties : public UInteractiveToolPropertySet
 {
 	GENERATED_BODY()
 public:
@@ -145,24 +145,8 @@ public:
 };
 
 
-UCLASS()
-class MESHMODELINGTOOLSEXP_API UBakedOcclusionMapVisualizationProperties : public UInteractiveToolPropertySet
-{
-	GENERATED_BODY()
-public:
-	/** Adjust the brightness of the preview material; does not affect results stored in textures */
-	UPROPERTY(EditAnywhere, Category = Preview, meta = (DisplayName = "Brightness", UIMin = "0.0", UIMax = "1.0"))
-	float Brightness = 1.0f;
-
-	/** Ambient Occlusion multiplier in the viewport; does not affect results stored in textures */
-	UPROPERTY(EditAnywhere, Category = Preview, meta = (DisplayName = "AO Multiplier", UIMin = "0.0", UIMax = "1.0",
-		ClampMin = "0.0", ClampMax = "1.0"))
-	float AOMultiplier = 1.0f;
-};
-
-
 UENUM()
-enum class EBakedCurvatureTypeMode
+enum class EBakeCurvatureTypeMode
 {
 	/** Average of the minimum and maximum principal curvatures */
 	MeanAverage,
@@ -176,7 +160,7 @@ enum class EBakedCurvatureTypeMode
 
 
 UENUM()
-enum class EBakedCurvatureColorMode
+enum class EBakeCurvatureColorMode
 {
 	/** Map curvature values to grayscale such that black is negative, grey is zero, and white is positive */
 	Grayscale,
@@ -188,7 +172,7 @@ enum class EBakedCurvatureColorMode
 
 
 UENUM()
-enum class EBakedCurvatureClampMode
+enum class EBakeCurvatureClampMode
 {
 	/** Include both negative and positive curvatures */
 	None,
@@ -200,17 +184,17 @@ enum class EBakedCurvatureClampMode
 
 
 UCLASS()
-class MESHMODELINGTOOLSEXP_API UBakedCurvatureMapToolProperties : public UInteractiveToolPropertySet
+class MESHMODELINGTOOLSEXP_API UBakeCurvatureMapToolProperties : public UInteractiveToolPropertySet
 {
 	GENERATED_BODY()
 public:
 	/** Type of curvature */
 	UPROPERTY(EditAnywhere, Category = CurvatureOutput)
-	EBakedCurvatureTypeMode CurvatureType = EBakedCurvatureTypeMode::MeanAverage;
+	EBakeCurvatureTypeMode CurvatureType = EBakeCurvatureTypeMode::MeanAverage;
 
 	/** How to map calculated curvature values to colors */
 	UPROPERTY(EditAnywhere, Category = CurvatureOutput)
-	EBakedCurvatureColorMode ColorMapping = EBakedCurvatureColorMode::Grayscale;
+	EBakeCurvatureColorMode ColorMapping = EBakeCurvatureColorMode::Grayscale;
 
 	/** Multiplier for how the curvature values fill the available range in the selected Color Mapping; a larger value means that higher curvature is required to achieve the maximum color value. */
 	UPROPERTY(EditAnywhere, Category = CurvatureOutput, meta = (UIMin = "0.1", UIMax = "2.0", ClampMin = "0.001", ClampMax = "100.0"))
@@ -222,12 +206,12 @@ public:
 
 	/** Clamping applied to curvature values before color mapping */
 	UPROPERTY(EditAnywhere, Category = CurvatureOutput)
-	EBakedCurvatureClampMode Clamping = EBakedCurvatureClampMode::None;
+	EBakeCurvatureClampMode Clamping = EBakeCurvatureClampMode::None;
 };
 
 
 UCLASS()
-class MESHMODELINGTOOLSEXP_API UBakedTexture2DImageProperties : public UInteractiveToolPropertySet
+class MESHMODELINGTOOLSEXP_API UBakeTexture2DProperties : public UInteractiveToolPropertySet
 {
 	GENERATED_BODY()
 public:
@@ -242,7 +226,7 @@ public:
 
 
 UCLASS()
-class MESHMODELINGTOOLSEXP_API UBakedMultiTexture2DImageProperties : public UInteractiveToolPropertySet
+class MESHMODELINGTOOLSEXP_API UBakeMultiTexture2DProperties : public UInteractiveToolPropertySet
 {
 	GENERATED_BODY()
 public:
@@ -264,6 +248,22 @@ public:
 };
 
 
+UCLASS()
+class MESHMODELINGTOOLSEXP_API UBakeVisualizationProperties : public UInteractiveToolPropertySet
+{
+	GENERATED_BODY()
+public:
+	/** Adjust the brightness of the preview material; does not affect results stored in textures */
+	UPROPERTY(EditAnywhere, Category = Preview, meta = (DisplayName = "Brightness", UIMin = "0.0", UIMax = "1.0"))
+	float Brightness = 1.0f;
+
+	/** Ambient Occlusion multiplier in the viewport; does not affect results stored in textures */
+	UPROPERTY(EditAnywhere, Category = Preview, meta = (DisplayName = "AO Multiplier", UIMin = "0.0", UIMax = "1.0",
+		ClampMin = "0.0", ClampMax = "1.0"))
+	float AOMultiplier = 1.0f;
+};
+
+
 /**
  * Bake tool property settings structs
  */
@@ -331,12 +331,12 @@ struct FMeshPropertyMapSettings
 	}
 };
 
-struct FTexture2DImageSettings
+struct FTexture2DSettings
 {
 	FImageDimensions Dimensions;
 	int32 UVLayer = 0;
 
-	bool operator==(const FTexture2DImageSettings& Other) const
+	bool operator==(const FTexture2DSettings& Other) const
 	{
 		return Dimensions == Other.Dimensions && UVLayer == Other.UVLayer;
 	}
diff --git a/Engine/Plugins/Experimental/MeshModelingToolsetExp/Source/MeshModelingToolsExp/Public/BakeMeshAttributeVertexTool.h b/Engine/Plugins/Experimental/MeshModelingToolsetExp/Source/MeshModelingToolsExp/Public/BakeMeshAttributeVertexTool.h
--- a/Engine/Plugins/Experimental/MeshModelingToolsetExp/Source/MeshModelingToolsExp/Public/BakeMeshAttributeVertexTool.h
+++ b/Engine/Plugins/Experimental/MeshModelingToolsetExp/Source/MeshModelingToolsExp/Public/BakeMeshAttributeVertexTool.h
@@ -140,18 +140,6 @@ protected:
 	UPROPERTY()
 	TObjectPtr<UBakeMeshAttributeVertexToolProperties> Settings;
 
-	UPROPERTY()
-	TObjectPtr<UBakedOcclusionMapToolProperties> OcclusionSettings;
-
-	UPROPERTY()
-	TObjectPtr<UBakedCurvatureMapToolProperties> CurvatureSettings;
-
-	UPROPERTY()
-	TObjectPtr<UBakedTexture2DImageProperties> TextureSettings;
-
-	UPROPERTY()
-	TObjectPtr<UBakedMultiTexture2DImageProperties> MultiTextureSettings;
-
 protected:
 	friend class FMeshVertexBakerOp;
 	
@@ -167,15 +155,12 @@ protected:
 	TUniquePtr<TGenericDataBackgroundCompute<UE::Geometry::FMeshVertexBaker>> Compute = nullptr;
 	void OnResultUpdated(const TUniquePtr<UE::Geometry::FMeshVertexBaker>& NewResult);
 
-	TSharedPtr<UE::Geometry::TMeshTangents<double>, ESPMode::ThreadSafe> BaseMeshTangents;
-	UE::Geometry::FDynamicMesh3 BaseMesh;
-	UE::Geometry::FDynamicMeshAABBTree3 BaseSpatial;
-
 	TSharedPtr<UE::Geometry::FDynamicMesh3, ESPMode::ThreadSafe> DetailMesh;
 	TSharedPtr<UE::Geometry::FDynamicMeshAABBTree3, ESPMode::ThreadSafe> DetailSpatial;
 	int32 DetailMeshTimestamp = 0;
 	void UpdateDetailMesh();
 
+	int32 NumColorElements = 0;
 	bool bColorTopologyValid = false;
 	bool bIsBakeToSelf = false;
 	void UpdateOnModeChange();
@@ -183,8 +168,6 @@ protected:
 	void UpdateColorTopology();
 	void UpdateResult();
 
-	const bool bPreferPlatformData = false;
-
 	struct FBakeSettings
 	{
 		EBakeVertexOutput OutputMode = EBakeVertexOutput::RGBA;
@@ -212,24 +195,6 @@ protected:
 	};
 	FBakeSettings CachedBakeSettings;
 
-	//FNormalMapSettings CachedNormalMapSettings;
-	EBakeOpState UpdateResult_Normal();
-
-	FOcclusionMapSettings CachedOcclusionMapSettings;
-	EBakeOpState UpdateResult_Occlusion();
-
-	FCurvatureMapSettings CachedCurvatureMapSettings;
-	EBakeOpState UpdateResult_Curvature();
-
-	//FMeshPropertyMapSettings CachedMeshPropertyMapSettings;
-	EBakeOpState UpdateResult_MeshProperty();
-
-	TSharedPtr<UE::Geometry::TImageBuilder<FVector4f>, ESPMode::ThreadSafe> CachedTextureImage;
-	FTexture2DImageSettings CachedTexture2DImageSettings;
-	EBakeOpState UpdateResult_Texture2DImage();
-
-	TArray<TSharedPtr<UE::Geometry::TImageBuilder<FVector4f>, ESPMode::ThreadSafe>> CachedMultiTextures;
-	EBakeOpState UpdateResult_MultiTexture();
 
 	//
 	// Analytics
