ModelingTools: Minor Bake tools rename pass.

#rnx

diff --git a/Engine/Plugins/Experimental/MeshModelingToolsetExp/Source/MeshModelingToolsExp/Private/BakeMeshAttributeMapsTool.cpp b/Engine/Plugins/Experimental/MeshModelingToolsetExp/Source/MeshModelingToolsExp/Private/BakeMeshAttributeMapsTool.cpp
--- a/Engine/Plugins/Experimental/MeshModelingToolsetExp/Source/MeshModelingToolsExp/Private/BakeMeshAttributeMapsTool.cpp
+++ b/Engine/Plugins/Experimental/MeshModelingToolsetExp/Source/MeshModelingToolsExp/Private/BakeMeshAttributeMapsTool.cpp
@@ -221,7 +221,7 @@ public:
 			case EBakeMapType::Texture:
 			{
 				TSharedPtr<FMeshResampleImageEvaluator, ESPMode::ThreadSafe> TextureEval = MakeShared<FMeshResampleImageEvaluator, ESPMode::ThreadSafe>();
-				DetailSampler.SetColorMap(DetailMesh.Get(), IMeshBakerDetailSampler::FBakeDetailTexture(TextureImage.Get(), TextureSettings.UVLayer));
+				DetailSampler.SetTextureMap(DetailMesh.Get(), IMeshBakerDetailSampler::FBakeDetailTexture(TextureImage.Get(), TextureSettings.UVLayer));
 				Baker->AddEvaluator(TextureEval);
 				break;
 			}
@@ -269,25 +269,25 @@ void UBakeMeshAttributeMapsTool::Setup()
 	
 	// Setup tool property sets
 
-	MeshProps = NewObject<UBakeInputMeshProperties>(this);
-	MeshProps->RestoreProperties(this);
-	AddToolPropertySource(MeshProps);
-	SetToolPropertySourceEnabled(MeshProps, true);
-	MeshProps->bHasTargetUVLayer = true;
-	MeshProps->bHasSourceNormalMap = true;
-	MeshProps->TargetStaticMesh = GetStaticMeshTarget(Target);
-	MeshProps->TargetSkeletalMesh = GetSkeletalMeshTarget(Target);
-	MeshProps->TargetDynamicMesh = GetDynamicMeshTarget(Target);
-	MeshProps->SourceStaticMesh = !bIsBakeToSelf ? GetStaticMeshTarget(DetailTarget) : nullptr;
-	MeshProps->SourceSkeletalMesh = !bIsBakeToSelf ? GetSkeletalMeshTarget(DetailTarget) : nullptr;
-	MeshProps->SourceDynamicMesh = !bIsBakeToSelf ? GetDynamicMeshTarget(DetailTarget) : nullptr;
-	MeshProps->SourceNormalMap = nullptr;
-	UpdateUVLayerNames(MeshProps->TargetUVLayer, MeshProps->TargetUVLayerNamesList, TargetMesh);
-	MeshProps->WatchProperty(MeshProps->TargetUVLayer, [this](FString) { OpState |= EBakeOpState::Evaluate; });
-	MeshProps->WatchProperty(MeshProps->ProjectionDistance, [this](float) { OpState |= EBakeOpState::Evaluate; });
-	MeshProps->WatchProperty(MeshProps->bProjectionInWorldSpace, [this](bool) { OpState |= EBakeOpState::EvaluateDetailMesh; });
-	MeshProps->WatchProperty(MeshProps->SourceNormalMapUVLayer, [this](FString) { OpState |= EBakeOpState::Evaluate; });
-	MeshProps->WatchProperty(MeshProps->SourceNormalMap, [this](UTexture2D*)
+	InputMeshSettings = NewObject<UBakeInputMeshProperties>(this);
+	InputMeshSettings->RestoreProperties(this);
+	AddToolPropertySource(InputMeshSettings);
+	SetToolPropertySourceEnabled(InputMeshSettings, true);
+	InputMeshSettings->bHasTargetUVLayer = true;
+	InputMeshSettings->bHasSourceNormalMap = true;
+	InputMeshSettings->TargetStaticMesh = GetStaticMeshTarget(Target);
+	InputMeshSettings->TargetSkeletalMesh = GetSkeletalMeshTarget(Target);
+	InputMeshSettings->TargetDynamicMesh = GetDynamicMeshTarget(Target);
+	InputMeshSettings->SourceStaticMesh = !bIsBakeToSelf ? GetStaticMeshTarget(DetailTarget) : nullptr;
+	InputMeshSettings->SourceSkeletalMesh = !bIsBakeToSelf ? GetSkeletalMeshTarget(DetailTarget) : nullptr;
+	InputMeshSettings->SourceDynamicMesh = !bIsBakeToSelf ? GetDynamicMeshTarget(DetailTarget) : nullptr;
+	InputMeshSettings->SourceNormalMap = nullptr;
+	UpdateUVLayerNames(InputMeshSettings->TargetUVLayer, InputMeshSettings->TargetUVLayerNamesList, TargetMesh);
+	InputMeshSettings->WatchProperty(InputMeshSettings->TargetUVLayer, [this](FString) { OpState |= EBakeOpState::Evaluate; });
+	InputMeshSettings->WatchProperty(InputMeshSettings->ProjectionDistance, [this](float) { OpState |= EBakeOpState::Evaluate; });
+	InputMeshSettings->WatchProperty(InputMeshSettings->bProjectionInWorldSpace, [this](bool) { OpState |= EBakeOpState::EvaluateDetailMesh; });
+	InputMeshSettings->WatchProperty(InputMeshSettings->SourceNormalMapUVLayer, [this](FString) { OpState |= EBakeOpState::Evaluate; });
+	InputMeshSettings->WatchProperty(InputMeshSettings->SourceNormalMap, [this](UTexture2D*)
 	{
 		// Only invalidate detail mesh if we need to recompute tangents.
 		if (!DetailMeshTangents)
@@ -448,7 +448,7 @@ void UBakeMeshAttributeMapsTool::Shutdown(EToolShutdownType ShutdownType)
 	Super::Shutdown(ShutdownType);
 	
 	Settings->SaveProperties(this);
-	MeshProps->SaveProperties(this);
+	InputMeshSettings->SaveProperties(this);
 	OcclusionSettings->SaveProperties(this);
 	CurvatureSettings->SaveProperties(this);
 	TextureSettings->SaveProperties(this);
@@ -479,10 +479,10 @@ void UBakeMeshAttributeMapsTool::UpdateDetailMesh()
 {
 	UToolTarget* DetailTarget = Targets[bIsBakeToSelf ? 0 : 1];
 
-	const bool bWantMeshTangents = (MeshProps->SourceNormalMap != nullptr);
+	const bool bWantMeshTangents = (InputMeshSettings->SourceNormalMap != nullptr);
 	DetailMesh = MakeShared<FDynamicMesh3, ESPMode::ThreadSafe>(UE::ToolTarget::GetDynamicMeshCopy(DetailTarget, bWantMeshTangents));
 
-	if (MeshProps->bProjectionInWorldSpace && bIsBakeToSelf == false)
+	if (InputMeshSettings->bProjectionInWorldSpace && bIsBakeToSelf == false)
 	{
 		using FTransform3d = UE::Geometry::FTransform3d;
 		const FTransform3d DetailToWorld = UE::ToolTarget::GetLocalToWorldTransform(DetailTarget);
@@ -505,7 +505,7 @@ void UBakeMeshAttributeMapsTool::UpdateDetailMesh()
 		DetailMeshTangents = nullptr;
 	}
 
-	UpdateUVLayerNames(MeshProps->SourceNormalMapUVLayer, MeshProps->SourceUVLayerNamesList, *DetailMesh);
+	UpdateUVLayerNames(InputMeshSettings->SourceNormalMapUVLayer, InputMeshSettings->SourceUVLayerNamesList, *DetailMesh);
 
 	// Clear detail mesh evaluation flag and mark evaluation.
 	OpState &= ~EBakeOpState::EvaluateDetailMesh;
@@ -536,11 +536,11 @@ void UBakeMeshAttributeMapsTool::UpdateResult()
 	FBakeSettings BakeSettings;
 	BakeSettings.Dimensions = Dimensions;
 	BakeSettings.BitDepth = Settings->BitDepth;
-	BakeSettings.TargetUVLayer = MeshProps->TargetUVLayerNamesList.IndexOfByKey(MeshProps->TargetUVLayer);
+	BakeSettings.TargetUVLayer = InputMeshSettings->TargetUVLayerNamesList.IndexOfByKey(InputMeshSettings->TargetUVLayer);
 	BakeSettings.DetailTimestamp = this->DetailMeshTimestamp;
-	BakeSettings.ProjectionDistance = MeshProps->ProjectionDistance;
+	BakeSettings.ProjectionDistance = InputMeshSettings->ProjectionDistance;
 	BakeSettings.SamplesPerPixel = (int32)Settings->SamplesPerPixel;
-	BakeSettings.bProjectionInWorldSpace = MeshProps->bProjectionInWorldSpace;
+	BakeSettings.bProjectionInWorldSpace = InputMeshSettings->bProjectionInWorldSpace;
 
 	// Record the original map types and process the raw bitfield which may add
 	// additional targets.
@@ -613,7 +613,7 @@ EBakeOpState UBakeMeshAttributeMapsTool::UpdateResult_DetailNormalMap()
 {
 	EBakeOpState ResultState = EBakeOpState::Clean;
 
-	const int DetailUVLayer = MeshProps->SourceUVLayerNamesList.IndexOfByKey(MeshProps->SourceNormalMapUVLayer);
+	const int DetailUVLayer = InputMeshSettings->SourceUVLayerNamesList.IndexOfByKey(InputMeshSettings->SourceNormalMapUVLayer);
 	const FDynamicMeshUVOverlay* UVOverlay = DetailMesh->Attributes()->GetUVLayer(DetailUVLayer);
 	if (UVOverlay == nullptr)
 	{
@@ -621,7 +621,7 @@ EBakeOpState UBakeMeshAttributeMapsTool::UpdateResult_DetailNormalMap()
 		return EBakeOpState::Invalid;
 	}
 	
-	UTexture2D* DetailMeshNormalMap = MeshProps->SourceNormalMap; 
+	UTexture2D* DetailMeshNormalMap = InputMeshSettings->SourceNormalMap; 
 	if (DetailMeshNormalMap)
 	{
 		CachedDetailNormalMap = MakeShared<UE::Geometry::TImageBuilder<FVector4f>, ESPMode::ThreadSafe>();
diff --git a/Engine/Plugins/Experimental/MeshModelingToolsetExp/Source/MeshModelingToolsExp/Private/BakeMeshAttributeVertexTool.cpp b/Engine/Plugins/Experimental/MeshModelingToolsetExp/Source/MeshModelingToolsExp/Private/BakeMeshAttributeVertexTool.cpp
--- a/Engine/Plugins/Experimental/MeshModelingToolsetExp/Source/MeshModelingToolsExp/Private/BakeMeshAttributeVertexTool.cpp
+++ b/Engine/Plugins/Experimental/MeshModelingToolsetExp/Source/MeshModelingToolsExp/Private/BakeMeshAttributeVertexTool.cpp
@@ -198,7 +198,7 @@ public:
 			case EBakeMapType::Texture:
 			{
 				TSharedPtr<FMeshResampleImageEvaluator, ESPMode::ThreadSafe> TextureEval = MakeShared<FMeshResampleImageEvaluator, ESPMode::ThreadSafe>();
-				DetailSampler.SetColorMap(DetailMesh.Get(), IMeshBakerDetailSampler::FBakeDetailTexture(TextureImage.Get(), TextureSettings.UVLayer));
+				DetailSampler.SetTextureMap(DetailMesh.Get(), IMeshBakerDetailSampler::FBakeDetailTexture(TextureImage.Get(), TextureSettings.UVLayer));
 				Baker->ColorEvaluator = TextureEval;
 				break;
 			}
@@ -274,21 +274,21 @@ void UBakeMeshAttributeVertexTool::Setup()
 
 	// Setup tool property sets
 
-	MeshProps = NewObject<UBakeInputMeshProperties>(this);
-	MeshProps->RestoreProperties(this);
-	AddToolPropertySource(MeshProps);
-	SetToolPropertySourceEnabled(MeshProps, true);
-	MeshProps->bHasTargetUVLayer = false;
-	MeshProps->bHasSourceNormalMap = false;
-	MeshProps->TargetStaticMesh = GetStaticMeshTarget(Target);
-	MeshProps->TargetSkeletalMesh = GetSkeletalMeshTarget(Target);
-	MeshProps->TargetDynamicMesh = GetDynamicMeshTarget(Target);
-	MeshProps->SourceStaticMesh = !bIsBakeToSelf ? GetStaticMeshTarget(DetailTarget) : nullptr;
-	MeshProps->SourceSkeletalMesh = !bIsBakeToSelf ? GetSkeletalMeshTarget(DetailTarget) : nullptr;
-	MeshProps->SourceDynamicMesh = !bIsBakeToSelf ? GetDynamicMeshTarget(DetailTarget) : nullptr;
-	MeshProps->SourceNormalMap = nullptr;
-	MeshProps->WatchProperty(MeshProps->ProjectionDistance, [this](float) { OpState |= EBakeOpState::Evaluate; });
-	MeshProps->WatchProperty(MeshProps->bProjectionInWorldSpace, [this](bool) { OpState |= EBakeOpState::EvaluateDetailMesh; });
+	InputMeshSettings = NewObject<UBakeInputMeshProperties>(this);
+	InputMeshSettings->RestoreProperties(this);
+	AddToolPropertySource(InputMeshSettings);
+	SetToolPropertySourceEnabled(InputMeshSettings, true);
+	InputMeshSettings->bHasTargetUVLayer = false;
+	InputMeshSettings->bHasSourceNormalMap = false;
+	InputMeshSettings->TargetStaticMesh = GetStaticMeshTarget(Target);
+	InputMeshSettings->TargetSkeletalMesh = GetSkeletalMeshTarget(Target);
+	InputMeshSettings->TargetDynamicMesh = GetDynamicMeshTarget(Target);
+	InputMeshSettings->SourceStaticMesh = !bIsBakeToSelf ? GetStaticMeshTarget(DetailTarget) : nullptr;
+	InputMeshSettings->SourceSkeletalMesh = !bIsBakeToSelf ? GetSkeletalMeshTarget(DetailTarget) : nullptr;
+	InputMeshSettings->SourceDynamicMesh = !bIsBakeToSelf ? GetDynamicMeshTarget(DetailTarget) : nullptr;
+	InputMeshSettings->SourceNormalMap = nullptr;
+	InputMeshSettings->WatchProperty(InputMeshSettings->ProjectionDistance, [this](float) { OpState |= EBakeOpState::Evaluate; });
+	InputMeshSettings->WatchProperty(InputMeshSettings->bProjectionInWorldSpace, [this](bool) { OpState |= EBakeOpState::EvaluateDetailMesh; });
 	
 	Settings = NewObject<UBakeMeshAttributeVertexToolProperties>(this);
 	Settings->RestoreProperties(this);
@@ -366,7 +366,7 @@ void UBakeMeshAttributeVertexTool::Shutdown(EToolShutdownType ShutdownType)
 	TRACE_CPUPROFILER_EVENT_SCOPE(UBakeMeshAttributeVertexTool::Shutdown);
 	
 	Settings->SaveProperties(this);
-	MeshProps->SaveProperties(this);
+	InputMeshSettings->SaveProperties(this);
 	OcclusionSettings->SaveProperties(this);
 	CurvatureSettings->SaveProperties(this);
 	TextureSettings->SaveProperties(this);
@@ -459,7 +459,7 @@ void UBakeMeshAttributeVertexTool::UpdateDetailMesh()
 	DetailMesh = MakeShared<FDynamicMesh3, ESPMode::ThreadSafe>();
 	FMeshDescriptionToDynamicMesh Converter;
 	Converter.Convert(DetailMeshProvider->GetMeshDescription(), *DetailMesh);
-	if (MeshProps->bProjectionInWorldSpace && bIsBakeToSelf == false)
+	if (InputMeshSettings->bProjectionInWorldSpace && bIsBakeToSelf == false)
 	{
 		const UE::Geometry::FTransform3d DetailToWorld(DetailComponent->GetWorldTransform());
 		MeshTransforms::ApplyTransform(*DetailMesh, DetailToWorld);
@@ -634,8 +634,8 @@ void UBakeMeshAttributeVertexTool::UpdateResult()
 	BakeSettings.OutputTypePerChannel[3] = Settings->OutputTypeA;
 	BakeSettings.bSplitAtNormalSeams = Settings->bSplitAtNormalSeams;
 	BakeSettings.bSplitAtUVSeams = Settings->bSplitAtUVSeams;
-	BakeSettings.bProjectionInWorldSpace = MeshProps->bProjectionInWorldSpace;
-	BakeSettings.ProjectionDistance = MeshProps->ProjectionDistance;
+	BakeSettings.bProjectionInWorldSpace = InputMeshSettings->bProjectionInWorldSpace;
+	BakeSettings.ProjectionDistance = InputMeshSettings->ProjectionDistance;
 	if (!(BakeSettings == CachedBakeSettings))
 	{
 		CachedBakeSettings = BakeSettings;
diff --git a/Engine/Plugins/Experimental/MeshModelingToolsetExp/Source/MeshModelingToolsExp/Private/BakeMultiMeshAttributeMapsTool.cpp b/Engine/Plugins/Experimental/MeshModelingToolsetExp/Source/MeshModelingToolsExp/Private/BakeMultiMeshAttributeMapsTool.cpp
--- a/Engine/Plugins/Experimental/MeshModelingToolsetExp/Source/MeshModelingToolsExp/Private/BakeMultiMeshAttributeMapsTool.cpp
+++ b/Engine/Plugins/Experimental/MeshModelingToolsetExp/Source/MeshModelingToolsExp/Private/BakeMultiMeshAttributeMapsTool.cpp
@@ -97,9 +97,9 @@ public:
 		return Spatial->GetTriangleCount();
 	}
 
-	virtual void SetColorMap(const void* Mesh, const FBakeDetailTexture& Map) override
+	virtual void SetTextureMap(const void* Mesh, const FBakeDetailTexture& Map) override
 	{
-		DetailColorMaps[Mesh] = Map;
+		DetailTextureMaps[Mesh] = Map;
 	}
 
 	virtual void SetNormalMap(const void* Mesh, const FBakeDetailTexture& Map) override
@@ -107,9 +107,9 @@ public:
 		DetailNormalMaps[Mesh] = Map;
 	}
 
-	virtual const FBakeDetailTexture* GetColorMap(const void* Mesh) const override
+	virtual const FBakeDetailTexture* GetTextureMap(const void* Mesh) const override
 	{
-		return DetailColorMaps.Find(Mesh);
+		return DetailTextureMaps.Find(Mesh);
 	}
 	
 	virtual const FBakeDetailTexture* GetNormalMap(const void* Mesh) const override
@@ -288,9 +288,9 @@ public:
 	// End IMeshBakerDetailSampler interface
 
 	/** Initialize the mesh to color textures map */
-	void SetColorMaps(const FDetailTextureMap& Map)
+	void SetTextureMaps(const FDetailTextureMap& Map)
 	{
-		DetailColorMaps = Map;
+		DetailTextureMaps = Map;
 	}
 
 	/** Initialize the mesh to normal textures map */
@@ -301,7 +301,7 @@ public:
 
 protected:
 	FMeshSceneAdapter* MeshScene = nullptr;
-	FDetailTextureMap DetailColorMaps;
+	FDetailTextureMap DetailTextureMaps;
 	FDetailTextureMap DetailNormalMaps;
 };
 
@@ -354,7 +354,7 @@ public:
 			case EBakeMapType::Texture:
 			{
 				TSharedPtr<FMeshResampleImageEvaluator, ESPMode::ThreadSafe> TextureEval = MakeShared<FMeshResampleImageEvaluator, ESPMode::ThreadSafe>();
-				DetailSampler.SetColorMaps(CachedMeshToColorImageMap);
+				DetailSampler.SetTextureMaps(CachedMeshToColorImageMap);
 				Baker->AddEvaluator(TextureEval);
 				break;
 			}
@@ -410,13 +410,13 @@ void UBakeMultiMeshAttributeMapsTool::Setup()
 
 	// Setup tool property sets
 	
-	InputMeshProps = NewObject<UBakeMultiMeshInputToolProperties>(this);
-	InputMeshProps->RestoreProperties(this);
-	AddToolPropertySource(InputMeshProps);
-	InputMeshProps->TargetStaticMesh = GetStaticMeshTarget(Target);
-	UpdateUVLayerNames(InputMeshProps->TargetUVLayer, InputMeshProps->TargetUVLayerNamesList, TargetMesh);
-	InputMeshProps->WatchProperty(InputMeshProps->TargetUVLayer, [this](FString) { OpState |= EBakeOpState::Evaluate; });
-	InputMeshProps->WatchProperty(InputMeshProps->ProjectionDistance, [this](float) { OpState |= EBakeOpState::Evaluate; });
+	InputMeshSettings = NewObject<UBakeMultiMeshInputToolProperties>(this);
+	InputMeshSettings->RestoreProperties(this);
+	AddToolPropertySource(InputMeshSettings);
+	InputMeshSettings->TargetStaticMesh = GetStaticMeshTarget(Target);
+	UpdateUVLayerNames(InputMeshSettings->TargetUVLayer, InputMeshSettings->TargetUVLayerNamesList, TargetMesh);
+	InputMeshSettings->WatchProperty(InputMeshSettings->TargetUVLayer, [this](FString) { OpState |= EBakeOpState::Evaluate; });
+	InputMeshSettings->WatchProperty(InputMeshSettings->ProjectionDistance, [this](float) { OpState |= EBakeOpState::Evaluate; });
 	
 	Settings = NewObject<UBakeMultiMeshAttributeMapsToolProperties>(this);
 	Settings->RestoreProperties(this);
@@ -457,9 +457,9 @@ void UBakeMultiMeshAttributeMapsTool::Setup()
 		FBakeMultiMeshDetailProperties SourceMeshProp;
 		SourceMeshProp.SourceMesh = GetStaticMeshTarget(DetailTarget);
 		SourceMeshProp.SourceTexture = DetailColorTexture;
-		InputMeshProps->SourceMeshes.Add(SourceMeshProp);
-		InputMeshProps->WatchProperty(InputMeshProps->SourceMeshes[Idx-1].SourceTexture, [this](UTexture2D*) { OpState |= EBakeOpState::Evaluate; });
-		InputMeshProps->WatchProperty(InputMeshProps->SourceMeshes[Idx-1].SourceTextureUVLayer, [this](int) { OpState |= EBakeOpState::Evaluate; });
+		InputMeshSettings->SourceMeshes.Add(SourceMeshProp);
+		InputMeshSettings->WatchProperty(InputMeshSettings->SourceMeshes[Idx-1].SourceTexture, [this](UTexture2D*) { OpState |= EBakeOpState::Evaluate; });
+		InputMeshSettings->WatchProperty(InputMeshSettings->SourceMeshes[Idx-1].SourceTextureUVLayer, [this](int) { OpState |= EBakeOpState::Evaluate; });
 	}
 
 	UpdateOnModeChange();
@@ -522,7 +522,7 @@ void UBakeMultiMeshAttributeMapsTool::Shutdown(EToolShutdownType ShutdownType)
 	Super::Shutdown(ShutdownType);
 	
 	Settings->SaveProperties(this);
-	InputMeshProps->SaveProperties(this);
+	InputMeshSettings->SaveProperties(this);
 
 	if (Compute)
 	{
@@ -562,8 +562,8 @@ void UBakeMultiMeshAttributeMapsTool::UpdateResult()
 	FBakeSettings BakeSettings;
 	BakeSettings.Dimensions = Dimensions;
 	BakeSettings.BitDepth = Settings->BitDepth;
-	BakeSettings.TargetUVLayer = InputMeshProps->TargetUVLayerNamesList.IndexOfByKey(InputMeshProps->TargetUVLayer);
-	BakeSettings.ProjectionDistance = InputMeshProps->ProjectionDistance;
+	BakeSettings.TargetUVLayer = InputMeshSettings->TargetUVLayerNamesList.IndexOfByKey(InputMeshSettings->TargetUVLayer);
+	BakeSettings.ProjectionDistance = InputMeshSettings->ProjectionDistance;
 	BakeSettings.bProjectionInWorldSpace = true;  // Always world space
 	BakeSettings.SamplesPerPixel = static_cast<int32>(Settings->SamplesPerPixel);
 
@@ -607,8 +607,8 @@ EBakeOpState UBakeMultiMeshAttributeMapsTool::UpdateResult_DetailMeshes()
 {
 	FBakeMultiMeshDetailSettings NewSettings;
 
-	// Iterate through our detail props to build our detail mesh data.
-	const int32 NumDetail = InputMeshProps->SourceMeshes.Num();
+	// Iterate through our detail properties to build our detail mesh data.
+	const int32 NumDetail = InputMeshSettings->SourceMeshes.Num();
 	CachedColorImages.SetNum(NumDetail);
 	CachedColorUVLayers.SetNum(NumDetail);
 	TMap<UActorComponent*, int> ActorToDataMap;
@@ -620,8 +620,8 @@ EBakeOpState UBakeMultiMeshAttributeMapsTool::UpdateResult_DetailMeshes()
 		// Color map data
 		if (static_cast<bool>(CachedBakeSettings.BakeMapTypes & EBakeMapType::Texture))
 		{
-			UTexture2D* ColorMapSourceTexture = InputMeshProps->SourceMeshes[Idx].SourceTexture;
-			const int ColorMapUVLayer = InputMeshProps->SourceMeshes[Idx].SourceTextureUVLayer;
+			UTexture2D* ColorMapSourceTexture = InputMeshSettings->SourceMeshes[Idx].SourceTexture;
+			const int ColorMapUVLayer = InputMeshSettings->SourceMeshes[Idx].SourceTextureUVLayer;
 			if (!ColorMapSourceTexture)
 			{
 				GetToolManager()->DisplayMessage(LOCTEXT("InvalidTextureWarning", "The Source Texture is not valid"), EToolMessageLevel::UserWarning);
diff --git a/Engine/Plugins/Experimental/MeshModelingToolsetExp/Source/MeshModelingToolsExp/Public/BakeMeshAttributeMapsTool.h b/Engine/Plugins/Experimental/MeshModelingToolsetExp/Source/MeshModelingToolsExp/Public/BakeMeshAttributeMapsTool.h
--- a/Engine/Plugins/Experimental/MeshModelingToolsetExp/Source/MeshModelingToolsExp/Public/BakeMeshAttributeMapsTool.h
+++ b/Engine/Plugins/Experimental/MeshModelingToolsetExp/Source/MeshModelingToolsExp/Public/BakeMeshAttributeMapsTool.h
@@ -114,7 +114,7 @@ protected:
 	// @todo setters/getters for these
 
 	UPROPERTY()
-	TObjectPtr<UBakeInputMeshProperties> MeshProps;
+	TObjectPtr<UBakeInputMeshProperties> InputMeshSettings;
 
 	UPROPERTY()
 	TObjectPtr<UBakeMeshAttributeMapsToolProperties> Settings;
diff --git a/Engine/Plugins/Experimental/MeshModelingToolsetExp/Source/MeshModelingToolsExp/Public/BakeMeshAttributeVertexTool.h b/Engine/Plugins/Experimental/MeshModelingToolsetExp/Source/MeshModelingToolsExp/Public/BakeMeshAttributeVertexTool.h
--- a/Engine/Plugins/Experimental/MeshModelingToolsetExp/Source/MeshModelingToolsExp/Public/BakeMeshAttributeVertexTool.h
+++ b/Engine/Plugins/Experimental/MeshModelingToolsetExp/Source/MeshModelingToolsExp/Public/BakeMeshAttributeVertexTool.h
@@ -135,7 +135,7 @@ public:
 
 protected:
 	UPROPERTY()
-	TObjectPtr<UBakeInputMeshProperties> MeshProps;
+	TObjectPtr<UBakeInputMeshProperties> InputMeshSettings;
 	
 	UPROPERTY()
 	TObjectPtr<UBakeMeshAttributeVertexToolProperties> Settings;
diff --git a/Engine/Plugins/Experimental/MeshModelingToolsetExp/Source/MeshModelingToolsExp/Public/BakeMultiMeshAttributeMapsTool.h b/Engine/Plugins/Experimental/MeshModelingToolsetExp/Source/MeshModelingToolsExp/Public/BakeMultiMeshAttributeMapsTool.h
--- a/Engine/Plugins/Experimental/MeshModelingToolsetExp/Source/MeshModelingToolsExp/Public/BakeMultiMeshAttributeMapsTool.h
+++ b/Engine/Plugins/Experimental/MeshModelingToolsetExp/Source/MeshModelingToolsExp/Public/BakeMultiMeshAttributeMapsTool.h
@@ -182,7 +182,7 @@ protected:
 	TObjectPtr<UBakeMultiMeshAttributeMapsToolProperties> Settings;
 
 	UPROPERTY()
-	TObjectPtr<UBakeMultiMeshInputToolProperties> InputMeshProps;
+	TObjectPtr<UBakeMultiMeshInputToolProperties> InputMeshSettings;
 
 
 protected:
diff --git a/Engine/Plugins/Runtime/GeometryProcessing/Source/DynamicMesh/Private/Sampling/MeshResampleImageEvaluator.cpp b/Engine/Plugins/Runtime/GeometryProcessing/Source/DynamicMesh/Private/Sampling/MeshResampleImageEvaluator.cpp
--- a/Engine/Plugins/Runtime/GeometryProcessing/Source/DynamicMesh/Private/Sampling/MeshResampleImageEvaluator.cpp
+++ b/Engine/Plugins/Runtime/GeometryProcessing/Source/DynamicMesh/Private/Sampling/MeshResampleImageEvaluator.cpp
@@ -46,7 +46,7 @@ FVector4f FMeshResampleImageEvaluator::ImageSampleFunction(const FCorrespondence
 	const int32 DetailTriID = SampleData.DetailTriID;
 	const TImageBuilder<FVector4f>* TextureImage = nullptr;
 	int DetailUVLayer = 0;
-	const IMeshBakerDetailSampler::FBakeDetailTexture* ColorMap = DetailSampler->GetColorMap(DetailMesh);
+	const IMeshBakerDetailSampler::FBakeDetailTexture* ColorMap = DetailSampler->GetTextureMap(DetailMesh);
 	if (ColorMap)
 	{
 		Tie(TextureImage, DetailUVLayer) = *ColorMap;
diff --git a/Engine/Plugins/Runtime/GeometryProcessing/Source/DynamicMesh/Public/Sampling/MeshBakerCommon.h b/Engine/Plugins/Runtime/GeometryProcessing/Source/DynamicMesh/Public/Sampling/MeshBakerCommon.h
--- a/Engine/Plugins/Runtime/GeometryProcessing/Source/DynamicMesh/Public/Sampling/MeshBakerCommon.h
+++ b/Engine/Plugins/Runtime/GeometryProcessing/Source/DynamicMesh/Public/Sampling/MeshBakerCommon.h
@@ -105,14 +105,14 @@ public:
 	/** @return the triangle count of a given mesh. */
 	virtual int32 GetTriangleCount(const void* Mesh) const = 0;
 
-	/** Associate a color map and UV layer index for a given mesh in the detail set */
-	virtual void SetColorMap(const void* Mesh, const FBakeDetailTexture& Map) = 0;
+	/** Associate a texture map and UV layer index for a given mesh in the detail set */
+	virtual void SetTextureMap(const void* Mesh, const FBakeDetailTexture& Map) = 0;
 
 	/** Associate a normal map and UV layer index for a given mesh in the detail set */
 	virtual void SetNormalMap(const void* Mesh, const FBakeDetailTexture& Map) = 0;
 
-	/** Retrieve a color map and UV layer index from a given mesh in the detail set */
-	virtual const FBakeDetailTexture* GetColorMap(const void* Mesh) const = 0;
+	/** Retrieve a texture map and UV layer index from a given mesh in the detail set */
+	virtual const FBakeDetailTexture* GetTextureMap(const void* Mesh) const = 0;
 
 	/** Retrieve a normal map and UV layer index from a given mesh in the detail set */
 	virtual const FBakeDetailTexture* GetNormalMap(const void* Mesh) const = 0;
@@ -337,9 +337,9 @@ public:
 		return DynamicMesh->TriangleCount();
 	}
 
-	virtual void SetColorMap(const void* Mesh, const FBakeDetailTexture& Map) override
+	virtual void SetTextureMap(const void* Mesh, const FBakeDetailTexture& Map) override
 	{
-		DetailColorMap = Map;
+		DetailTextureMap = Map;
 	}
 
 	virtual void SetNormalMap(const void* Mesh, const FBakeDetailTexture& Map) override
@@ -347,9 +347,9 @@ public:
 		DetailNormalMap = Map;
 	}
 
-	virtual const FBakeDetailTexture* GetColorMap(const void* Mesh) const override
+	virtual const FBakeDetailTexture* GetTextureMap(const void* Mesh) const override
 	{
-		return &DetailColorMap;
+		return &DetailTextureMap;
 	}
 	
 	virtual const FBakeDetailTexture* GetNormalMap(const void* Mesh) const override
@@ -553,7 +553,7 @@ protected:
 	const FDynamicMesh3* DetailMesh = nullptr;
 	const FDynamicMeshAABBTree3* DetailSpatial = nullptr;
 	const FMeshTangentsd* DetailTangents = nullptr;
-	FBakeDetailTexture DetailColorMap = FBakeDetailTexture(nullptr, 0);
+	FBakeDetailTexture DetailTextureMap = FBakeDetailTexture(nullptr, 0);
 	FBakeDetailTexture DetailNormalMap = FBakeDetailTexture(nullptr, 0);
 };		
 	
