ModelingTools: add output mesh type property set to VolumeToMeshTool, allow Static and Dynamic Mesh output. Change normals computation to compute polygroup normals instead of per-triangle normals (which would then have normal seams)
#rnx

diff --git a/Engine/Plugins/Experimental/MeshModelingToolsetExp/Source/MeshModelingToolsExp/Private/VolumeToMeshTool.cpp b/Engine/Plugins/Experimental/MeshModelingToolsetExp/Source/MeshModelingToolsExp/Private/VolumeToMeshTool.cpp
--- a/Engine/Plugins/Experimental/MeshModelingToolsetExp/Source/MeshModelingToolsExp/Private/VolumeToMeshTool.cpp
+++ b/Engine/Plugins/Experimental/MeshModelingToolsetExp/Source/MeshModelingToolsExp/Private/VolumeToMeshTool.cpp
@@ -13,11 +13,6 @@
 #include "Selection/ToolSelectionUtil.h"
 #include "ModelingObjectsCreationAPI.h"
 
-#include "DynamicMesh/Operations/MergeCoincidentMeshEdges.h"
-#include "CompGeom/PolygonTriangulation.h"
-#include "MeshBoundaryLoops.h"
-#include "Operations/MinimalHoleFiller.h"
-
 #include "Model.h"
 
 using namespace UE::Geometry;
@@ -87,6 +82,11 @@ void UVolumeToMeshTool::Setup()
 	VolumeEdgesSet->SetLineMaterial(ToolSetupUtil::GetDefaultLineComponentMaterial(GetToolManager()));
 	VolumeEdgesSet->RegisterComponent();
 
+	OutputTypeProperties = NewObject<UCreateMeshObjectTypeProperties>(this);
+	OutputTypeProperties->RestoreProperties(this, TEXT("VolumeToMeshTool"));
+	OutputTypeProperties->Initialize(true, false, true);
+	OutputTypeProperties->WatchProperty(OutputTypeProperties->OutputType, [this](FString) { OutputTypeProperties->UpdatePropertyVisibility(); });
+	AddToolPropertySource(OutputTypeProperties);
 
 	Settings = NewObject<UVolumeToMeshToolProperties>(this);
 	Settings->RestoreProperties(this);
@@ -100,7 +100,7 @@ void UVolumeToMeshTool::Setup()
 	bResultValid = false;
 
 	GetToolManager()->DisplayMessage( 
-		LOCTEXT("OnStartTool", "Convert a Volume to a Static Mesh"),
+		LOCTEXT("OnStartTool", "Convert a Volume to a Mesh"),
 		EToolMessageLevel::UserNotification);
 }
 
@@ -108,6 +108,7 @@ void UVolumeToMeshTool::Setup()
 
 void UVolumeToMeshTool::Shutdown(EToolShutdownType ShutdownType)
 {
+	OutputTypeProperties->SaveProperties(this, TEXT("VolumeToMeshTool"));
 	Settings->SaveProperties(this);
 
 	UE::Geometry::FTransform3d Transform(PreviewMesh->GetTransform());
@@ -130,6 +131,7 @@ void UVolumeToMeshTool::Shutdown(EToolShutdownType ShutdownType)
 		NewMeshObjectParams.BaseName = NewName;
 		NewMeshObjectParams.Materials.Add(UseMaterial);
 		NewMeshObjectParams.SetMesh(&CurrentMesh);
+		OutputTypeProperties->ConfigureCreateMeshObjectParams(NewMeshObjectParams);
 		FCreateMeshObjectResult Result = UE::Modeling::CreateMeshObject(GetToolManager(), MoveTemp(NewMeshObjectParams));
 		if (Result.IsOK() && Result.NewActor != nullptr)
 		{
@@ -206,7 +208,13 @@ void UVolumeToMeshTool::RecalculateMesh()
 
 		CurrentMesh = FDynamicMesh3(EMeshComponents::FaceGroups);
 		UE::Conversion::VolumeToDynamicMesh(TargetVolume.Get(), CurrentMesh, Options);
-		FMeshNormals::InitializeMeshToPerTriangleNormals(&CurrentMesh);
+
+		// compute normals for current polygroup topology
+		CurrentMesh.EnableAttributes();
+		FDynamicMeshNormalOverlay* Normals = CurrentMesh.Attributes()->PrimaryNormals();
+		FMeshNormals::InitializeOverlayTopologyFromFaceGroups(&CurrentMesh, Normals);
+		FMeshNormals::QuickRecomputeOverlayNormals(CurrentMesh);
+
 		PreviewMesh->UpdatePreview(&CurrentMesh);
 	}
 
diff --git a/Engine/Plugins/Experimental/MeshModelingToolsetExp/Source/MeshModelingToolsExp/Public/VolumeToMeshTool.h b/Engine/Plugins/Experimental/MeshModelingToolsetExp/Source/MeshModelingToolsExp/Public/VolumeToMeshTool.h
--- a/Engine/Plugins/Experimental/MeshModelingToolsetExp/Source/MeshModelingToolsExp/Public/VolumeToMeshTool.h
+++ b/Engine/Plugins/Experimental/MeshModelingToolsetExp/Source/MeshModelingToolsExp/Public/VolumeToMeshTool.h
@@ -4,11 +4,11 @@
 
 #include "CoreMinimal.h"
 #include "GameFramework/Volume.h"
-#include "SingleSelectionTool.h"
 #include "InteractiveToolBuilder.h"
 #include "DynamicMesh/DynamicMesh3.h"
 #include "PreviewMesh.h"
 #include "Drawing/LineSetComponent.h"
+#include "PropertySets/CreateMeshObjectTypeProperties.h"
 #include "VolumeToMeshTool.generated.h"
 
 /**
@@ -78,6 +78,9 @@ protected:
 	UPROPERTY()
 	TObjectPtr<UVolumeToMeshToolProperties> Settings;
 
+	UPROPERTY()
+	TObjectPtr<UCreateMeshObjectTypeProperties> OutputTypeProperties;
+
 	UPROPERTY()
 	TObjectPtr<UPreviewMesh> PreviewMesh;
 
