ModelingTools: do a planar-simplification pass in MeshToVolumeTool when using minimal polygons. This removes spurious vertices when the input mesh is more tessellated than strictly necessary.
#rnx

diff --git a/Engine/Plugins/Experimental/MeshModelingToolsetExp/Source/MeshModelingToolsEditorOnly/Private/MeshToVolumeTool.cpp b/Engine/Plugins/Experimental/MeshModelingToolsetExp/Source/MeshModelingToolsEditorOnly/Private/MeshToVolumeTool.cpp
--- a/Engine/Plugins/Experimental/MeshModelingToolsetExp/Source/MeshModelingToolsEditorOnly/Private/MeshToVolumeTool.cpp
+++ b/Engine/Plugins/Experimental/MeshModelingToolsetExp/Source/MeshModelingToolsEditorOnly/Private/MeshToVolumeTool.cpp
@@ -5,20 +5,15 @@
 #include "ToolBuilderUtil.h"
 
 #include "DynamicMesh/DynamicMesh3.h"
-#include "DynamicMesh/MeshNormals.h"
-#include "DynamicMeshEditor.h"
-#include "MeshRegionBoundaryLoops.h"
+#include "MeshSimplification.h"
 #include "Util/ColorConstants.h"
 #include "ToolSetupUtil.h"
-#include "Selections/MeshConnectedComponents.h"
 #include "Selection/ToolSelectionUtil.h"
-#include "ToolTargetManager.h"
 #include "ModelingToolTargetUtil.h"
 
-
-#include "Engine/Classes/Engine/BlockingVolume.h"
-#include "Engine/Classes/Components/BrushComponent.h"
-#include "Engine/Classes/Engine/Polys.h"
+#include "Engine/BlockingVolume.h"
+#include "Components/BrushComponent.h"
+#include "Engine/Polys.h"
 #include "Model.h"
 #include "BSPOps.h"
 
@@ -199,7 +194,13 @@ void UMeshToVolumeTool::RecalculateVolume()
 		// the group differentiations if they are coplanar.
 		bool bRespectGroupBoundaries = false;
 
-		UE::Conversion::GetPolygonFaces(InputMesh, Faces, bRespectGroupBoundaries);
+		// Apply minimal-planar simplification to remove extra vertices along straight edges
+		FDynamicMesh3 LocalMesh = InputMesh;
+		LocalMesh.DiscardAttributes();
+		FQEMSimplification PlanarSimplifier(&LocalMesh);
+		PlanarSimplifier.SimplifyToMinimalPlanar(0.1);		// angle tolerance in degrees
+
+		UE::Conversion::GetPolygonFaces(LocalMesh, Faces, bRespectGroupBoundaries);
 	}
 	else
 	{
diff --git a/Engine/Plugins/Runtime/MeshModelingToolset/Source/ModelingComponentsEditorOnly/Private/ConversionUtils/DynamicMeshToVolume.cpp b/Engine/Plugins/Runtime/MeshModelingToolset/Source/ModelingComponentsEditorOnly/Private/ConversionUtils/DynamicMeshToVolume.cpp
--- a/Engine/Plugins/Runtime/MeshModelingToolset/Source/ModelingComponentsEditorOnly/Private/ConversionUtils/DynamicMeshToVolume.cpp
+++ b/Engine/Plugins/Runtime/MeshModelingToolset/Source/ModelingComponentsEditorOnly/Private/ConversionUtils/DynamicMeshToVolume.cpp
@@ -421,25 +421,6 @@ void GetPolygonFaces(const FDynamicMesh3& InputMesh, TArray<FDynamicMeshFace>& F
 
 			Faces.Add(Face);
 		}
-
-		//for (const FEdgeLoop& Loop : Loops.Loops)
-		//{
-		//	FDynamicMeshFace Face;
-
-		//	FVector3d AvgPos(0, 0, 0);
-		//	for (int32 vid : Loop.Vertices)
-		//	{
-		//		FVector3d Position = InputMesh.GetVertex(vid);
-		//		Face.BoundaryLoop.Add(Position);
-		//		AvgPos += Position;
-		//	}
-		//	AvgPos /= (double)Loop.Vertices.Num();
-		//	Algo::Reverse(Face.BoundaryLoop);
-
-		//	Face.Plane = FFrame3d(AvgPos, FaceNormal);
-
-		//	Faces.Add(Face);
-		//}
 	}
 }
 
