Attribute editor tool: when adding a new attribute, or undoing the deletion of an attribute, set the selected attribute in the UI to the name of new attribute.

#rnx

diff --git a/Engine/Plugins/Experimental/MeshModelingToolsetExp/Source/MeshModelingToolsEditorOnly/Private/AttributeEditorTool.cpp b/Engine/Plugins/Experimental/MeshModelingToolsetExp/Source/MeshModelingToolsEditorOnly/Private/AttributeEditorTool.cpp
--- a/Engine/Plugins/Experimental/MeshModelingToolsetExp/Source/MeshModelingToolsEditorOnly/Private/AttributeEditorTool.cpp
+++ b/Engine/Plugins/Experimental/MeshModelingToolsetExp/Source/MeshModelingToolsEditorOnly/Private/AttributeEditorTool.cpp
@@ -314,6 +314,8 @@ void UAttributeEditorTool::InitializeAttributeLists()
 	ExtractAttribList(Mesh, Mesh->EdgeAttributes(), EAttributeEditorElementType::Edge, this->EdgeAttributes, AttributeProps->EdgeAttributes);
 	ExtractAttribList(Mesh, Mesh->PolygonGroupAttributes(), EAttributeEditorElementType::PolygonGroup, this->GroupAttributes, AttributeProps->GroupAttributes);
 
+	TArray<FString> OldAttributeNames = ModifyAttributeProps->AttributeNamesList;
+
 	ModifyAttributeProps->AttributeNamesList.Reset();
 	CopyAttributeProps->FromAttribute.Reset();
 	CopyAttributeProps->ToAttribute.Reset();
@@ -343,6 +345,16 @@ void UAttributeEditorTool::InitializeAttributeLists()
 		ModifyAttributeProps->Attribute.Reset();
 	}
 
+	// If we've added a new attribute to the list, set it as the selected attribute (useful when undoing an attribute deletion)
+	for (const FString& AttributeName : ModifyAttributeProps->AttributeNamesList)
+	{
+		if (!OldAttributeNames.Contains(AttributeName))
+		{
+			ModifyAttributeProps->Attribute = AttributeName;
+			break;
+		}
+	}
+
 	bAttributeListsValid = true;
 }
 
