ModelingTools: Update the Texture/MultiTexture UV parameters to use GetOptions UI in BakeTexture/BakeVertex.
Fixed MultiTexture UV channel parameter not propagating to evaluator in BakeTexture.
Fixed crash when using invalid source UV channel for BakeAll.

#rnx

diff --git a/Engine/Plugins/Experimental/MeshModelingToolsetExp/Source/MeshModelingToolsExp/Private/BakeMeshAttributeMapsTool.cpp b/Engine/Plugins/Experimental/MeshModelingToolsetExp/Source/MeshModelingToolsExp/Private/BakeMeshAttributeMapsTool.cpp
--- a/Engine/Plugins/Experimental/MeshModelingToolsetExp/Source/MeshModelingToolsExp/Private/BakeMeshAttributeMapsTool.cpp
+++ b/Engine/Plugins/Experimental/MeshModelingToolsetExp/Source/MeshModelingToolsExp/Private/BakeMeshAttributeMapsTool.cpp
@@ -334,7 +334,7 @@ void UBakeMeshAttributeMapsTool::Setup()
 	TextureSettings->RestoreProperties(this);
 	AddToolPropertySource(TextureSettings);
 	SetToolPropertySourceEnabled(TextureSettings, false);
-	TextureSettings->WatchProperty(TextureSettings->UVLayer, [this](float) { OpState |= EBakeOpState::Evaluate; });
+	TextureSettings->WatchProperty(TextureSettings->UVLayer, [this](FString) { OpState |= EBakeOpState::Evaluate; });
 	TextureSettings->WatchProperty(TextureSettings->SourceTexture, [this](UTexture2D*) { OpState |= EBakeOpState::Evaluate; });
 
 	MultiTextureSettings = NewObject<UBakeMultiTexture2DProperties>(this);
@@ -345,7 +345,7 @@ void UBakeMeshAttributeMapsTool::Setup()
 	auto SetDirtyCallback = [this](decltype(MultiTextureSettings->MaterialIDSourceTextures)) { OpState |= EBakeOpState::Evaluate; };
 	auto NotEqualsCallback = [](const decltype(MultiTextureSettings->MaterialIDSourceTextures)& A, const decltype(MultiTextureSettings->MaterialIDSourceTextures)& B) -> bool { return A != B; };
 	MultiTextureSettings->WatchProperty(MultiTextureSettings->MaterialIDSourceTextures, SetDirtyCallback, NotEqualsCallback);
-	MultiTextureSettings->WatchProperty(MultiTextureSettings->UVLayer, [this](float) { OpState |= EBakeOpState::Evaluate; });
+	MultiTextureSettings->WatchProperty(MultiTextureSettings->UVLayer, [this](FString) { OpState |= EBakeOpState::Evaluate; });
 	UpdateMultiTextureMaterialIDs(DetailTarget, MultiTextureSettings->AllSourceTextures, MultiTextureSettings->MaterialIDSourceTextures);
 
 	UpdateOnModeChange();
@@ -433,7 +433,7 @@ TUniquePtr<UE::Geometry::TGenericDataOperator<FMeshMapBaker>> UBakeMeshAttribute
 
 	if ((bool)(CachedBakeSettings.BakeMapTypes & EBakeMapType::MultiTexture))
 	{
-		Op->TextureSettings = CachedMultiTexture2DSettings;
+		Op->MultiTextureSettings = CachedMultiTexture2DSettings;
 		Op->MaterialIDTextures = CachedMultiTextures;
 	}
 
@@ -506,6 +506,8 @@ void UBakeMeshAttributeMapsTool::UpdateDetailMesh()
 	}
 
 	UpdateUVLayerNames(InputMeshSettings->SourceNormalMapUVLayer, InputMeshSettings->SourceUVLayerNamesList, *DetailMesh);
+	UpdateUVLayerNames(TextureSettings->UVLayer, TextureSettings->UVLayerNamesList, *DetailMesh);
+	UpdateUVLayerNames(MultiTextureSettings->UVLayer, MultiTextureSettings->UVLayerNamesList, *DetailMesh);
 
 	// Clear detail mesh evaluation flag and mark evaluation.
 	OpState &= ~EBakeOpState::EvaluateDetailMesh;
diff --git a/Engine/Plugins/Experimental/MeshModelingToolsetExp/Source/MeshModelingToolsExp/Private/BakeMeshAttributeTool.cpp b/Engine/Plugins/Experimental/MeshModelingToolsetExp/Source/MeshModelingToolsExp/Private/BakeMeshAttributeTool.cpp
--- a/Engine/Plugins/Experimental/MeshModelingToolsetExp/Source/MeshModelingToolsExp/Private/BakeMeshAttributeTool.cpp
+++ b/Engine/Plugins/Experimental/MeshModelingToolsetExp/Source/MeshModelingToolsExp/Private/BakeMeshAttributeTool.cpp
@@ -153,7 +153,7 @@ EBakeOpState UBakeMeshAttributeTool::UpdateResult_Texture2DImage(const FImageDim
 
 	FTexture2DSettings NewSettings;
 	NewSettings.Dimensions = Dimensions;
-	NewSettings.UVLayer = TextureSettings->UVLayer;
+	NewSettings.UVLayer = TextureSettings->UVLayerNamesList.IndexOfByKey(TextureSettings->UVLayer);
 
 	const FDynamicMeshUVOverlay* UVOverlay = DetailMesh->Attributes()->GetUVLayer(NewSettings.UVLayer);
 	if (UVOverlay == nullptr)
@@ -196,7 +196,7 @@ EBakeOpState UBakeMeshAttributeTool::UpdateResult_MultiTexture(const FImageDimen
 
 	FTexture2DSettings NewSettings;
 	NewSettings.Dimensions = Dimensions;
-	NewSettings.UVLayer = MultiTextureSettings->UVLayer;
+	NewSettings.UVLayer = MultiTextureSettings->UVLayerNamesList.IndexOfByKey(MultiTextureSettings->UVLayer);
 
 	const FDynamicMeshUVOverlay* UVOverlay = DetailMesh->Attributes()->GetUVLayer(NewSettings.UVLayer);
 	if (UVOverlay == nullptr)
diff --git a/Engine/Plugins/Experimental/MeshModelingToolsetExp/Source/MeshModelingToolsExp/Private/BakeMeshAttributeVertexTool.cpp b/Engine/Plugins/Experimental/MeshModelingToolsetExp/Source/MeshModelingToolsExp/Private/BakeMeshAttributeVertexTool.cpp
--- a/Engine/Plugins/Experimental/MeshModelingToolsetExp/Source/MeshModelingToolsExp/Private/BakeMeshAttributeVertexTool.cpp
+++ b/Engine/Plugins/Experimental/MeshModelingToolsetExp/Source/MeshModelingToolsExp/Private/BakeMeshAttributeVertexTool.cpp
@@ -327,7 +327,7 @@ void UBakeMeshAttributeVertexTool::Setup()
 	TextureSettings->RestoreProperties(this);
 	AddToolPropertySource(TextureSettings);
 	SetToolPropertySourceEnabled(TextureSettings, false);
-	TextureSettings->WatchProperty(TextureSettings->UVLayer, [this](float) { OpState |= EBakeOpState::Evaluate; });
+	TextureSettings->WatchProperty(TextureSettings->UVLayer, [this](FString) { OpState |= EBakeOpState::Evaluate; });
 	TextureSettings->WatchProperty(TextureSettings->SourceTexture, [this](UTexture2D*) { OpState |= EBakeOpState::Evaluate; });
 
 	MultiTextureSettings = NewObject<UBakeMultiTexture2DProperties>(this);
@@ -337,7 +337,7 @@ void UBakeMeshAttributeVertexTool::Setup()
 	auto SetDirtyCallback = [this](decltype(MultiTextureSettings->MaterialIDSourceTextures)) { OpState |= EBakeOpState::Evaluate; };
 	auto NotEqualsCallback = [](const decltype(MultiTextureSettings->MaterialIDSourceTextures)& A, const decltype(MultiTextureSettings->MaterialIDSourceTextures)& B) -> bool { return A != B; };
 	MultiTextureSettings->WatchProperty(MultiTextureSettings->MaterialIDSourceTextures, SetDirtyCallback, NotEqualsCallback);
-	MultiTextureSettings->WatchProperty(MultiTextureSettings->UVLayer, [this](float) { OpState |= EBakeOpState::Evaluate; });
+	MultiTextureSettings->WatchProperty(MultiTextureSettings->UVLayer, [this](FString) { OpState |= EBakeOpState::Evaluate; });
 	UpdateMultiTextureMaterialIDs(DetailTarget, MultiTextureSettings->AllSourceTextures, MultiTextureSettings->MaterialIDSourceTextures);
 
 	UpdateOnModeChange();
@@ -470,6 +470,9 @@ void UBakeMeshAttributeVertexTool::UpdateDetailMesh()
 	DetailSpatial = MakeShared<FDynamicMeshAABBTree3, ESPMode::ThreadSafe>();
 	DetailSpatial->SetMesh(DetailMesh.Get(), true);
 
+	UpdateUVLayerNames(TextureSettings->UVLayer, TextureSettings->UVLayerNamesList, *DetailMesh);
+	UpdateUVLayerNames(MultiTextureSettings->UVLayer, MultiTextureSettings->UVLayerNamesList, *DetailMesh);
+
 	OpState &= ~EBakeOpState::EvaluateDetailMesh;
 	OpState |= EBakeOpState::Evaluate;
 	DetailMeshTimestamp++;
diff --git a/Engine/Plugins/Experimental/MeshModelingToolsetExp/Source/MeshModelingToolsExp/Public/BakeMeshAttributeToolCommon.h b/Engine/Plugins/Experimental/MeshModelingToolsetExp/Source/MeshModelingToolsExp/Public/BakeMeshAttributeToolCommon.h
--- a/Engine/Plugins/Experimental/MeshModelingToolsetExp/Source/MeshModelingToolsExp/Public/BakeMeshAttributeToolCommon.h
+++ b/Engine/Plugins/Experimental/MeshModelingToolsetExp/Source/MeshModelingToolsExp/Public/BakeMeshAttributeToolCommon.h
@@ -220,8 +220,18 @@ public:
 	TObjectPtr<UTexture2D> SourceTexture;
 
 	/** UV channel to use for the source mesh texture */
-	UPROPERTY(EditAnywhere, Category = TextureOutput, meta = (DisplayName = "Source Texture UV Channel"))
-	int32 UVLayer = 0;
+	UPROPERTY(EditAnywhere, Category = TextureOutput, meta = (DisplayName = "Source Texture UV Channel",
+		GetOptions = GetUVLayerNamesFunc, NoResetToDefault))
+	FString UVLayer;
+
+	UFUNCTION()
+	const TArray<FString>& GetUVLayerNamesFunc() const
+	{
+		return UVLayerNamesList;
+	}
+
+	UPROPERTY(meta = (TransientToolProperty))
+	TArray<FString> UVLayerNamesList;
 };
 
 
@@ -232,16 +242,26 @@ class MESHMODELINGTOOLSEXP_API UBakeMultiTexture2DProperties : public UInteracti
 public:
 
 	/** For each material ID, the source texture that will be resampled in that material's region*/
-	UPROPERTY(EditAnywhere, EditFixedSize, Category = MultiTexture, meta = (DisplayName = "Material IDs",
+	UPROPERTY(EditAnywhere, EditFixedSize, Category = MultiTextureOutput, meta = (DisplayName = "Material IDs",
 		TransientToolProperty, EditFixedOrder))
 	TArray<TObjectPtr<UTexture2D>> MaterialIDSourceTextures;
 
-	/** UV channel to use for the source mesh textures */
-	UPROPERTY(EditAnywhere, Category = MultiTexture, meta = (DisplayName = "Source Texture UV Channel"))
-	int32 UVLayer = 0;
+	/** UV channel to use for the source mesh texture */
+	UPROPERTY(EditAnywhere, Category = MultiTextureOutput, meta = (DisplayName = "Source Texture UV Channel",
+		GetOptions = GetUVLayerNamesFunc, NoResetToDefault))
+	FString UVLayer;
+
+	UFUNCTION()
+	const TArray<FString>& GetUVLayerNamesFunc() const
+	{
+		return UVLayerNamesList;
+	}
+
+	UPROPERTY(meta = (TransientToolProperty))
+	TArray<FString> UVLayerNamesList;
 
 	/** The set of all source textures from all input materials */
-	UPROPERTY(VisibleAnywhere, BlueprintReadOnly, Category = MultiTexture, meta = (DisplayName = "Source Textures",
+	UPROPERTY(VisibleAnywhere, BlueprintReadOnly, Category = MultiTextureOutput, meta = (DisplayName = "Source Textures",
 		TransientToolProperty))
 	TArray<TObjectPtr<UTexture2D>> AllSourceTextures;
 
diff --git a/Engine/Plugins/Experimental/MeshModelingToolsetExp/Source/MeshModelingToolsExp/Public/BakeMultiMeshAttributeMapsTool.h b/Engine/Plugins/Experimental/MeshModelingToolsetExp/Source/MeshModelingToolsExp/Public/BakeMultiMeshAttributeMapsTool.h
--- a/Engine/Plugins/Experimental/MeshModelingToolsetExp/Source/MeshModelingToolsExp/Public/BakeMultiMeshAttributeMapsTool.h
+++ b/Engine/Plugins/Experimental/MeshModelingToolsetExp/Source/MeshModelingToolsExp/Public/BakeMultiMeshAttributeMapsTool.h
@@ -91,7 +91,7 @@ struct MESHMODELINGTOOLSEXP_API FBakeMultiMeshDetailProperties
 
 	/** UV channel to use for the source mesh color texture */
 	UPROPERTY(EditAnywhere, Category = BakeSources, meta = (TransientToolProperty, DisplayName = "Source Texture UV Channel",
-		EditCondition="SourceTexture != nullptr"))
+		EditCondition="SourceTexture != nullptr", ClampMin=0, ClampMax=7))
 	int32 SourceTextureUVLayer = 0;
 };
 
@@ -112,7 +112,7 @@ public:
 	FString TargetUVLayer;
 
 	/** Source meshes and textures to sample from */
-	UPROPERTY(EditAnywhere, EditFixedSize, Category = BakeInput, meta = (TransientToolProperty, EditFixedOrder))
+	UPROPERTY(EditAnywhere, EditFixedSize, Category = BakeInput, meta = (TransientToolProperty, EditFixedOrder, NoResetToDefault))
 	TArray<FBakeMultiMeshDetailProperties> SourceMeshes;
 
 	/** Maximum allowed distance for the projection from target mesh to source mesh for the sample to be considered valid.
diff --git a/Engine/Plugins/Runtime/MeshModelingToolset/Source/ModelingComponents/Private/Scene/MeshSceneAdapter.cpp b/Engine/Plugins/Runtime/MeshModelingToolset/Source/ModelingComponents/Private/Scene/MeshSceneAdapter.cpp
--- a/Engine/Plugins/Runtime/MeshModelingToolset/Source/ModelingComponents/Private/Scene/MeshSceneAdapter.cpp
+++ b/Engine/Plugins/Runtime/MeshModelingToolset/Source/ModelingComponents/Private/Scene/MeshSceneAdapter.cpp
@@ -693,7 +693,7 @@ public:
 
 	virtual bool HasUVs(const int UVLayer = 0) const override
 	{
-		return SourceMesh ? Adapter->HasUVs() : false;
+		return SourceMesh ? Adapter->HasUVs(UVLayer) : false;
 	}
 
 	virtual FVector3d TriBaryInterpolatePoint(int32 TriId, const FVector3d& BaryCoords) const override
@@ -723,7 +723,7 @@ public:
 
 	virtual bool TriBaryInterpolateUV(const int32 TriId, const FVector3d& BaryCoords, const int UVLayer, FVector2f& UVOut) const override
 	{
-		if (SourceMesh && Adapter->HasUVs())
+		if (SourceMesh && Adapter->HasUVs(UVLayer))
 		{
 			FVector2f UV0, UV1, UV2;
 			Adapter->GetTriUVs(TriId, UVLayer, UV0, UV1, UV2);
diff --git a/Engine/Source/Runtime/MeshConversion/Public/MeshDescriptionAdapter.h b/Engine/Source/Runtime/MeshConversion/Public/MeshDescriptionAdapter.h
--- a/Engine/Source/Runtime/MeshConversion/Public/MeshDescriptionAdapter.h
+++ b/Engine/Source/Runtime/MeshConversion/Public/MeshDescriptionAdapter.h
@@ -147,9 +147,9 @@ public:
 		N2 = GetNormal(TriangleVertexInstanceIndices[TriId*3+2]);
 	}
 
-	inline bool HasUVs() const
+	inline bool HasUVs(const int32 UVLayer=0) const
 	{
-		return VertexInstanceUVs.IsValid();
+		return VertexInstanceUVs.IsValid() && Mesh && UVLayer >= 0 && UVLayer < Mesh->GetNumUVElementChannels();
 	}
 	inline bool IsUV(const int32 UVId) const
 	{
