ModelingTools: add Duplicate button to TriSel Tool, does Separate but without the deleting
#rnx

diff --git a/Engine/Plugins/Experimental/MeshModelingToolsetExp/Source/MeshModelingToolsExp/Private/MeshSelectionTool.cpp b/Engine/Plugins/Experimental/MeshModelingToolsetExp/Source/MeshModelingToolsExp/Private/MeshSelectionTool.cpp
--- a/Engine/Plugins/Experimental/MeshModelingToolsetExp/Source/MeshModelingToolsExp/Private/MeshSelectionTool.cpp
+++ b/Engine/Plugins/Experimental/MeshModelingToolsetExp/Source/MeshModelingToolsExp/Private/MeshSelectionTool.cpp
@@ -1019,7 +1019,11 @@ void UMeshSelectionTool::ApplyAction(EMeshSelectionToolActions ActionType)
 			break;
 
 		case EMeshSelectionToolActions::SeparateSelected:
-			SeparateSelectedTriangles();
+			SeparateSelectedTriangles(true);
+			break;
+
+		case EMeshSelectionToolActions::DuplicateSelected:
+			SeparateSelectedTriangles(false);
 			break;
 
 		case EMeshSelectionToolActions::FlipSelected:
@@ -1450,7 +1454,7 @@ void UMeshSelectionTool::DisconnectSelectedTriangles()
 
 
 
-void UMeshSelectionTool::SeparateSelectedTriangles()
+void UMeshSelectionTool::SeparateSelectedTriangles(bool bDeleteSelected)
 {
 	check(SelectionType == EMeshSelectionElementType::Face);
 	TArray<int32> SelectedFaces = Selection->GetElements(EMeshSelectionElementType::Face);
@@ -1465,7 +1469,6 @@ void UMeshSelectionTool::SeparateSelectedTriangles()
 		return;		// don't separate entire mesh
 	}
 
-
 	// extract copy of triangles
 	FDynamicMesh3 SeparatedMesh;
 	SeparatedMesh.EnableAttributes();
@@ -1476,16 +1479,20 @@ void UMeshSelectionTool::SeparateSelectedTriangles()
 
 	// emit new asset
 	UE::Geometry::FTransform3d Transform(PreviewMesh->GetTransform());
-	GetToolManager()->BeginUndoTransaction(LOCTEXT("MeshSelectionToolSeparate", "Separate"));
+	GetToolManager()->BeginUndoTransaction(
+		(bDeleteSelected) ? LOCTEXT("MeshSelectionToolSeparate", "Separate") : LOCTEXT("MeshSelectionToolDuplicate", "Duplicate") );
 
 	// build array of materials from the original
 	FComponentMaterialSet MaterialSet = UE::ToolTarget::GetMaterialSet(Target);
 	TArray<UMaterialInterface*> Materials = MaterialSet.Materials;
 
+	AActor* TargetActor = UE::ToolTarget::GetTargetActor(Target);
+	FString AssetName = TargetActor->GetActorNameOrLabel();
+
 	FCreateMeshObjectParams NewMeshObjectParams;
 	NewMeshObjectParams.TargetWorld = TargetWorld;
 	NewMeshObjectParams.Transform = (FTransform)Transform;
-	NewMeshObjectParams.BaseName = TEXT("Submesh");
+	NewMeshObjectParams.BaseName = (AssetName.IsEmpty()) ? TEXT("Submesh") : (AssetName + TEXT("_Submesh"));
 	NewMeshObjectParams.Materials = Materials;
 	NewMeshObjectParams.SetMesh(&SeparatedMesh);
 	FCreateMeshObjectResult Result = UE::Modeling::CreateMeshObject(GetToolManager(), MoveTemp(NewMeshObjectParams));
@@ -1497,7 +1504,14 @@ void UMeshSelectionTool::SeparateSelectedTriangles()
 	GetToolManager()->EndUndoTransaction();
 
 	// delete selected triangles from this mesh
-	DeleteSelectedTriangles();
+	if (bDeleteSelected)
+	{
+		DeleteSelectedTriangles();
+	}
+
+	// Currently have to mark the mesh as 'modified' so we can Accept, because if we Cancel,
+	// Actor created by duplicate operation will be rolled back
+	bHaveModifiedMesh = true;		
 }
 
 
diff --git a/Engine/Plugins/Experimental/MeshModelingToolsetExp/Source/MeshModelingToolsExp/Public/MeshSelectionTool.h b/Engine/Plugins/Experimental/MeshModelingToolsetExp/Source/MeshModelingToolsExp/Public/MeshSelectionTool.h
--- a/Engine/Plugins/Experimental/MeshModelingToolsetExp/Source/MeshModelingToolsExp/Public/MeshSelectionTool.h
+++ b/Engine/Plugins/Experimental/MeshModelingToolsetExp/Source/MeshModelingToolsExp/Public/MeshSelectionTool.h
@@ -52,6 +52,7 @@ enum class EMeshSelectionToolActions
 	DeleteSelected,
 	DisconnectSelected,
 	SeparateSelected,
+	DuplicateSelected,
 	FlipSelected,
 	CreateGroup,
 
@@ -185,13 +186,20 @@ public:
 		PostAction(EMeshSelectionToolActions::CreateGroup);
 	}
 
-	/** Delete the selected triangles from the active Mesh Object and create a new Mesh containing those triangles */
+	/** Delete the selected triangles from the active Mesh Object and create a new Mesh Object containing those triangles */
 	UFUNCTION(CallInEditor, Category = MeshEdits, meta = (DisplayName = "Separate", DisplayPriority = 10))
 	void Separate() 
 	{
 		PostAction(EMeshSelectionToolActions::SeparateSelected);
 	}
 
+	/** Create a new Mesh Object containing the selected triangles */
+	UFUNCTION(CallInEditor, Category = MeshEdits, meta = (DisplayName = "Duplicate", DisplayPriority = 11))
+	void Duplicate() 
+	{
+		PostAction(EMeshSelectionToolActions::DuplicateSelected);
+	}
+
 };
 
 
@@ -418,7 +426,7 @@ protected:
 
 	void DeleteSelectedTriangles();
 	void DisconnectSelectedTriangles(); // disconnects edges between selected and not-selected triangles; keeps all triangles in the same mesh
-	void SeparateSelectedTriangles(); // separates out selected triangles to a new mesh, removing them from the current mesh
+	void SeparateSelectedTriangles(bool bDeleteSelected); // separates out selected triangles to a new mesh, optionally removing them from the current mesh
 	void FlipSelectedTriangles();
 	void AssignNewGroupToSelectedTriangles();
 
