Expose component/asset physics data via a ToolTarget interface

ToolsFramework:
- add PhysicsDataSource ToolTarget Interface, that exposes UBodySetup and CollisionDataProvider interfaces
ModelingComponents:
- Add tooltarget accessor/utility functions in UE::ToolTarget namespace.
- Implement PhysicsDataSource for StaticMeshComponentToolTarget and DynamicMeshCompnentToolTarget.
- Update FPhysicsDataCollection::InitializeFromComponent() to support DynamicMeshComponent and BrushComponent.
ModelingTools:
- update ExtractCollisionGeometryTool and PhysicsInspectorTool

#rnx

diff --git a/Engine/Plugins/Experimental/MeshModelingToolsetExp/Source/MeshModelingToolsExp/Private/Physics/ExtractCollisionGeometryTool.cpp b/Engine/Plugins/Experimental/MeshModelingToolsetExp/Source/MeshModelingToolsExp/Private/Physics/ExtractCollisionGeometryTool.cpp
--- a/Engine/Plugins/Experimental/MeshModelingToolsetExp/Source/MeshModelingToolsExp/Private/Physics/ExtractCollisionGeometryTool.cpp
+++ b/Engine/Plugins/Experimental/MeshModelingToolsetExp/Source/MeshModelingToolsExp/Private/Physics/ExtractCollisionGeometryTool.cpp
@@ -23,14 +23,11 @@
 #include "Physics/CollisionGeometryVisualization.h"
 
 // physics data
-#include "Engine/Classes/Engine/StaticMesh.h"
-#include "Engine/Classes/Components/StaticMeshComponent.h"
-#include "Engine/Classes/PhysicsEngine/BodySetup.h"
-#include "Engine/Classes/PhysicsEngine/AggregateGeom.h"
+#include "PhysicsEngine/BodySetup.h"
+#include "PhysicsEngine/AggregateGeom.h"
 
-#include "TargetInterfaces/MeshDescriptionProvider.h"
 #include "TargetInterfaces/PrimitiveComponentBackedTarget.h"
-#include "TargetInterfaces/StaticMeshBackedTarget.h"
+#include "TargetInterfaces/PhysicsDataSource.h"
 #include "ModelingToolTargetUtil.h"
 
 using namespace UE::Geometry;
@@ -40,9 +37,8 @@ using namespace UE::Geometry;
 const FToolTargetTypeRequirements& UExtractCollisionGeometryToolBuilder::GetTargetRequirements() const
 {
 	static FToolTargetTypeRequirements TypeRequirements({
-		UMeshDescriptionProvider::StaticClass(),
 		UPrimitiveComponentBackedTarget::StaticClass(),
-		UStaticMeshBackedTarget::StaticClass()
+		UPhysicsDataSource::StaticClass()
 		});
 	return TypeRequirements;
 }
@@ -94,15 +90,11 @@ void UExtractCollisionGeometryTool::Setup()
 	VizSettings->WatchProperty(VizSettings->LineThickness, [this](float NewValue) { bVisualizationDirty = true; });
 	VizSettings->WatchProperty(VizSettings->Color, [this](FColor NewValue) { bVisualizationDirty = true; });
 
-
-	UPrimitiveComponent* TargetComponent = UE::ToolTarget::GetTargetComponent(Target);
-	UStaticMeshComponent* StaticMeshComponent = (TargetComponent) ? Cast<UStaticMeshComponent>(TargetComponent) : nullptr;
-	TObjectPtr<UStaticMesh> StaticMesh = (StaticMeshComponent) ? StaticMeshComponent->GetStaticMesh() : nullptr;
-	UBodySetup* BodySetup = (StaticMesh) ? StaticMesh->GetBodySetup() : nullptr;
-	if (ensure(BodySetup))
+	UBodySetup* BodySetup = UE::ToolTarget::GetPhysicsBodySetup(Target);
+	if (BodySetup)
 	{
 		PhysicsInfo = MakeShared<FPhysicsDataCollection>();
-		PhysicsInfo->InitializeFromComponent(StaticMeshComponent, true);
+		PhysicsInfo->InitializeFromComponent( UE::ToolTarget::GetTargetComponent(Target), true);
 
 		PreviewElements = NewObject<UPreviewGeometry>(this);
 		FTransform TargetTransform = (FTransform)UE::ToolTarget::GetLocalToWorldTransform(Target);
@@ -172,7 +164,7 @@ void UExtractCollisionGeometryTool::Shutdown(EToolShutdownType ShutdownType)
 		{
 			for ( int32 k = 0; k < NumParts; ++k)
 			{
-				FDynamicMesh3& MeshPart = CurrentMeshParts[k];
+				FDynamicMesh3& MeshPart = *CurrentMeshParts[k];
 				FAxisAlignedBox3d Bounds = MeshPart.GetBounds();
 				MeshTransforms::Translate(MeshPart, -Bounds.Center());
 				FTransform3d CenterTransform = ActorTransform;
@@ -269,7 +261,7 @@ void UExtractCollisionGeometryTool::RecalculateMesh_Simple()
 
 		FMeshIndexMappings Mappings;
 		Editor.AppendMesh(&SphereMesh, Mappings);
-		CurrentMeshParts.Add(MoveTemp(SphereMesh));
+		CurrentMeshParts.Add(MakeShared<FDynamicMesh3>(MoveTemp(SphereMesh)));
 	}
 
 
@@ -286,7 +278,7 @@ void UExtractCollisionGeometryTool::RecalculateMesh_Simple()
 
 		FMeshIndexMappings Mappings;
 		Editor.AppendMesh(&BoxMesh, Mappings);
-		CurrentMeshParts.Add(MoveTemp(BoxMesh));
+		CurrentMeshParts.Add(MakeShared<FDynamicMesh3>(MoveTemp(BoxMesh)));
 	}
 
 
@@ -308,7 +300,7 @@ void UExtractCollisionGeometryTool::RecalculateMesh_Simple()
 
 		FMeshIndexMappings Mappings;
 		Editor.AppendMesh(&CapsuleMesh, Mappings);
-		CurrentMeshParts.Add(MoveTemp(CapsuleMesh));
+		CurrentMeshParts.Add(MakeShared<FDynamicMesh3>(MoveTemp(CapsuleMesh)));
 	}
 
 
@@ -334,11 +326,12 @@ void UExtractCollisionGeometryTool::RecalculateMesh_Simple()
 		UVEditor.SetPerTriangleUVs();
 		FMeshIndexMappings Mappings;
 		Editor.AppendMesh(&ConvexMesh, Mappings);
-		CurrentMeshParts.Add(MoveTemp(ConvexMesh));
+		CurrentMeshParts.Add(MakeShared<FDynamicMesh3>(MoveTemp(ConvexMesh)));
 	}
 
-	for (FDynamicMesh3& MeshPart : CurrentMeshParts)
+	for ( int32 k = 0; k < CurrentMeshParts.Num(); ++k)
 	{
+		FDynamicMesh3& MeshPart = *CurrentMeshParts[k];
 		FMeshNormals::InitializeMeshToPerTriangleNormals(&MeshPart);
 	}
 	FMeshNormals::InitializeMeshToPerTriangleNormals(&CurrentMesh);
@@ -361,9 +354,7 @@ void UExtractCollisionGeometryTool::RecalculateMesh_Complex()
 
 	bool bMeshErrors = false;
 
-	UStaticMeshComponent* StaticMeshComponent = Cast<UStaticMeshComponent>(UE::ToolTarget::GetTargetComponent(Target));
-	TObjectPtr<UStaticMesh> StaticMesh = (StaticMeshComponent) ? StaticMeshComponent->GetStaticMesh() : nullptr;
-	IInterface_CollisionDataProvider* CollisionProvider = (StaticMesh) ? Cast<IInterface_CollisionDataProvider>(StaticMesh) : nullptr;
+	IInterface_CollisionDataProvider* CollisionProvider = UE::ToolTarget::GetPhysicsCollisionDataProvider(Target);
 	if (CollisionProvider && CollisionProvider->ContainsPhysicsTriMeshData(true))
 	{
 		FTriMeshCollisionData CollisionData;
diff --git a/Engine/Plugins/Experimental/MeshModelingToolsetExp/Source/MeshModelingToolsExp/Private/Physics/PhysicsInspectorTool.cpp b/Engine/Plugins/Experimental/MeshModelingToolsetExp/Source/MeshModelingToolsExp/Private/Physics/PhysicsInspectorTool.cpp
--- a/Engine/Plugins/Experimental/MeshModelingToolsetExp/Source/MeshModelingToolsExp/Private/Physics/PhysicsInspectorTool.cpp
+++ b/Engine/Plugins/Experimental/MeshModelingToolsetExp/Source/MeshModelingToolsExp/Private/Physics/PhysicsInspectorTool.cpp
@@ -11,13 +11,11 @@
 #include "Physics/CollisionGeometryVisualization.h"
 
 // physics data
-#include "Engine/Classes/Engine/StaticMesh.h"
-#include "Engine/Classes/Components/StaticMeshComponent.h"
 #include "Engine/Classes/PhysicsEngine/BodySetup.h"
 #include "Engine/Classes/PhysicsEngine/AggregateGeom.h"
 
 #include "TargetInterfaces/PrimitiveComponentBackedTarget.h"
-#include "TargetInterfaces/StaticMeshBackedTarget.h"
+#include "TargetInterfaces/PhysicsDataSource.h"
 #include "ToolTargetManager.h"
 
 using namespace UE::Geometry;
@@ -29,7 +27,7 @@ const FToolTargetTypeRequirements& UPhysicsInspectorToolBuilder::GetTargetRequir
 {
 	static FToolTargetTypeRequirements TypeRequirements({
 		UPrimitiveComponentBackedTarget::StaticClass(),
-		UStaticMeshBackedTarget::StaticClass()
+		UPhysicsDataSource::StaticClass()
 		});
 	return TypeRequirements;
 }
@@ -54,14 +52,11 @@ void UPhysicsInspectorTool::Setup()
 
 	for (int32 ComponentIdx = 0; ComponentIdx < Targets.Num(); ComponentIdx++)
 	{
-		const UStaticMeshComponent* Component = CastChecked<UStaticMeshComponent>(UE::ToolTarget::GetTargetComponent(Targets[ComponentIdx]));
-		const UStaticMesh* StaticMesh = Component->GetStaticMesh();
-		if (ensure(StaticMesh && StaticMesh->GetBodySetup()))
+		UBodySetup* BodySetup = UE::ToolTarget::GetPhysicsBodySetup(Targets[ComponentIdx]);
+		if (BodySetup)
 		{
 			TSharedPtr<FPhysicsDataCollection> PhysicsData = MakeShared<FPhysicsDataCollection>();
-			PhysicsData->SourceComponent = Component;
-			PhysicsData->BodySetup = StaticMesh->GetBodySetup();
-			PhysicsData->AggGeom = StaticMesh->GetBodySetup()->AggGeom;
+			PhysicsData->InitializeFromComponent( UE::ToolTarget::GetTargetComponent(Targets[ComponentIdx]), true);
 
 			PhysicsInfos.Add(PhysicsData);
 
@@ -82,7 +77,7 @@ void UPhysicsInspectorTool::Setup()
 
 	SetToolDisplayName(LOCTEXT("ToolName", "Physics Inspector"));
 	GetToolManager()->DisplayMessage(
-		LOCTEXT("OnStartTool", "Inspect Physics data for the seleced Static Meshes"),
+		LOCTEXT("OnStartTool", "Inspect Physics data for the selected Objects"),
 		EToolMessageLevel::UserNotification);
 }
 
diff --git a/Engine/Plugins/Experimental/MeshModelingToolsetExp/Source/MeshModelingToolsExp/Public/Physics/ExtractCollisionGeometryTool.h b/Engine/Plugins/Experimental/MeshModelingToolsetExp/Source/MeshModelingToolsExp/Public/Physics/ExtractCollisionGeometryTool.h
--- a/Engine/Plugins/Experimental/MeshModelingToolsetExp/Source/MeshModelingToolsExp/Public/Physics/ExtractCollisionGeometryTool.h
+++ b/Engine/Plugins/Experimental/MeshModelingToolsetExp/Source/MeshModelingToolsExp/Public/Physics/ExtractCollisionGeometryTool.h
@@ -105,7 +105,7 @@ protected:
 	// these are TSharedPtr because TPimplPtr cannot currently be added to a TArray?
 	TSharedPtr<FPhysicsDataCollection> PhysicsInfo;
 
-	TArray<UE::Geometry::FDynamicMesh3> CurrentMeshParts;
+	TArray<TSharedPtr<UE::Geometry::FDynamicMesh3>> CurrentMeshParts;
 	UE::Geometry::FDynamicMesh3 CurrentMesh;
 	bool bResultValid = false;
 	void RecalculateMesh_Simple();
diff --git a/Engine/Plugins/Runtime/MeshModelingToolset/Source/ModelingComponents/Private/ModelingToolTargetUtil.cpp b/Engine/Plugins/Runtime/MeshModelingToolset/Source/ModelingComponents/Private/ModelingToolTargetUtil.cpp
--- a/Engine/Plugins/Runtime/MeshModelingToolset/Source/ModelingComponents/Private/ModelingToolTargetUtil.cpp
+++ b/Engine/Plugins/Runtime/MeshModelingToolset/Source/ModelingComponents/Private/ModelingToolTargetUtil.cpp
@@ -11,6 +11,7 @@
 #include "TargetInterfaces/PrimitiveComponentBackedTarget.h"
 #include "TargetInterfaces/AssetBackedTarget.h"
 #include "TargetInterfaces/DynamicMeshSource.h"
+#include "TargetInterfaces/PhysicsDataSource.h"
 
 #include "ModelingObjectsCreationAPI.h"
 
@@ -421,5 +422,26 @@ bool UE::ToolTarget::ConfigureCreateMeshObjectParams(UToolTarget* SourceTarget,
 }
 
 
+UBodySetup* UE::ToolTarget::GetPhysicsBodySetup(UToolTarget* Target)
+{
+	if (IPhysicsDataSource* PhysicsSource = Cast<IPhysicsDataSource>(Target))
+	{
+		return PhysicsSource->GetBodySetup();
+	}
+	return nullptr;
+}
+
+IInterface_CollisionDataProvider* UE::ToolTarget::GetPhysicsCollisionDataProvider(UToolTarget* Target)
+{
+	if (IPhysicsDataSource* PhysicsSource = Cast<IPhysicsDataSource>(Target))
+	{
+		return PhysicsSource->GetComplexCollisionProvider();
+	}
+	return nullptr;
+}
+
+
+
+
 
 #undef LOCTEXT_NAMESPACE
\ No newline at end of file
diff --git a/Engine/Plugins/Runtime/MeshModelingToolset/Source/ModelingComponents/Private/Physics/PhysicsDataCollection.cpp b/Engine/Plugins/Runtime/MeshModelingToolset/Source/ModelingComponents/Private/Physics/PhysicsDataCollection.cpp
--- a/Engine/Plugins/Runtime/MeshModelingToolset/Source/ModelingComponents/Private/Physics/PhysicsDataCollection.cpp
+++ b/Engine/Plugins/Runtime/MeshModelingToolset/Source/ModelingComponents/Private/Physics/PhysicsDataCollection.cpp
@@ -5,30 +5,40 @@
 
 #include "Engine/StaticMesh.h"
 #include "Components/StaticMeshComponent.h"
+#include "Components/BrushComponent.h"
+#include "Components/DynamicMeshComponent.h"
 #include "PhysicsEngine/BodySetup.h"
 
 using namespace UE::Geometry;
 
 void FPhysicsDataCollection::InitializeFromComponent(const UActorComponent* Component, bool bInitializeAggGeom)
 {
-	const UStaticMeshComponent* StaticMeshComponent = Cast<UStaticMeshComponent>(Component);
-	if (ensure(StaticMeshComponent))
+	BodySetup = nullptr;
+
+	if ( const UStaticMeshComponent* StaticMeshComponent = Cast<UStaticMeshComponent>(Component) )
 	{
-		const UStaticMesh* StaticMesh = StaticMeshComponent->GetStaticMesh();
-		if (ensure(StaticMesh))
+		if (UStaticMesh* StaticMesh = StaticMeshComponent->GetStaticMesh())
 		{
 			SourceComponent = StaticMeshComponent;
 			SourceStaticMesh = StaticMesh;
 			BodySetup = StaticMesh->GetBodySetup();
+		}
+	}
+	else if (const UBrushComponent* BrushComponent = Cast<UBrushComponent>(Component))
+	{
+		BodySetup = BrushComponent->BrushBodySetup;
+	}
+	else if (const UDynamicMeshComponent* DynamicMeshComponent = Cast<UDynamicMeshComponent>(Component))
+	{
+		BodySetup = DynamicMeshComponent->GetBodySetup();
+	}
 
-			ExternalScale3D = FVector(1.f, 1.f, 1.f);
+	ExternalScale3D = FVector(1.f, 1.f, 1.f);
 
-			if (bInitializeAggGeom)
-			{
-				AggGeom = BodySetup->AggGeom;
-				UE::Geometry::GetShapeSet(AggGeom, Geometry);
-			}
-		}
+	if (bInitializeAggGeom && BodySetup != nullptr)
+	{
+		AggGeom = BodySetup->AggGeom;
+		UE::Geometry::GetShapeSet(AggGeom, Geometry);
 	}
 }
 
diff --git a/Engine/Plugins/Runtime/MeshModelingToolset/Source/ModelingComponents/Public/ModelingToolTargetUtil.h b/Engine/Plugins/Runtime/MeshModelingToolset/Source/ModelingComponents/Public/ModelingToolTargetUtil.h
--- a/Engine/Plugins/Runtime/MeshModelingToolset/Source/ModelingComponents/Public/ModelingToolTargetUtil.h
+++ b/Engine/Plugins/Runtime/MeshModelingToolset/Source/ModelingComponents/Public/ModelingToolTargetUtil.h
@@ -10,6 +10,8 @@
 class UToolTarget;
 class UPrimitiveComponent;
 class AActor;
+class UBodySetup;
+class IInterface_CollisionDataProvider;
 struct FMeshDescription;
 struct FCreateMeshObjectParams;
 class IPersistentDynamicMeshSource;
@@ -172,6 +174,17 @@ namespace Internal
 		const UE::Geometry::FDynamicMesh3& UpdatedMesh, bool bHaveModifiedTopology);
 }
 
+/**
+ * @return the Physics UBodySetup for the given ToolTarget, or nullptr if it does not exist
+ */
+MODELINGCOMPONENTS_API UBodySetup* GetPhysicsBodySetup(UToolTarget* Target);
+
+/**
+* @return the Physics CollisionDataProvider (ie Complex Collision source) for the given ToolTarget, or nullptr if it does not exist
+*/
+MODELINGCOMPONENTS_API IInterface_CollisionDataProvider* GetPhysicsCollisionDataProvider(UToolTarget* Target);
+
+
 
 }  // end namespace ToolTarget
 }  // end namespace UE
\ No newline at end of file
diff --git a/Engine/Plugins/Runtime/MeshModelingToolset/Source/ModelingComponentsEditorOnly/Private/ToolTargets/StaticMeshComponentToolTarget.cpp b/Engine/Plugins/Runtime/MeshModelingToolset/Source/ModelingComponentsEditorOnly/Private/ToolTargets/StaticMeshComponentToolTarget.cpp
--- a/Engine/Plugins/Runtime/MeshModelingToolset/Source/ModelingComponentsEditorOnly/Private/ToolTargets/StaticMeshComponentToolTarget.cpp
+++ b/Engine/Plugins/Runtime/MeshModelingToolset/Source/ModelingComponentsEditorOnly/Private/ToolTargets/StaticMeshComponentToolTarget.cpp
@@ -192,6 +192,28 @@ UStaticMesh* UStaticMeshComponentToolTarget::GetStaticMesh() const
 }
 
 
+UBodySetup* UStaticMeshComponentToolTarget::GetBodySetup() const
+{
+	UStaticMesh* StaticMesh = GetStaticMesh();
+	if (StaticMesh)
+	{
+		return StaticMesh->GetBodySetup();
+	}
+	return nullptr;
+}
+
+
+IInterface_CollisionDataProvider* UStaticMeshComponentToolTarget::GetComplexCollisionProvider() const
+{
+	UStaticMesh* StaticMesh = GetStaticMesh();
+	if (StaticMesh)
+	{
+		return Cast<IInterface_CollisionDataProvider>(StaticMesh);
+	}
+	return nullptr;
+}
+
+
 // Factory
 
 bool UStaticMeshComponentToolTargetFactory::CanBuildTarget(UObject* SourceObject, const FToolTargetTypeRequirements& Requirements) const
diff --git a/Engine/Plugins/Runtime/MeshModelingToolset/Source/ModelingComponentsEditorOnly/Public/ToolTargets/StaticMeshComponentToolTarget.h b/Engine/Plugins/Runtime/MeshModelingToolset/Source/ModelingComponentsEditorOnly/Public/ToolTargets/StaticMeshComponentToolTarget.h
--- a/Engine/Plugins/Runtime/MeshModelingToolset/Source/ModelingComponentsEditorOnly/Public/ToolTargets/StaticMeshComponentToolTarget.h
+++ b/Engine/Plugins/Runtime/MeshModelingToolset/Source/ModelingComponentsEditorOnly/Public/ToolTargets/StaticMeshComponentToolTarget.h
@@ -10,6 +10,7 @@
 #include "TargetInterfaces/MeshDescriptionCommitter.h"
 #include "TargetInterfaces/MeshDescriptionProvider.h"
 #include "TargetInterfaces/StaticMeshBackedTarget.h"
+#include "TargetInterfaces/PhysicsDataSource.h"
 #include "ToolTargets/PrimitiveComponentToolTarget.h"
 #include "ComponentSourceInterfaces.h"  // for EStaticMeshEditingLOD
 
@@ -24,7 +25,7 @@ class UStaticMesh;
 UCLASS(Transient)
 class MODELINGCOMPONENTSEDITORONLY_API UStaticMeshComponentToolTarget : public UPrimitiveComponentToolTarget,
 	public IMeshDescriptionCommitter, public IMeshDescriptionProvider, public IMaterialProvider, public IStaticMeshBackedTarget,
-	public IDynamicMeshProvider, public IDynamicMeshCommitter
+	public IDynamicMeshProvider, public IDynamicMeshCommitter, public IPhysicsDataSource
 {
 	GENERATED_BODY()
 
@@ -67,6 +68,10 @@ public:
 	// IStaticMeshBackedTarget
 	UStaticMesh* GetStaticMesh() const override;
 
+	// IPhysicsDataSource
+	virtual UBodySetup* GetBodySetup() const override;
+	virtual IInterface_CollisionDataProvider* GetComplexCollisionProvider() const override;
+
 	// Rest provided by parent class
 
 public:
diff --git a/Engine/Source/Runtime/GeometryFramework/Public/Components/DynamicMeshComponent.h b/Engine/Source/Runtime/GeometryFramework/Public/Components/DynamicMeshComponent.h
--- a/Engine/Source/Runtime/GeometryFramework/Public/Components/DynamicMeshComponent.h
+++ b/Engine/Source/Runtime/GeometryFramework/Public/Components/DynamicMeshComponent.h
@@ -482,7 +482,11 @@ public:
 	virtual bool ContainsPhysicsTriMeshData(bool InUseAllTriData) const override;
 	virtual bool WantsNegXTriMesh() override;
 
+	/** @return current BodySetup for this Component, or nullptr if it does not exist */
+	virtual const UBodySetup* GetBodySetup() const { return MeshBodySetup; }
+	/** @return BodySetup for this Component. A new BodySetup will be created if one does not exist. */
 	virtual UBodySetup* GetBodySetup() override;
+	/** @return Set new BodySetup for this Component. */
 	virtual void SetBodySetup(UBodySetup* NewSetup);
 
 	/**
diff --git a/Engine/Source/Runtime/InteractiveToolsFramework/Public/TargetInterfaces/PhysicsDataSource.h b/Engine/Source/Runtime/InteractiveToolsFramework/Public/TargetInterfaces/PhysicsDataSource.h
new file mode 100644
--- /dev/null
+++ b/Engine/Source/Runtime/InteractiveToolsFramework/Public/TargetInterfaces/PhysicsDataSource.h
@@ -0,0 +1,39 @@
+// Copyright Epic Games, Inc. All Rights Reserved.
+
+#pragma once
+
+#include "CoreMinimal.h"
+#include "UObject/Interface.h"
+
+#include "PhysicsDataSource.generated.h"
+
+class UBodySetup;
+class IInterface_CollisionDataProvider;
+
+UINTERFACE()
+class INTERACTIVETOOLSFRAMEWORK_API UPhysicsDataSource : public UInterface
+{
+	GENERATED_BODY()
+};
+
+/**
+ * IPhysicsDataSource is a ToolTarget Interface that provides read/write access to physics-related data structures.
+ */
+class INTERACTIVETOOLSFRAMEWORK_API IPhysicsDataSource
+{
+	GENERATED_BODY()
+
+public:
+
+	/**
+	 * @return The UBodySetup for this physics data source. If Nullptr, no physics data exists or is available.
+	 */
+	virtual UBodySetup* GetBodySetup() const = 0;
+
+
+	/**
+	* @return The CollisionDataProvider for this physics data source. If Nullptr, no physics data exists or is available.
+	*/
+	virtual IInterface_CollisionDataProvider* GetComplexCollisionProvider() const = 0;
+
+};
