ModelingTools: Reorder BakeTexture/BakeAll results to more closely match the UI order.
Remove EBakeMapType::Occlusion and check AO & BentNormal directly.
Rename output texture suffixes to match UI naming.

#rnx

diff --git a/Engine/Plugins/Experimental/MeshModelingToolsetExp/Source/MeshModelingToolsExp/Private/BakeMeshAttributeMapsTool.cpp b/Engine/Plugins/Experimental/MeshModelingToolsetExp/Source/MeshModelingToolsExp/Private/BakeMeshAttributeMapsTool.cpp
--- a/Engine/Plugins/Experimental/MeshModelingToolsetExp/Source/MeshModelingToolsExp/Private/BakeMeshAttributeMapsTool.cpp
+++ b/Engine/Plugins/Experimental/MeshModelingToolsetExp/Source/MeshModelingToolsExp/Private/BakeMeshAttributeMapsTool.cpp
@@ -138,7 +138,25 @@ public:
 		FMeshBakerDynamicMeshSampler DetailSampler(DetailMesh.Get(), DetailSpatial.Get(), DetailMeshTangents.Get());
 		Baker->SetDetailSampler(&DetailSampler);
 
-		for (const EBakeMapType MapType : ALL_BAKE_MAP_TYPES)
+		// Occlusion evaluator is shared by both AmbientOcclusion and BentNormal.
+		// Only initialize it once. OcclusionType is initialized to None. Callers
+		// must manually update the OcclusionType.
+		auto InitOcclusionEval = [this](TSharedPtr<FMeshOcclusionMapEvaluator>& Eval)
+		{
+			if (!Eval)
+			{
+				Eval = MakeShared<FMeshOcclusionMapEvaluator>();
+				Eval->OcclusionType = EMeshOcclusionMapType::None;
+				Eval->NumOcclusionRays = OcclusionSettings.OcclusionRays;
+				Eval->MaxDistance = OcclusionSettings.MaxDistance;
+				Eval->SpreadAngle = OcclusionSettings.SpreadAngle;
+				Eval->BiasAngleDeg = OcclusionSettings.BiasAngle;
+				Baker->AddEvaluator(Eval);
+			}
+		};
+
+		TSharedPtr<FMeshOcclusionMapEvaluator> OcclusionEval; 
+		for (const EBakeMapType MapType : ENUM_EBAKEMAPTYPE_ALL)
 		{
 			switch (BakeSettings.BakeMapTypes & MapType)
 			{
@@ -150,26 +168,14 @@ public:
 				break;
 			}
 			case EBakeMapType::AmbientOcclusion:
+			{
+				InitOcclusionEval(OcclusionEval);
+				OcclusionEval->OcclusionType |= EMeshOcclusionMapType::AmbientOcclusion; 
+			}
 			case EBakeMapType::BentNormal:
-			case EBakeMapType::Occlusion:
 			{
-				TSharedPtr<FMeshOcclusionMapEvaluator> OcclusionEval = MakeShared<FMeshOcclusionMapEvaluator>();
-				OcclusionEval->OcclusionType = EMeshOcclusionMapType::None;
-				if ((bool)(BakeSettings.BakeMapTypes & EBakeMapType::AmbientOcclusion))
-				{
-					OcclusionEval->OcclusionType |= EMeshOcclusionMapType::AmbientOcclusion;
-				}
-				if ((bool)(BakeSettings.BakeMapTypes & EBakeMapType::BentNormal))
-				{
-					OcclusionEval->OcclusionType |= EMeshOcclusionMapType::BentNormal;
-				}
-				OcclusionEval->NumOcclusionRays = OcclusionSettings.OcclusionRays;
-				OcclusionEval->MaxDistance = OcclusionSettings.MaxDistance;
-				OcclusionEval->SpreadAngle = OcclusionSettings.SpreadAngle;
-				OcclusionEval->BiasAngleDeg = OcclusionSettings.BiasAngle;
-
-				Baker->AddEvaluator(OcclusionEval);
-				break;
+				InitOcclusionEval(OcclusionEval);
+				OcclusionEval->OcclusionType |= EMeshOcclusionMapType::BentNormal;
 			}
 			case EBakeMapType::Curvature:
 			{
@@ -677,7 +683,7 @@ void UBakeMeshAttributeMapsTool::UpdateOnModeChange()
 	SetToolPropertySourceEnabled(TextureSettings, false);
 	SetToolPropertySourceEnabled(MultiTextureSettings, false);
 
-	for (const EBakeMapType MapType : ALL_BAKE_MAP_TYPES)
+	for (const EBakeMapType MapType : ENUM_EBAKEMAPTYPE_ALL)
 	{
 		switch ((EBakeMapType)Settings->MapTypes & MapType)
 		{
@@ -685,7 +691,6 @@ void UBakeMeshAttributeMapsTool::UpdateOnModeChange()
 			break;
 		case EBakeMapType::AmbientOcclusion:
 		case EBakeMapType::BentNormal:
-		case EBakeMapType::Occlusion:
 			SetToolPropertySourceEnabled(OcclusionSettings, true);
 			break;
 		case EBakeMapType::Curvature:
diff --git a/Engine/Plugins/Experimental/MeshModelingToolsetExp/Source/MeshModelingToolsExp/Private/BakeMeshAttributeMapsToolBase.cpp b/Engine/Plugins/Experimental/MeshModelingToolsetExp/Source/MeshModelingToolsExp/Private/BakeMeshAttributeMapsToolBase.cpp
--- a/Engine/Plugins/Experimental/MeshModelingToolsetExp/Source/MeshModelingToolsExp/Private/BakeMeshAttributeMapsToolBase.cpp
+++ b/Engine/Plugins/Experimental/MeshModelingToolsetExp/Source/MeshModelingToolsExp/Private/BakeMeshAttributeMapsToolBase.cpp
@@ -24,7 +24,6 @@
 #include "EngineAnalytics.h"
 
 #include "Sampling/MeshCurvatureMapEvaluator.h"
-#include "Sampling/MeshResampleImageEvaluator.h"
 using namespace UE::Geometry;
 
 #define LOCTEXT_NAMESPACE "UBakeMeshAttributeMapsToolBase"
@@ -409,10 +408,10 @@ void UBakeMeshAttributeMapsToolBase::GetTextureName(const EBakeMapType MapType,
 		checkNoEntry();
 		break;
 	case EBakeMapType::TangentSpaceNormal:
-		TexName = FString::Printf(TEXT("%s_Normals"), *BaseName);
+		TexName = FString::Printf(TEXT("%s_TangentNormal"), *BaseName);
 		break;
 	case EBakeMapType::AmbientOcclusion:
-		TexName = FString::Printf(TEXT("%s_Occlusion"), *BaseName);
+		TexName = FString::Printf(TEXT("%s_AmbientOcclusion"), *BaseName);
 		break;
 	case EBakeMapType::BentNormal:
 		TexName = FString::Printf(TEXT("%s_BentNormal"), *BaseName);
@@ -421,7 +420,7 @@ void UBakeMeshAttributeMapsToolBase::GetTextureName(const EBakeMapType MapType,
 		TexName = FString::Printf(TEXT("%s_Curvature"), *BaseName);
 		break;
 	case EBakeMapType::ObjectSpaceNormal:
-		TexName = FString::Printf(TEXT("%s_NormalsObject"), *BaseName);
+		TexName = FString::Printf(TEXT("%s_ObjectNormal"), *BaseName);
 		break;
 	case EBakeMapType::FaceNormal:
 		TexName = FString::Printf(TEXT("%s_FaceNormal"), *BaseName);
@@ -430,7 +429,7 @@ void UBakeMeshAttributeMapsToolBase::GetTextureName(const EBakeMapType MapType,
 		TexName = FString::Printf(TEXT("%s_MaterialID"), *BaseName);
 		break;
 	case EBakeMapType::VertexColor:
-		TexName = FString::Printf(TEXT("%s_VertexColorID"), *BaseName);
+		TexName = FString::Printf(TEXT("%s_VertexColor"), *BaseName);
 		break;
 	case EBakeMapType::Position:
 		TexName = FString::Printf(TEXT("%s_Position"), *BaseName);
diff --git a/Engine/Plugins/Experimental/MeshModelingToolsetExp/Source/MeshModelingToolsExp/Private/BakeMultiMeshAttributeMapsTool.cpp b/Engine/Plugins/Experimental/MeshModelingToolsetExp/Source/MeshModelingToolsExp/Private/BakeMultiMeshAttributeMapsTool.cpp
--- a/Engine/Plugins/Experimental/MeshModelingToolsetExp/Source/MeshModelingToolsExp/Private/BakeMultiMeshAttributeMapsTool.cpp
+++ b/Engine/Plugins/Experimental/MeshModelingToolsetExp/Source/MeshModelingToolsExp/Private/BakeMultiMeshAttributeMapsTool.cpp
@@ -341,7 +341,7 @@ public:
 		FMeshBakerMeshSceneSampler DetailSampler(DetailMeshScene);
 		Baker->SetDetailSampler(&DetailSampler);
 
-		for (const EBakeMapType MapType : ALL_BAKE_MAP_TYPES)
+		for (const EBakeMapType MapType : ENUM_EBAKEMAPTYPE_ALL)
 		{
 			switch (BakeSettings.BakeMapTypes & MapType)
 			{
diff --git a/Engine/Plugins/Experimental/MeshModelingToolsetExp/Source/MeshModelingToolsExp/Public/BakeMeshAttributeMapsToolBase.h b/Engine/Plugins/Experimental/MeshModelingToolsetExp/Source/MeshModelingToolsExp/Public/BakeMeshAttributeMapsToolBase.h
--- a/Engine/Plugins/Experimental/MeshModelingToolsetExp/Source/MeshModelingToolsExp/Public/BakeMeshAttributeMapsToolBase.h
+++ b/Engine/Plugins/Experimental/MeshModelingToolsetExp/Source/MeshModelingToolsExp/Public/BakeMeshAttributeMapsToolBase.h
@@ -5,11 +5,9 @@
 #include "CoreMinimal.h"
 #include "Engine/Classes/Engine/Texture2D.h"
 #include "Materials/MaterialInstanceDynamic.h"
-#include "MultiSelectionTool.h"
 #include "InteractiveToolManager.h"
 #include "AssetUtils/Texture2DBuilder.h"
 #include "DynamicMesh/DynamicMesh3.h"
-#include "DynamicMesh/DynamicMeshAABBTree3.h"
 #include "Image/ImageDimensions.h"
 #include "Sampling/MeshMapBaker.h"
 #include "ModelingOperators.h"
@@ -317,20 +315,9 @@ void UBakeMeshAttributeMapsToolBase::OnMapTypesUpdated(PropertySet& Properties)
 	// Use the processed bitfield which may contain additional targets
 	// (ex. AO if BentNormal was requested).
 	CachedMaps.Empty();
-	for (EBakeMapType MapType : ALL_BAKE_MAP_TYPES)
+	for (EBakeMapType MapType : ENUM_EBAKEMAPTYPE_ALL)
 	{
-		if (MapType == EBakeMapType::Occlusion)
-		{
-			if ((bool)(BakeMapTypes & EBakeMapType::AmbientOcclusion))
-			{
-				CachedMaps.Add(EBakeMapType::AmbientOcclusion, nullptr);
-			}
-			if ((bool)(BakeMapTypes & EBakeMapType::BentNormal))
-			{
-				CachedMaps.Add(EBakeMapType::BentNormal, nullptr);
-			}
-		}
-		else if( (bool)(BakeMapTypes & MapType) )
+		if( (bool)(BakeMapTypes & MapType) )
 		{
 			CachedMaps.Add(MapType, nullptr);
 		}
diff --git a/Engine/Plugins/Experimental/MeshModelingToolsetExp/Source/MeshModelingToolsExp/Public/BakeMeshAttributeTool.h b/Engine/Plugins/Experimental/MeshModelingToolsetExp/Source/MeshModelingToolsExp/Public/BakeMeshAttributeTool.h
--- a/Engine/Plugins/Experimental/MeshModelingToolsetExp/Source/MeshModelingToolsExp/Public/BakeMeshAttributeTool.h
+++ b/Engine/Plugins/Experimental/MeshModelingToolsetExp/Source/MeshModelingToolsExp/Public/BakeMeshAttributeTool.h
@@ -26,53 +26,52 @@ enum class EBakeMapType
 
 	/* Normals in tangent space */
 	TangentSpaceNormal     = 1 << 0 UMETA(DisplayName = "Tangent Normal"),
-	/* Ambient occlusion sampled across the hemisphere */
-	AmbientOcclusion       = 1 << 1,
-	/* Normals skewed towards the least occluded direction */
-	BentNormal             = 1 << 2,
-	/* Local curvature of the mesh surface */
-	Curvature              = 1 << 3,
-	/* Transfer a given texture */
-	Texture                = 1 << 4,
 	/* Interpolated normals in object space */
-	ObjectSpaceNormal      = 1 << 5 UMETA(DisplayName = "Object Normal"),
+	ObjectSpaceNormal      = 1 << 1 UMETA(DisplayName = "Object Normal"),
 	/* Geometric face normals in object space */
-	FaceNormal             = 1 << 6,
+	FaceNormal             = 1 << 2,
+	/* Normals skewed towards the least occluded direction */
+	BentNormal             = 1 << 3,
 	/* Positions in object space */
-	Position               = 1 << 7,
-	/* Material IDs as unique colors */
-	MaterialID             = 1 << 8 UMETA(DisplayName = "Material ID"),
+	Position               = 1 << 4,
+	/* Local curvature of the mesh surface */
+	Curvature              = 1 << 5,
+	/* Ambient occlusion sampled across the hemisphere */
+	AmbientOcclusion       = 1 << 6,
+	/* Transfer a given texture */
+	Texture                = 1 << 7,
 	/* Transfer a texture per material ID */
-	MultiTexture           = 1 << 9,
+	MultiTexture           = 1 << 8,
 	/* Interpolated vertex colors */
-	VertexColor            = 1 << 10,
+	VertexColor            = 1 << 9,
+	/* Material IDs as unique colors */
+	MaterialID             = 1 << 10 UMETA(DisplayName = "Material ID"),
 
-	Occlusion              = (AmbientOcclusion | BentNormal) UMETA(Hidden),
 	All                    = 0x7FF UMETA(Hidden)
 };
 ENUM_CLASS_FLAGS(EBakeMapType);
 
-
-// Only include the Occlusion bitmask rather than its components
-// (AmbientOcclusion | BentNormal). Since the Occlusion evaluator can
-// evaluate both types in a single pass, only iterating over the Occlusion
-// bitmask gives direct access to both types without the need to
-// externally track if we've handled the Occlusion evaluator in a prior
-// iteration loop.
-static constexpr EBakeMapType ALL_BAKE_MAP_TYPES[] =
+/**
+ * Array of all bake map types. This defines the order in which
+ * the bake map types are iterated in the bake tools. Consequently this
+ * ordering is reflected in the order of the bake Results property.
+ */
+static constexpr EBakeMapType ENUM_EBAKEMAPTYPE_ALL[] =
 {
 	EBakeMapType::TangentSpaceNormal,
-	EBakeMapType::Occlusion, // (AmbientOcclusion | BentNormal)
-	EBakeMapType::Curvature,
-	EBakeMapType::Texture,
 	EBakeMapType::ObjectSpaceNormal,
 	EBakeMapType::FaceNormal,
+	EBakeMapType::BentNormal,
 	EBakeMapType::Position,
-	EBakeMapType::MaterialID,
+	EBakeMapType::Curvature,
+	EBakeMapType::AmbientOcclusion,
+	EBakeMapType::Texture,
 	EBakeMapType::MultiTexture,
-	EBakeMapType::VertexColor
+	EBakeMapType::VertexColor,
+	EBakeMapType::MaterialID
 };
 
+
 /**
  * Base Mesh Bake tool
  */
