Increased the fractional precision of the surface area and volume numbers in the Inspector tool. The precision is 4 by default. It increases further for smaller numbers.
Added a utility function to convert a number to a string while being able to specify the fractional precision if the number is below a certain other number.

#rnx

diff --git a/Engine/Plugins/Experimental/MeshModelingToolsetExp/Source/MeshModelingToolsExp/Private/Properties/MeshAnalysisProperties.cpp b/Engine/Plugins/Experimental/MeshModelingToolsetExp/Source/MeshModelingToolsExp/Private/Properties/MeshAnalysisProperties.cpp
--- a/Engine/Plugins/Experimental/MeshModelingToolsetExp/Source/MeshModelingToolsExp/Private/Properties/MeshAnalysisProperties.cpp
+++ b/Engine/Plugins/Experimental/MeshModelingToolsetExp/Source/MeshModelingToolsExp/Private/Properties/MeshAnalysisProperties.cpp
@@ -6,17 +6,66 @@
 #include "MeshQueries.h"
 #include "MeshAdapter.h"
 #include "DynamicMesh/MeshAdapterUtil.h"
+#include "Internationalization/FastDecimalFormat.h"
+#include "Internationalization/Text.h"
+#include "Math/BasicMathExpressionEvaluator.h"
 
 #define LOCTEXT_NAMESPACE "UMeshAnalysisProperites"
 
 using namespace UE::Geometry;
 
+/** 
+ * Utility function that converts a number to a string. The precision of the fractional part can be adjusted based on 
+ * the magnitude of the input number.
+ * 
+ * @param Number to convert to a string
+ * @param MagnitudeToFractionalDigits Array of tuples (x, y) such that if the Number is less than x then set the fractional precision to y. 
+ * @param DefaultMaxFractionalDigits Default precision if the Number doesn't fall into any range specified in MagnitudeToFractionalDigits
+ */
+template<typename RealType>
+FString NumberToStringWithFractionalPrecision(const RealType Number, 
+											  const TArray<TTuple<RealType, int32>>& MagnitudeToFractionalDigits, 
+											  const int DefaultMaxFractionalDigits = 6) 
+{
+	FNumberFormattingOptions NumberFormattingOptions;
+	NumberFormattingOptions.SetUseGrouping(false);
+	NumberFormattingOptions.SetMaximumFractionalDigits(DefaultMaxFractionalDigits);
+	
+	if (MagnitudeToFractionalDigits.Num() > 0) 
+	{
+		TArray<TTuple<RealType, int32>> SortedMagnitudeToFractionalDigits = MagnitudeToFractionalDigits;
+		SortedMagnitudeToFractionalDigits.Sort([](const TTuple<RealType, int32>& LHS, const TTuple<RealType, int32>& RHS) { return LHS.Key < RHS.Key; });
+
+		if (Number < SortedMagnitudeToFractionalDigits.Last().Key) // First check if the number falls into any range
+		{
+			for (const TTuple<RealType, int32>& Tuple : SortedMagnitudeToFractionalDigits) 
+			{
+				if (Number < Tuple.Key)
+				{
+					NumberFormattingOptions.SetMaximumFractionalDigits(Tuple.Value);
+					break;
+				}
+			}
+		}
+	}
+
+	return FastDecimalFormat::NumberToString(Number, ExpressionParser::GetLocalizedNumberFormattingRules(), NumberFormattingOptions);
+}
+
 void UMeshAnalysisProperties::Update(const FDynamicMesh3& MeshIn, const FTransform& Transform)
 {
 	FTriangleMeshAdapterd TransformedMesh = UE::Geometry::MakeTransformedDynamicMeshAdapter(&MeshIn, Transform);
 	FVector2d VolArea = TMeshQueries<FTriangleMeshAdapterd>::GetVolumeArea(TransformedMesh);
-	this->SurfaceArea = FString::Printf(TEXT("%.2f m^2"), VolArea.Y / 10000);
-	this->Volume = FString::Printf(TEXT("%.2f m^3"), VolArea.X / 1000000);
+	
+	TArray<TTuple<double, int32>> MagnitudeToFractionalDigits;
+	MagnitudeToFractionalDigits.Add(MakeTuple(0.1, 8)); // If the number is less than 0.1 show 8 fractional digits
+	MagnitudeToFractionalDigits.Add(MakeTuple(1, 6)); // If the number is less than 1 show 6 fractional digits
+	
+	const double VolInMeters = VolArea[0]/1000000;
+	const double AreaInMeters = VolArea[1]/10000;
+	
+	this->Volume = NumberToStringWithFractionalPrecision(VolInMeters, MagnitudeToFractionalDigits, 4);
+	this->SurfaceArea = NumberToStringWithFractionalPrecision(AreaInMeters, MagnitudeToFractionalDigits, 4);
 }
 
 
diff --git a/Engine/Plugins/Experimental/MeshModelingToolsetExp/Source/MeshModelingToolsExp/Public/Properties/MeshAnalysisProperties.h b/Engine/Plugins/Experimental/MeshModelingToolsetExp/Source/MeshModelingToolsExp/Public/Properties/MeshAnalysisProperties.h
--- a/Engine/Plugins/Experimental/MeshModelingToolsetExp/Source/MeshModelingToolsExp/Public/Properties/MeshAnalysisProperties.h
+++ b/Engine/Plugins/Experimental/MeshModelingToolsetExp/Source/MeshModelingToolsExp/Public/Properties/MeshAnalysisProperties.h
@@ -19,9 +19,11 @@ class MESHMODELINGTOOLSEXP_API UMeshAnalysisProperties : public UInteractiveTool
 	GENERATED_BODY()
 public:
 
+	/** In meters squared */
 	UPROPERTY(VisibleAnywhere, Category = MeshAnalysis)
 	FString SurfaceArea;
 
+	/** In cubic meters */
 	UPROPERTY(VisibleAnywhere, Category = MeshAnalysis)
 	FString Volume;
 
