Mirror tool: closing a dialog that asks the user is they want to delete a mesh component does the same thing as hitting No

#rnx

diff --git a/Engine/Plugins/Experimental/MeshModelingToolsetExp/Source/MeshModelingToolsExp/Private/MirrorTool.cpp b/Engine/Plugins/Experimental/MeshModelingToolsetExp/Source/MeshModelingToolsExp/Private/MirrorTool.cpp
--- a/Engine/Plugins/Experimental/MeshModelingToolsetExp/Source/MeshModelingToolsExp/Private/MirrorTool.cpp
+++ b/Engine/Plugins/Experimental/MeshModelingToolsetExp/Source/MeshModelingToolsExp/Private/MirrorTool.cpp
@@ -392,7 +392,7 @@ void UMirrorTool::GenerateAsset(const TArray<FDynamicMeshOpResult>& Results)
 		FText Title = LOCTEXT("MirrorDestroyTitle", "Delete mesh components?");
 		EAppReturnType::Type Ret = FMessageDialog::Open(EAppMsgType::YesNo,
 			LOCTEXT("MirrorDestroyQuestion", "The mirror plane cropping has entirely cut away at least one mesh. Do you actually want to delete these mesh components? Note that either way all actors will remain, and meshes that are not fully cut away will still be mirrored as normal."), &Title);
-		if (Ret == EAppReturnType::No)
+		if (Ret == EAppReturnType::No || Ret == EAppReturnType::Cancel)
 		{
 			bWantToDestroy = false;
 		}
