ModelingTools: Reorder bake input mesh and output settings UI.
Remove TransientToolProperty from Resolution and Bit Depth properties.
Minor refactoring of BakeMeshAttributeMapsToolBase.

#rnx

diff --git a/Engine/Plugins/Experimental/MeshModelingToolsetExp/Source/MeshModelingToolsExp/Private/BakeMeshAttributeMapsTool.cpp b/Engine/Plugins/Experimental/MeshModelingToolsetExp/Source/MeshModelingToolsExp/Private/BakeMeshAttributeMapsTool.cpp
--- a/Engine/Plugins/Experimental/MeshModelingToolsetExp/Source/MeshModelingToolsExp/Private/BakeMeshAttributeMapsTool.cpp
+++ b/Engine/Plugins/Experimental/MeshModelingToolsetExp/Source/MeshModelingToolsExp/Private/BakeMeshAttributeMapsTool.cpp
@@ -275,6 +275,17 @@ void UBakeMeshAttributeMapsTool::Setup()
 	
 	// Setup tool property sets
 
+	Settings = NewObject<UBakeMeshAttributeMapsToolProperties>(this);
+	Settings->RestoreProperties(this);
+	AddToolPropertySource(Settings);
+
+	Settings->WatchProperty(Settings->MapTypes, [this](int32) { OpState |= EBakeOpState::Evaluate; UpdateOnModeChange(); });
+	Settings->WatchProperty(Settings->MapPreview, [this](FString) { UpdateVisualization(); GetToolManager()->PostInvalidation(); });
+	Settings->WatchProperty(Settings->Resolution, [this](EBakeTextureResolution) { OpState |= EBakeOpState::Evaluate; });
+	Settings->WatchProperty(Settings->BitDepth, [this](EBakeTextureBitDepth) { OpState |= EBakeOpState::Evaluate; });
+	Settings->WatchProperty(Settings->SamplesPerPixel, [this](EBakeTextureSamplesPerPixel) { OpState |= EBakeOpState::Evaluate; });
+	
+
 	InputMeshSettings = NewObject<UBakeInputMeshProperties>(this);
 	InputMeshSettings->RestoreProperties(this);
 	AddToolPropertySource(InputMeshSettings);
@@ -302,17 +313,12 @@ void UBakeMeshAttributeMapsTool::Setup()
 		}
 		OpState |= EBakeOpState::Evaluate;
 	});
-
 	
-	Settings = NewObject<UBakeMeshAttributeMapsToolProperties>(this);
-	Settings->RestoreProperties(this);
-	AddToolPropertySource(Settings);
 
-	Settings->WatchProperty(Settings->MapTypes, [this](int32) { OpState |= EBakeOpState::Evaluate; UpdateOnModeChange(); });
-	Settings->WatchProperty(Settings->MapPreview, [this](FString) { UpdateVisualization(); GetToolManager()->PostInvalidation(); });
-	Settings->WatchProperty(Settings->Resolution, [this](EBakeTextureResolution) { OpState |= EBakeOpState::Evaluate; });
-	Settings->WatchProperty(Settings->BitDepth, [this](EBakeTextureBitDepth) { OpState |= EBakeOpState::Evaluate; });
-	Settings->WatchProperty(Settings->SamplesPerPixel, [this](EBakeTextureSamplesPerPixel) { OpState |= EBakeOpState::Evaluate; });
+	ResultSettings = NewObject<UBakeMeshAttributeMapsResultToolProperties>(this);
+	ResultSettings->RestoreProperties(this);
+	AddToolPropertySource(ResultSettings);
+	SetToolPropertySourceEnabled(ResultSettings, true);
 	
 
 	OcclusionSettings = NewObject<UBakeOcclusionMapToolProperties>(this);
@@ -376,7 +382,7 @@ bool UBakeMeshAttributeMapsTool::CanAccept() const
 	if (bCanAccept)
 	{
 		// Allow Accept if all non-None types have valid results.
-		for (const TTuple<EBakeMapType, TObjectPtr<UTexture2D>>& Result : Settings->Result)
+		for (const TTuple<EBakeMapType, TObjectPtr<UTexture2D>>& Result : ResultSettings->Result)
 		{
 			bCanAccept = bCanAccept && Result.Get<1>();
 		}
@@ -476,7 +482,7 @@ void UBakeMeshAttributeMapsTool::Shutdown(EToolShutdownType ShutdownType)
 			SourceAsset = SkeletalMeshTarget ? SkeletalMeshTarget->GetSkeletalMesh() : nullptr;
 		}
 		const UPrimitiveComponent* SourceComponent = UE::ToolTarget::GetTargetComponent(Targets[0]);
-		CreateTextureAssets(Settings->Result, SourceComponent->GetWorld(), SourceAsset);
+		CreateTextureAssets(ResultSettings->Result, SourceComponent->GetWorld(), SourceAsset);
 	}
 }
 
@@ -663,19 +669,24 @@ void UBakeMeshAttributeMapsTool::UpdateVisualization()
 	// Populate Settings->Result from CachedMaps
 	for (const TTuple<EBakeMapType, TObjectPtr<UTexture2D>>& Map : CachedMaps)
 	{
-		if (Settings->Result.Contains(Map.Get<0>()))
+		if (ResultSettings->Result.Contains(Map.Get<0>()))
 		{
-			Settings->Result[Map.Get<0>()] = Map.Get<1>();
+			ResultSettings->Result[Map.Get<0>()] = Map.Get<1>();
 		}
 	}
 
-	UpdatePreview(Settings);
+	UpdatePreview(Settings->MapPreview, Settings->MapPreviewNamesMap);
 }
 
 
 void UBakeMeshAttributeMapsTool::UpdateOnModeChange()
 {
-	OnMapTypesUpdated(Settings);
+	OnMapTypesUpdated(
+		static_cast<EBakeMapType>(Settings->MapTypes),
+		ResultSettings->Result,
+		Settings->MapPreview,
+		Settings->MapPreviewNamesList,
+		Settings->MapPreviewNamesMap);
 
 	// Update tool property sets.
 	SetToolPropertySourceEnabled(OcclusionSettings, false);
@@ -717,7 +728,7 @@ void UBakeMeshAttributeMapsTool::UpdateOnModeChange()
 
 void UBakeMeshAttributeMapsTool::InvalidateResults()
 {
-	for (TTuple<EBakeMapType, TObjectPtr<UTexture2D>>& Result : Settings->Result)
+	for (TTuple<EBakeMapType, TObjectPtr<UTexture2D>>& Result : ResultSettings->Result)
 	{
 		Result.Get<1>() = nullptr;
 	}
diff --git a/Engine/Plugins/Experimental/MeshModelingToolsetExp/Source/MeshModelingToolsExp/Private/BakeMeshAttributeMapsToolBase.cpp b/Engine/Plugins/Experimental/MeshModelingToolsetExp/Source/MeshModelingToolsExp/Private/BakeMeshAttributeMapsToolBase.cpp
--- a/Engine/Plugins/Experimental/MeshModelingToolsetExp/Source/MeshModelingToolsExp/Private/BakeMeshAttributeMapsToolBase.cpp
+++ b/Engine/Plugins/Experimental/MeshModelingToolsetExp/Source/MeshModelingToolsExp/Private/BakeMeshAttributeMapsToolBase.cpp
@@ -171,6 +171,19 @@ void UBakeMeshAttributeMapsToolBase::CreateTextureAssets(const TMap<EBakeMapType
 }
 
 
+void UBakeMeshAttributeMapsToolBase::UpdatePreview(const FString& PreviewDisplayName, const TMap<FString, FString>& MapPreviewNamesMap)
+{
+	if (const FString* PreviewNameString = MapPreviewNamesMap.Find(PreviewDisplayName))
+	{
+		const int64 PreviewValue = StaticEnum<EBakeMapType>()->GetValueByNameString(*PreviewNameString);
+		if (PreviewValue != INDEX_NONE)
+		{
+			UpdatePreview(static_cast<EBakeMapType>(PreviewValue));
+		}
+	}
+}
+
+
 void UBakeMeshAttributeMapsToolBase::UpdatePreview(const EBakeMapType PreviewMapType)
 {
 	if (PreviewMapType == EBakeMapType::None)
@@ -229,6 +242,69 @@ void UBakeMeshAttributeMapsToolBase::UpdatePreview(const EBakeMapType PreviewMap
 }
 
 
+void UBakeMeshAttributeMapsToolBase::UpdatePreviewNames(
+	const EBakeMapType MapTypes,
+	FString& MapPreview,
+	TArray<FString>& MapPreviewNamesList,
+	TMap<FString, FString>& MapPreviewNamesMap)
+{
+	// Update our preview names list.
+	MapPreviewNamesList.Reset();
+	MapPreviewNamesMap.Reset();
+	bool bFoundMapType = false;
+	for (const TTuple<EBakeMapType, TObjectPtr<UTexture2D>>& Map : CachedMaps)
+	{
+		// Only populate map types that were requested. Some map types like
+		// AO may have only been added for preview of other types (ex. BentNormal)
+		if (static_cast<bool>(MapTypes & Map.Get<0>()))
+		{
+			const UEnum* BakeTypeEnum = StaticEnum<EBakeMapType>();
+			const int64 BakeEnumValue = static_cast<int64>(Map.Get<0>());
+			const FString BakeTypeDisplayName = BakeTypeEnum->GetDisplayNameTextByValue(BakeEnumValue).ToString();
+			const FString BakeTypeNameString = BakeTypeEnum->GetNameStringByValue(BakeEnumValue);
+			MapPreviewNamesList.Add(BakeTypeDisplayName);
+			MapPreviewNamesMap.Emplace(BakeTypeDisplayName, BakeTypeNameString);
+			if (MapPreview == MapPreviewNamesList.Last())
+			{
+				bFoundMapType = true;
+			}
+		}
+	}
+	if (!bFoundMapType)
+	{
+		MapPreview = MapPreviewNamesList.Num() > 0 ? MapPreviewNamesList[0] : TEXT("");
+	}
+}
+
+
+void UBakeMeshAttributeMapsToolBase::OnMapTypesUpdated(
+	const EBakeMapType ResultMapTypes,
+	TMap<EBakeMapType, TObjectPtr<UTexture2D>>& Result,
+	FString& MapPreview,
+	TArray<FString>& MapPreviewNamesList,
+	TMap<FString, FString>& MapPreviewNamesMap)
+{
+	// Use the processed bitfield which may contain additional targets
+	// (ex. AO if BentNormal was requested) to preallocate cached map storage.
+	EBakeMapType CachedMapTypes = GetMapTypes(static_cast<int32>(ResultMapTypes));
+	CachedMaps.Empty();
+	Result.Empty();
+	for (EBakeMapType MapType : ENUM_EBAKEMAPTYPE_ALL)
+	{
+		if (static_cast<bool>(CachedMapTypes & MapType))
+		{
+			CachedMaps.Add(MapType, nullptr);
+		}
+		
+		if (static_cast<bool>(ResultMapTypes & MapType))
+		{
+			Result.Add(MapType, nullptr);
+		}
+	}
+	UpdatePreviewNames(ResultMapTypes, MapPreview, MapPreviewNamesList, MapPreviewNamesMap);
+}
+
+
 void UBakeMeshAttributeMapsToolBase::OnMapsUpdated(const TUniquePtr<UE::Geometry::FMeshMapBaker>& NewResult)
 {
 	const FImageDimensions BakeDimensions = NewResult->GetDimensions();
diff --git a/Engine/Plugins/Experimental/MeshModelingToolsetExp/Source/MeshModelingToolsExp/Private/BakeMultiMeshAttributeMapsTool.cpp b/Engine/Plugins/Experimental/MeshModelingToolsetExp/Source/MeshModelingToolsExp/Private/BakeMultiMeshAttributeMapsTool.cpp
--- a/Engine/Plugins/Experimental/MeshModelingToolsetExp/Source/MeshModelingToolsExp/Private/BakeMultiMeshAttributeMapsTool.cpp
+++ b/Engine/Plugins/Experimental/MeshModelingToolsetExp/Source/MeshModelingToolsExp/Private/BakeMultiMeshAttributeMapsTool.cpp
@@ -409,15 +409,7 @@ void UBakeMultiMeshAttributeMapsTool::Setup()
 	UToolTarget* Target = Targets[0];
 
 	// Setup tool property sets
-	
-	InputMeshSettings = NewObject<UBakeMultiMeshInputToolProperties>(this);
-	InputMeshSettings->RestoreProperties(this);
-	AddToolPropertySource(InputMeshSettings);
-	InputMeshSettings->TargetStaticMesh = GetStaticMeshTarget(Target);
-	UpdateUVLayerNames(InputMeshSettings->TargetUVLayer, InputMeshSettings->TargetUVLayerNamesList, TargetMesh);
-	InputMeshSettings->WatchProperty(InputMeshSettings->TargetUVLayer, [this](FString) { OpState |= EBakeOpState::Evaluate; });
-	InputMeshSettings->WatchProperty(InputMeshSettings->ProjectionDistance, [this](float) { OpState |= EBakeOpState::Evaluate; });
-	
+
 	Settings = NewObject<UBakeMultiMeshAttributeMapsToolProperties>(this);
 	Settings->RestoreProperties(this);
 	AddToolPropertySource(Settings);
@@ -427,6 +419,19 @@ void UBakeMultiMeshAttributeMapsTool::Setup()
 	Settings->WatchProperty(Settings->Resolution, [this](EBakeTextureResolution) { OpState |= EBakeOpState::Evaluate; });
 	Settings->WatchProperty(Settings->BitDepth, [this](EBakeTextureBitDepth) { OpState |= EBakeOpState::Evaluate; });
 	Settings->WatchProperty(Settings->SamplesPerPixel, [this](EBakeTextureSamplesPerPixel) { OpState |= EBakeOpState::Evaluate; });
+	
+	InputMeshSettings = NewObject<UBakeMultiMeshInputToolProperties>(this);
+	InputMeshSettings->RestoreProperties(this);
+	AddToolPropertySource(InputMeshSettings);
+	InputMeshSettings->TargetStaticMesh = GetStaticMeshTarget(Target);
+	UpdateUVLayerNames(InputMeshSettings->TargetUVLayer, InputMeshSettings->TargetUVLayerNamesList, TargetMesh);
+	InputMeshSettings->WatchProperty(InputMeshSettings->TargetUVLayer, [this](FString) { OpState |= EBakeOpState::Evaluate; });
+	InputMeshSettings->WatchProperty(InputMeshSettings->ProjectionDistance, [this](float) { OpState |= EBakeOpState::Evaluate; });
+
+	ResultSettings = NewObject<UBakeMeshAttributeMapsResultToolProperties>(this);
+	ResultSettings->RestoreProperties(this);
+	AddToolPropertySource(ResultSettings);
+	SetToolPropertySourceEnabled(ResultSettings, true);
 
 
 	// Pre-populate detail mesh data
@@ -484,7 +489,7 @@ bool UBakeMultiMeshAttributeMapsTool::CanAccept() const
 	if (bCanAccept)
 	{
 		// Allow Accept if all non-None types have valid results.
-		for (const TTuple<EBakeMapType, TObjectPtr<UTexture2D>>& Result : Settings->Result)
+		for (const TTuple<EBakeMapType, TObjectPtr<UTexture2D>>& Result : ResultSettings->Result)
 		{
 			bCanAccept = bCanAccept && Result.Get<1>();
 		}
@@ -541,7 +546,7 @@ void UBakeMultiMeshAttributeMapsTool::Shutdown(EToolShutdownType ShutdownType)
 		IStaticMeshBackedTarget* StaticMeshTarget = Cast<IStaticMeshBackedTarget>(Targets[0]);
 		UObject* SourceAsset = StaticMeshTarget ? StaticMeshTarget->GetStaticMesh() : nullptr;
 		const UPrimitiveComponent* SourceComponent = UE::ToolTarget::GetTargetComponent(Targets[0]);
-		CreateTextureAssets(Settings->Result, SourceComponent->GetWorld(), SourceAsset);
+		CreateTextureAssets(ResultSettings->Result, SourceComponent->GetWorld(), SourceAsset);
 	}
 }
 
@@ -670,25 +675,30 @@ void UBakeMultiMeshAttributeMapsTool::UpdateVisualization()
 	// Populate Settings->Result from CachedMaps
 	for (const TTuple<EBakeMapType, TObjectPtr<UTexture2D>>& Map : CachedMaps)
 	{
-		if (Settings->Result.Contains(Map.Get<0>()))
+		if (ResultSettings->Result.Contains(Map.Get<0>()))
 		{
-			Settings->Result[Map.Get<0>()] = Map.Get<1>();
+			ResultSettings->Result[Map.Get<0>()] = Map.Get<1>();
 		}
 	}
 
-	UpdatePreview(Settings);
+	UpdatePreview(Settings->MapPreview, Settings->MapPreviewNamesMap);
 }
 
 
 void UBakeMultiMeshAttributeMapsTool::UpdateOnModeChange()
 {
-	OnMapTypesUpdated(Settings);
+	OnMapTypesUpdated(
+		static_cast<EBakeMapType>(Settings->MapTypes),
+		ResultSettings->Result,
+		Settings->MapPreview,
+		Settings->MapPreviewNamesList,
+		Settings->MapPreviewNamesMap);
 }
 
 
 void UBakeMultiMeshAttributeMapsTool::InvalidateResults()
 {
-	for (TTuple<EBakeMapType, TObjectPtr<UTexture2D>>& Result : Settings->Result)
+	for (TTuple<EBakeMapType, TObjectPtr<UTexture2D>>& Result : ResultSettings->Result)
 	{
 		Result.Get<1>() = nullptr;
 	}
diff --git a/Engine/Plugins/Experimental/MeshModelingToolsetExp/Source/MeshModelingToolsExp/Public/BakeMeshAttributeMapsTool.h b/Engine/Plugins/Experimental/MeshModelingToolsetExp/Source/MeshModelingToolsExp/Public/BakeMeshAttributeMapsTool.h
--- a/Engine/Plugins/Experimental/MeshModelingToolsetExp/Source/MeshModelingToolsExp/Public/BakeMeshAttributeMapsTool.h
+++ b/Engine/Plugins/Experimental/MeshModelingToolsetExp/Source/MeshModelingToolsExp/Public/BakeMeshAttributeMapsTool.h
@@ -62,21 +62,17 @@ public:
 	FString MapPreview;
 
 	/** The pixel resolution of the generated textures */
-	UPROPERTY(EditAnywhere, Category = Textures, meta = (TransientToolProperty))
+	UPROPERTY(EditAnywhere, Category = Textures)
 	EBakeTextureResolution Resolution = EBakeTextureResolution::Resolution256;
 
 	/** The bit depth for each channel of the generated textures */
-	UPROPERTY(EditAnywhere, Category = Textures, meta = (TransientToolProperty))
+	UPROPERTY(EditAnywhere, Category = Textures)
 	EBakeTextureBitDepth BitDepth = EBakeTextureBitDepth::ChannelBits8;
 
 	/** Number of samples per pixel */
 	UPROPERTY(EditAnywhere, Category = Textures)
 	EBakeTextureSamplesPerPixel SamplesPerPixel = EBakeTextureSamplesPerPixel::Sample1;
 
-	/** Bake */
-	UPROPERTY(VisibleAnywhere, Category = Textures, meta = (DisplayName = "Results", TransientToolProperty))
-	TMap<EBakeMapType, TObjectPtr<UTexture2D>> Result;
-
 	UFUNCTION()
 	const TArray<FString>& GetMapPreviewNamesFunc();
 	UPROPERTY(meta = (TransientToolProperty))
@@ -119,6 +115,9 @@ protected:
 	UPROPERTY()
 	TObjectPtr<UBakeMeshAttributeMapsToolProperties> Settings;
 
+	UPROPERTY()
+	TObjectPtr<UBakeMeshAttributeMapsResultToolProperties> ResultSettings;
+
 
 	// Begin UBakeMeshAttributeMapsToolBase interface
 	virtual void UpdateResult() override;
diff --git a/Engine/Plugins/Experimental/MeshModelingToolsetExp/Source/MeshModelingToolsExp/Public/BakeMeshAttributeMapsToolBase.h b/Engine/Plugins/Experimental/MeshModelingToolsetExp/Source/MeshModelingToolsExp/Public/BakeMeshAttributeMapsToolBase.h
--- a/Engine/Plugins/Experimental/MeshModelingToolsetExp/Source/MeshModelingToolsExp/Public/BakeMeshAttributeMapsToolBase.h
+++ b/Engine/Plugins/Experimental/MeshModelingToolsetExp/Source/MeshModelingToolsExp/Public/BakeMeshAttributeMapsToolBase.h
@@ -56,6 +56,17 @@ enum class EBakeTextureSamplesPerPixel
 };
 
 
+UCLASS()
+class MESHMODELINGTOOLSEXP_API UBakeMeshAttributeMapsResultToolProperties : public UInteractiveToolPropertySet
+{
+	GENERATED_BODY()
+public:
+	/** Bake */
+	UPROPERTY(VisibleAnywhere, Category = Results, meta = (DisplayName = "Results", TransientToolProperty))
+	TMap<EBakeMapType, TObjectPtr<UTexture2D>> Result;
+};
+
+
 /**
  * Base Bake Maps tool
  */
@@ -167,9 +178,18 @@ protected:
 
 	/**
 	 * To be invoked by client when bake map types change.
+	 * @param ResultMapTypes the requested map types to compute
+	 * @param Result the output map of bake result textures
+	 * @param MapPreview the map preview string property
+	 * @param MapPreviewNamesList the stored list of map preview display names.
+	 * @param MapPreviewNamesMap the stored map of map preview display names to enum string values.
 	 */
-	template <typename PropertySet>
-	void OnMapTypesUpdated(PropertySet& Properties);
+	void OnMapTypesUpdated(
+		EBakeMapType ResultMapTypes,
+		TMap<EBakeMapType, TObjectPtr<UTexture2D>>& Result,
+		FString& MapPreview,
+		TArray<FString>& MapPreviewNamesList,
+		TMap<FString, FString>& MapPreviewNamesMap);
 
 	//
 	// Background compute
@@ -193,10 +213,10 @@ protected:
 	/**
 	 * Update the preview material parameters for a given Bake type
 	 * display name.
-	 * @param Properties Properties containing the display name of a Bake type to preview.
+	 * @param PreviewDisplayName Displayed UI preview name to preview
+	 * @param MapPreviewNamesMap Map containing the list of displayed preview name to enum string value.
 	 */
-	template <typename PropertySet>
-	void UpdatePreview(PropertySet& Properties);
+	void UpdatePreview(const FString& PreviewDisplayName, const TMap<FString, FString>& MapPreviewNamesMap);
 
 
 	/**
@@ -210,11 +230,17 @@ protected:
 	 * Updates a tool property set's MapPreviewNamesList from the list of
 	 * active map types. Also updates the MapPreview property if the current
 	 * preview option is no longer available.
-	 *
-	 * @param Properties the UInteractiveToolPropertySet to update.
+	 * 
+	 * @param MapTypes the requested map types to compute
+	 * @param MapPreview the map preview string property
+	 * @param MapPreviewNamesList the stored list of map preview display names.
+	 * @param MapPreviewNamesMap the stored map of map preview display names to enum string values.
 	 */
-	template <typename PropertySet>
-	void UpdatePreviewNames(PropertySet& Properties);
+	void UpdatePreviewNames(
+		const EBakeMapType MapTypes,
+		FString& MapPreview,
+		TArray<FString>& MapPreviewNamesList,
+		TMap<FString, FString>& MapPreviewNamesMap);
 
 
 	//
@@ -283,14 +309,6 @@ protected:
 
 	/** @return the texture name given a base name and map type */
 	static void GetTextureName(EBakeMapType MapType, const FString& BaseName, FString& TexName);
-
-	/**
-	 * @param Properties the tool property set to validate the map type against.
-	 * @param MapType the map type to validate.
-	 * @return true if MapType was requested from the given PropertySet
-	 */
-	template <typename PropertySet>
-	static bool IsRequestedMapType(PropertySet& Properties, EBakeMapType MapType);
 	
 
 	// empty maps are shown when nothing is computed
@@ -307,94 +325,4 @@ protected:
 };
 
 
-template <typename PropertySet>
-void UBakeMeshAttributeMapsToolBase::OnMapTypesUpdated(PropertySet& Properties)
-{
-	const EBakeMapType BakeMapTypes = GetMapTypes(Properties->MapTypes);
-	
-	// Use the processed bitfield which may contain additional targets
-	// (ex. AO if BentNormal was requested).
-	CachedMaps.Empty();
-	for (EBakeMapType MapType : ENUM_EBAKEMAPTYPE_ALL)
-	{
-		if( (bool)(BakeMapTypes & MapType) )
-		{
-			CachedMaps.Add(MapType, nullptr);
-		}
-	}
-
-	// Initialize Properties->Result with requested MapTypes
-	Properties->Result.Empty();
-	for (const TTuple<EBakeMapType, TObjectPtr<UTexture2D>>& Map : CachedMaps)
-	{
-		// Only populate map types that were requested. Some map types like
-		// AO may have only been added for preview of other types (ex. BentNormal)
-		if (IsRequestedMapType(Properties, Map.Get<0>()))
-		{
-			Properties->Result.Add(Map.Get<0>(), nullptr);
-		}
-	}
-
-	UpdatePreviewNames(Properties);
-}
-
-
-template <typename PropertySet>
-void UBakeMeshAttributeMapsToolBase::UpdatePreview(PropertySet& Properties)
-{
-	const FString& PreviewDisplayName = Properties->MapPreview;
-	if (const FString* PreviewNameString = Properties->MapPreviewNamesMap.Find(PreviewDisplayName))
-	{
-		const int64 PreviewValue = StaticEnum<EBakeMapType>()->GetValueByNameString(*PreviewNameString);
-		if (PreviewValue != INDEX_NONE)
-		{
-			UpdatePreview(static_cast<EBakeMapType>(PreviewValue));
-		}
-	}
-}
-
-
-template <typename PropertySet>
-void UBakeMeshAttributeMapsToolBase::UpdatePreviewNames(PropertySet& Properties)
-{
-	// Update our preview names list.
-	Properties->MapPreviewNamesList.Reset();
-	Properties->MapPreviewNamesMap.Reset();
-	bool bFoundMapType = false;
-	for (const TTuple<EBakeMapType, TObjectPtr<UTexture2D>>& Map : CachedMaps)
-	{
-		// Only populate map types that were requested. Some map types like
-		// AO may have only been added for preview of other types (ex. BentNormal)
-		if (IsRequestedMapType(Properties, Map.Get<0>()))
-		{
-			const UEnum* BakeTypeEnum = StaticEnum<EBakeMapType>();
-			const int64 BakeEnumValue = static_cast<int64>(Map.Get<0>());
-			const FString BakeTypeDisplayName = BakeTypeEnum->GetDisplayNameTextByValue(BakeEnumValue).ToString();
-			const FString BakeTypeNameString = BakeTypeEnum->GetNameStringByValue(BakeEnumValue);
-			Properties->MapPreviewNamesList.Add(BakeTypeDisplayName);
-			Properties->MapPreviewNamesMap.Emplace(BakeTypeDisplayName, BakeTypeNameString);
-			if (Properties->MapPreview == Properties->MapPreviewNamesList.Last())
-			{
-				bFoundMapType = true;
-			}
-		}
-	}
-	if (!bFoundMapType)
-	{
-		Properties->MapPreview = Properties->MapPreviewNamesList.Num() > 0 ? Properties->MapPreviewNamesList[0] : TEXT("");
-	}
-}
-
-
-template <typename PropertySet>
-bool UBakeMeshAttributeMapsToolBase::IsRequestedMapType(PropertySet& Properties, EBakeMapType MapType)
-{
-	EBakeMapType SourceMapTypes = static_cast<EBakeMapType>(Properties->MapTypes);
-	return static_cast<bool>(SourceMapTypes & MapType);
-}
-
-
-
-
-
 
diff --git a/Engine/Plugins/Experimental/MeshModelingToolsetExp/Source/MeshModelingToolsExp/Public/BakeMultiMeshAttributeMapsTool.h b/Engine/Plugins/Experimental/MeshModelingToolsetExp/Source/MeshModelingToolsExp/Public/BakeMultiMeshAttributeMapsTool.h
--- a/Engine/Plugins/Experimental/MeshModelingToolsetExp/Source/MeshModelingToolsExp/Public/BakeMultiMeshAttributeMapsTool.h
+++ b/Engine/Plugins/Experimental/MeshModelingToolsetExp/Source/MeshModelingToolsExp/Public/BakeMultiMeshAttributeMapsTool.h
@@ -49,11 +49,11 @@ public:
 	FString MapPreview;
 
 	/** The pixel resolution of the generated map */
-	UPROPERTY(EditAnywhere, Category = Textures, meta = (TransientToolProperty))
+	UPROPERTY(EditAnywhere, Category = Textures)
 	EBakeTextureResolution Resolution = EBakeTextureResolution::Resolution256;
 
 	/** The channel bit depth of the source data for the generated textures */
-	UPROPERTY(EditAnywhere, Category = Textures, meta = (TransientToolProperty))
+	UPROPERTY(EditAnywhere, Category = Textures)
 	EBakeTextureBitDepth BitDepth = EBakeTextureBitDepth::ChannelBits8;
 
 	/** Number of samples per pixel */
@@ -68,10 +68,6 @@ public:
 	UPROPERTY(meta = (TransientToolProperty))
 	TArray<FString> MapPreviewNamesList;
 	TMap<FString, FString> MapPreviewNamesMap;
-
-	UPROPERTY(VisibleAnywhere, Category = Textures, meta = (DisplayName = "Results", TransientToolProperty))
-	TMap<EBakeMapType, TObjectPtr<UTexture2D>> Result;
-
 };
 
 
@@ -184,6 +180,9 @@ protected:
 	UPROPERTY()
 	TObjectPtr<UBakeMultiMeshInputToolProperties> InputMeshSettings;
 
+	UPROPERTY()
+	TObjectPtr<UBakeMeshAttributeMapsResultToolProperties> ResultSettings;
+
 
 protected:
 	// Begin UBakeMeshAttributeMapsToolBase interface
