correctly handle cancel as 'no' for plane cut and jacketing tools (copy fix from mirror tool)

#rnx

diff --git a/Engine/Plugins/Experimental/MeshModelingToolsetExp/Source/MeshModelingToolsExp/Private/PlaneCutTool.cpp b/Engine/Plugins/Experimental/MeshModelingToolsetExp/Source/MeshModelingToolsExp/Private/PlaneCutTool.cpp
--- a/Engine/Plugins/Experimental/MeshModelingToolsetExp/Source/MeshModelingToolsExp/Private/PlaneCutTool.cpp
+++ b/Engine/Plugins/Experimental/MeshModelingToolsetExp/Source/MeshModelingToolsExp/Private/PlaneCutTool.cpp
@@ -430,7 +430,7 @@ void UPlaneCutTool::GenerateAsset(const TArray<FDynamicMeshOpResult>& Results)
 		FText Title = LOCTEXT("PlaneCutDestroyTitle", "Delete mesh components?");
 		EAppReturnType::Type Ret = FMessageDialog::Open(EAppMsgType::YesNo, 
 			LOCTEXT("PlaneCutDestroyQuestion", "Plane cuts have entirely cut away at least one mesh. Do you actually want to delete these mesh components? Note that either way all actors will remain, and meshes that are not fully cut away will still be cut as normal."), &Title);
-		if (Ret == EAppReturnType::No)
+		if (Ret == EAppReturnType::No || Ret == EAppReturnType::Cancel)
 		{
 			bWantDestroy = false; // quell destructive urge
 		}
diff --git a/Engine/Plugins/Experimental/MeshModelingToolsetExp/Source/MeshModelingToolsExp/Private/RemoveOccludedTrianglesTool.cpp b/Engine/Plugins/Experimental/MeshModelingToolsetExp/Source/MeshModelingToolsExp/Private/RemoveOccludedTrianglesTool.cpp
--- a/Engine/Plugins/Experimental/MeshModelingToolsetExp/Source/MeshModelingToolsExp/Private/RemoveOccludedTrianglesTool.cpp
+++ b/Engine/Plugins/Experimental/MeshModelingToolsetExp/Source/MeshModelingToolsExp/Private/RemoveOccludedTrianglesTool.cpp
@@ -479,7 +479,7 @@ void URemoveOccludedTrianglesTool::GenerateAsset(const TArray<FDynamicMeshOpResu
 		FText Title = LOCTEXT("RemoveOccludedDestroyTitle", "Delete mesh components?");
 		EAppReturnType::Type Ret = FMessageDialog::Open(EAppMsgType::YesNo,
 			LOCTEXT("RemoveOccludedDestroyQuestion", "Jacketing has removed all triangles from at least one mesh. Do you actually want to delete these mesh components? Note that either way all actors will remain, and meshes that would not have all triangles removed will still be jacketed as normal."), &Title);
-		if (Ret == EAppReturnType::No)
+		if (Ret == EAppReturnType::No || Ret == EAppReturnType::Cancel)
 		{
 			bWantDestroy = false;
 		}
