Subdivide tool: optionally create new polygroups for refined output faces

#rnx

diff --git a/Engine/Plugins/Experimental/MeshModelingToolsetExp/Source/MeshModelingToolsEditorOnly/Private/SubdividePolyTool.cpp b/Engine/Plugins/Experimental/MeshModelingToolsetExp/Source/MeshModelingToolsEditorOnly/Private/SubdividePolyTool.cpp
--- a/Engine/Plugins/Experimental/MeshModelingToolsetExp/Source/MeshModelingToolsEditorOnly/Private/SubdividePolyTool.cpp
+++ b/Engine/Plugins/Experimental/MeshModelingToolsetExp/Source/MeshModelingToolsEditorOnly/Private/SubdividePolyTool.cpp
@@ -28,17 +28,20 @@ public:
 	SubdivPostProcessor(int InSubdivisionLevel,
 						ESubdivisionScheme InSubdivisionScheme,
 						ESubdivisionOutputNormals InNormalComputationMethod,
-						ESubdivisionOutputUVs InUVComputationMethod) :
+						ESubdivisionOutputUVs InUVComputationMethod,
+						bool bInNewPolyGroups) :
 		SubdivisionLevel(InSubdivisionLevel),
 		SubdivisionScheme(InSubdivisionScheme),
 		NormalComputationMethod(InNormalComputationMethod),
-		UVComputationMethod(InUVComputationMethod)
+		UVComputationMethod(InUVComputationMethod),
+		bNewPolyGroups(bInNewPolyGroups)
 	{}
 
 	int SubdivisionLevel = 3;
 	ESubdivisionScheme SubdivisionScheme = ESubdivisionScheme::CatmullClark;
 	ESubdivisionOutputNormals NormalComputationMethod = ESubdivisionOutputNormals::Generated;
 	ESubdivisionOutputUVs UVComputationMethod = ESubdivisionOutputUVs::Interpolated;
+	bool bNewPolyGroups = false;
 
 	void ProcessMesh(const FDynamicMesh3& Mesh, FDynamicMesh3& OutRenderMesh) final
 	{
@@ -48,6 +51,7 @@ public:
 		Subd.SubdivisionScheme = SubdivisionScheme;
 		Subd.NormalComputationMethod = NormalComputationMethod;
 		Subd.UVComputationMethod = UVComputationMethod;
+		Subd.bNewPolyGroups = bNewPolyGroups;
 
 		ensure(Subd.ComputeTopologySubdivision());
 
@@ -209,7 +213,8 @@ void USubdividePolyTool::Setup()
 	PreviewDynamicMeshComponent->SetRenderMeshPostProcessor(MakeUnique<SubdivPostProcessor>(Properties->SubdivisionLevel,
 																							Properties->SubdivisionScheme,
 																							Properties->NormalComputationMethod,
-																							Properties->UVComputationMethod));
+																							Properties->UVComputationMethod,
+																							Properties->bNewPolyGroups));
 
 	// Use the input mesh's material on the preview
 	FComponentMaterialSet MaterialSet;
@@ -233,7 +238,8 @@ void USubdividePolyTool::Setup()
 		PreviewDynamicMeshComponent->SetRenderMeshPostProcessor(MakeUnique<SubdivPostProcessor>(Properties->SubdivisionLevel,
 																								Properties->SubdivisionScheme,
 																								Properties->NormalComputationMethod,
-																								Properties->UVComputationMethod));
+																								Properties->UVComputationMethod,
+																								Properties->bNewPolyGroups));
 		PreviewDynamicMeshComponent->NotifyMeshUpdated();
 	};
 
@@ -259,6 +265,10 @@ void USubdividePolyTool::Setup()
 	{
 		RebuildMeshPostProcessor();
 	});
+	Properties->WatchProperty(Properties->bNewPolyGroups, [this, RebuildMeshPostProcessor](bool)
+	{
+		RebuildMeshPostProcessor();
+	});
 
 
 	auto RenderGroupsChanged = [this](bool bNewRenderGroups)
diff --git a/Engine/Plugins/Experimental/MeshModelingToolsetExp/Source/MeshModelingToolsEditorOnly/Public/SubdividePolyTool.h b/Engine/Plugins/Experimental/MeshModelingToolsetExp/Source/MeshModelingToolsEditorOnly/Public/SubdividePolyTool.h
--- a/Engine/Plugins/Experimental/MeshModelingToolsetExp/Source/MeshModelingToolsEditorOnly/Public/SubdividePolyTool.h
+++ b/Engine/Plugins/Experimental/MeshModelingToolsetExp/Source/MeshModelingToolsEditorOnly/Public/SubdividePolyTool.h
@@ -54,6 +54,10 @@ public:
 	UPROPERTY(EditAnywhere, Category=Settings)
 	ESubdivisionOutputUVs UVComputationMethod = ESubdivisionOutputUVs::Interpolated;
 
+	/** Assign a new PolyGroup ID to each newly created face */
+	UPROPERTY(EditAnywhere, Category = Settings)
+	bool bNewPolyGroups = false;
+
 	UPROPERTY(EditAnywhere, Category = Settings)
 	bool bRenderGroups = true;
 
diff --git a/Engine/Plugins/Runtime/MeshModelingToolset/Source/ModelingComponentsEditorOnly/Private/Operations/SubdividePoly.cpp b/Engine/Plugins/Runtime/MeshModelingToolset/Source/ModelingComponentsEditorOnly/Private/Operations/SubdividePoly.cpp
--- a/Engine/Plugins/Runtime/MeshModelingToolset/Source/ModelingComponentsEditorOnly/Private/Operations/SubdividePoly.cpp
+++ b/Engine/Plugins/Runtime/MeshModelingToolset/Source/ModelingComponentsEditorOnly/Private/Operations/SubdividePoly.cpp
@@ -717,14 +717,17 @@ bool FSubdividePoly::ComputeSubdividedMesh(FDynamicMesh3& OutMesh)
 	check(SourceGroupIDs.Num() == Refiner->TopologyRefiner->GetLevel(0).GetNumFaces());
 
 	TArray<int> RefinedGroupIDs;
-	for (int CurrentLevel = 1; CurrentLevel <= Level; ++CurrentLevel)
+	if (!bNewPolyGroups)
 	{
-		// TODO: Don't keep resizing -- preallocate one big buffer and move through it
-		RefinedGroupIDs.SetNumUninitialized(Refiner->TopologyRefiner->GetLevel(CurrentLevel).GetNumFaces());
-		int* s = SourceGroupIDs.GetData();
-		int* d = RefinedGroupIDs.GetData();
-		Interpolator.InterpolateFaceUniform(CurrentLevel, s, d);
-		SourceGroupIDs = RefinedGroupIDs;
+		for (int CurrentLevel = 1; CurrentLevel <= Level; ++CurrentLevel)
+		{
+			// TODO: Don't keep resizing -- preallocate one big buffer and move through it
+			RefinedGroupIDs.SetNumUninitialized(Refiner->TopologyRefiner->GetLevel(CurrentLevel).GetNumFaces());
+			int* s = SourceGroupIDs.GetData();
+			int* d = RefinedGroupIDs.GetData();
+			Interpolator.InterpolateFaceUniform(CurrentLevel, s, d);
+			SourceGroupIDs = RefinedGroupIDs;
+		}
 	}
 
 
@@ -796,7 +799,7 @@ bool FSubdividePoly::ComputeSubdividedMesh(FDynamicMesh3& OutMesh)
 	const OpenSubdiv::Far::TopologyLevel& FinalLevel = Refiner->TopologyRefiner->GetLevel(Level);
 
 	check(UVComputationMethod != ESubdivisionOutputUVs::Interpolated || FinalLevel.GetNumFVarValues() == RefinedUVs.Num());
-	check(FinalLevel.GetNumFaces() == RefinedGroupIDs.Num());
+	check(bNewPolyGroups || FinalLevel.GetNumFaces() == RefinedGroupIDs.Num());
 
 	// Add the faces (manually triangulate the output here)
 
@@ -819,7 +822,7 @@ bool FSubdividePoly::ComputeSubdividedMesh(FDynamicMesh3& OutMesh)
 
 	for (int FaceID = 0; FaceID < FinalLevel.GetNumFaces(); ++FaceID)
 	{
-		int GroupID = RefinedGroupIDs[FaceID];
+		int GroupID = bNewPolyGroups ? OutMesh.AllocateTriangleGroup() : RefinedGroupIDs[FaceID];
 
 		OpenSubdiv::Far::ConstIndexArray Face = FinalLevel.GetFaceVertices(FaceID);
 		NumFaceVertices += Face.size();
diff --git a/Engine/Plugins/Runtime/MeshModelingToolset/Source/ModelingComponentsEditorOnly/Public/Operations/SubdividePoly.h b/Engine/Plugins/Runtime/MeshModelingToolset/Source/ModelingComponentsEditorOnly/Public/Operations/SubdividePoly.h
--- a/Engine/Plugins/Runtime/MeshModelingToolset/Source/ModelingComponentsEditorOnly/Public/Operations/SubdividePoly.h
+++ b/Engine/Plugins/Runtime/MeshModelingToolset/Source/ModelingComponentsEditorOnly/Public/Operations/SubdividePoly.h
@@ -89,6 +89,8 @@ public:
 
 	ESubdivisionOutputUVs UVComputationMethod = ESubdivisionOutputUVs::Interpolated;
 
+	bool bNewPolyGroups = false;
+
 	// wrapper around Opensubdiv::TopologyRefiner to avoid including OSD headers here. See cpp file.
 	class RefinerImpl;
 
