MeshModelingToolset: Cleanup UV Projection tool

diff --git a/Engine/Plugins/Experimental/MeshModelingToolsetExp/Source/MeshModelingToolsExp/Private/UVProjectionTool.cpp b/Engine/Plugins/Experimental/MeshModelingToolsetExp/Source/MeshModelingToolsExp/Private/UVProjectionTool.cpp
--- a/Engine/Plugins/Experimental/MeshModelingToolsetExp/Source/MeshModelingToolsExp/Private/UVProjectionTool.cpp
+++ b/Engine/Plugins/Experimental/MeshModelingToolsetExp/Source/MeshModelingToolsExp/Private/UVProjectionTool.cpp
@@ -60,15 +60,6 @@ void UUVProjectionTool::Setup()
 	BasicProperties = NewObject<UUVProjectionToolProperties>(this);
 	BasicProperties->RestoreProperties(this);
 	BasicProperties->GetOnModified().AddLambda([this](UObject*, FProperty*)	{ Preview->InvalidateResult(); });
-	BasicProperties->WatchProperty(BasicProperties->ProjectionType, [this](EUVProjectionMethod) { OnProjectionTypeChanged(); });
-
-	ExpMapProperties = NewObject<UUVProjectionToolExpMapProperties>(this);
-	ExpMapProperties->RestoreProperties(this);
-	ExpMapProperties->GetOnModified().AddLambda([this](UObject*, FProperty*) { Preview->InvalidateResult();	});
-
-	CylinderProperties = NewObject<UUVProjectionToolCylinderProperties>(this);
-	CylinderProperties->RestoreProperties(this);
-	CylinderProperties->GetOnModified().AddLambda([this](UObject*, FProperty*) { Preview->InvalidateResult();	});
 
 	// initialize input mesh
 	InitializeMesh();
@@ -100,13 +91,10 @@ void UUVProjectionTool::Setup()
 
 	AddToolPropertySource(UVChannelProperties);
 	AddToolPropertySource(BasicProperties);
-	AddToolPropertySource(ExpMapProperties);
-	AddToolPropertySource(CylinderProperties);
 	AddToolPropertySource(EditActions);
 	AddToolPropertySource(MaterialSettings);
 	
 	OnMaterialSettingsChanged();
-	OnProjectionTypeChanged();
 
 	// set up visualizers
 	ProjectionShapeVisualizer.LineColor = FLinearColor::Red;
@@ -123,12 +111,14 @@ void UUVProjectionTool::Setup()
 	TransformGizmo->SetActiveTarget(TransformProxy, GetToolManager());
 
 	InitialDimensions = BasicProperties->Dimensions;
+	bInitialUniformDimensions = BasicProperties->bUniformDimensions;
 	InitialTransform = TransformGizmo->GetGizmoTransform();
 	
 	ApplyInitializationMode();
 
 	// start watching for dimensions changes
 	DimensionsWatcher = BasicProperties->WatchProperty(BasicProperties->Dimensions, [this](FVector) { bTransformModified = true; });
+	DimensionsModeWatcher = BasicProperties->WatchProperty(BasicProperties->bUniformDimensions, [this](bool) { bTransformModified = true; });
 	BasicProperties->WatchProperty(BasicProperties->Initialization, [this](EUVProjectionToolInitializationMode NewMode) { OnInitializationModeChanged(); });
 	bTransformModified = false;
 	BasicProperties->SilentUpdateWatched();
@@ -210,10 +200,9 @@ void UUVProjectionTool::Shutdown(EToolShutdownType ShutdownType)
 {
 	UVChannelProperties->SaveProperties(this);
 	BasicProperties->SavedDimensions = BasicProperties->Dimensions;
+	BasicProperties->bSavedUniformDimensions = BasicProperties->bUniformDimensions;
 	BasicProperties->SavedTransform = TransformGizmo->GetGizmoTransform();
 	BasicProperties->SaveProperties(this);
-	ExpMapProperties->SaveProperties(this);
-	CylinderProperties->SaveProperties(this);
 	MaterialSettings->SaveProperties(this);
 
 	EdgeRenderer->Disconnect();
@@ -258,19 +247,14 @@ TUniquePtr<FDynamicMeshOperator> UUVProjectionOperatorFactory::MakeNewOperator()
 
 	Op->MeshToProjectionSpace = Tool->WorldTransform;
 	Op->ProjectionBox = Tool->GetProjectionBox();
-	if (Tool->BasicProperties->DimensionMode == EUVProjectionToolDimensionMode::UseFirst)
-	{
-		double UseDim = Op->ProjectionBox.Extents.X;
-		Op->ProjectionBox.Extents = FVector3d(UseDim, UseDim, UseDim);
-	}
-	Op->CylinderSplitAngle = Tool->CylinderProperties->SplitAngle;
-	Op->BlendWeight = Tool->ExpMapProperties->NormalBlending;
-	Op->SmoothingRounds = Tool->ExpMapProperties->SmoothingRounds;
-	Op->SmoothingAlpha = Tool->ExpMapProperties->SmoothingAlpha;
+	Op->CylinderSplitAngle = Tool->BasicProperties->CylinderSplitAngle;
+	Op->BlendWeight = Tool->BasicProperties->ExpMapNormalBlending;
+	Op->SmoothingRounds = Tool->BasicProperties->ExpMapSmoothingSteps;
+	Op->SmoothingAlpha = Tool->BasicProperties->ExpMapSmoothingAlpha;
 
-	Op->UVRotationAngleDeg = Tool->BasicProperties->UVRotation;
-	Op->UVScale = (FVector2f)Tool->BasicProperties->UVScale;
-	Op->UVTranslate = (FVector2f)Tool->BasicProperties->UVTranslate;
+	Op->UVRotationAngleDeg = Tool->BasicProperties->Rotation;
+	Op->UVScale = (FVector2f)Tool->BasicProperties->Scale;
+	Op->UVTranslate = (FVector2f)Tool->BasicProperties->Translation;
 
 	Op->OriginalMesh = Tool->InputMesh;
 	Op->TriangleROI = Tool->TriangleROI;
@@ -319,11 +303,13 @@ void UUVProjectionTool::InitializeMesh()
 
 
 
-UE::Geometry::FOrientedBox3d UUVProjectionTool::GetProjectionBox()
+FOrientedBox3d UUVProjectionTool::GetProjectionBox() const
 {
-	FFrame3d BoxFrame(TransformProxy->GetTransform());
-	FVector3d BoxDimensions = 0.5 * (FVector3d)BasicProperties->Dimensions;
-	return UE::Geometry::FOrientedBox3d(BoxFrame, BoxDimensions);
+	const FFrame3d BoxFrame(TransformProxy->GetTransform());
+	const FVector3d BoxDimensions = 0.5 * (BasicProperties->bUniformDimensions
+		                                       ? FVector3d(BasicProperties->Dimensions.X, BasicProperties->Dimensions.X, BasicProperties->Dimensions.X)
+		                                       : static_cast<FVector3d>(BasicProperties->Dimensions));
+	return FOrientedBox3d(BoxFrame, BoxDimensions);
 }
 
 
@@ -396,13 +382,6 @@ void UUVProjectionTool::OnMaterialSettingsChanged()
 }
 
 
-void UUVProjectionTool::OnProjectionTypeChanged()
-{
-	SetToolPropertySourceEnabled(ExpMapProperties, BasicProperties->ProjectionType == EUVProjectionMethod::ExpMap);
-	SetToolPropertySourceEnabled(CylinderProperties, BasicProperties->ProjectionType == EUVProjectionMethod::Cylinder);
-}
-
-
 void UUVProjectionTool::OnMeshUpdated(UMeshOpPreviewWithBackgroundCompute* PreviewCompute)
 {
 	const FColor UVSeamColor(15, 240, 15);
@@ -411,7 +390,7 @@ void UUVProjectionTool::OnMeshUpdated(UMeshOpPreviewWithBackgroundCompute* Previ
 
 	const FDynamicMesh3* UseMesh = Preview->PreviewMesh->GetMesh();
 	const FDynamicMeshUVOverlay* UVOverlay = UseMesh->Attributes()->GetUVLayer(UVChannelProperties->GetSelectedChannelIndex(true));
-	auto AppendSeamEdge = [UseMesh, UVOverlay, UVSeamColor, UVSeamThickness, UVSeamDepthBias](int32 eid, TArray<FRenderableLine>& LinesOut) {
+	auto AppendSeamEdge = [UseMesh, UVSeamColor, UVSeamThickness, UVSeamDepthBias](int32 eid, TArray<FRenderableLine>& LinesOut) {
 		FVector3d A, B;
 		UseMesh->GetEdgeV(eid, A, B);
 		LinesOut.Add(FRenderableLine((FVector)A, (FVector)B, UVSeamColor, UVSeamThickness, UVSeamDepthBias));
@@ -524,6 +503,8 @@ void UUVProjectionTool::ApplyAction(EUVProjectionToolActions ActionType)
 		case EUVProjectionToolActions::Reset:
 			ApplyAction_Reset();
 			break;
+		case EUVProjectionToolActions::NoAction:
+			break;
 	}
 }
 
@@ -595,6 +576,10 @@ void UUVProjectionTool::ApplyAction_AutoFit(bool bAlign)
 	{
 		BasicProperties->SilentUpdateWatcherAtIndex(DimensionsWatcher);
 	}
+	if (DimensionsModeWatcher >= 0)
+	{
+		BasicProperties->SilentUpdateWatcherAtIndex(DimensionsModeWatcher);
+	}
 
 	// update Gizmo
 	TransformGizmo->SetNewGizmoTransform(AlignFrame.ToFTransform());
diff --git a/Engine/Plugins/Experimental/MeshModelingToolsetExp/Source/MeshModelingToolsExp/Public/UVProjectionTool.h b/Engine/Plugins/Experimental/MeshModelingToolsetExp/Source/MeshModelingToolsExp/Public/UVProjectionTool.h
--- a/Engine/Plugins/Experimental/MeshModelingToolsetExp/Source/MeshModelingToolsExp/Public/UVProjectionTool.h
+++ b/Engine/Plugins/Experimental/MeshModelingToolsetExp/Source/MeshModelingToolsExp/Public/UVProjectionTool.h
@@ -16,12 +16,13 @@
 #include "UVProjectionTool.generated.h"
 
 
-// predeclarations
+// Forward declarations
 struct FMeshDescription;
 class UDynamicMeshComponent;
 class UCombinedTransformGizmo;
 class UTransformProxy;
 class USingleClickInputBehavior;
+class UUVProjectionTool;
 
 /**
  *
@@ -37,8 +38,6 @@ public:
 };
 
 
-
-
 UENUM()
 enum class EUVProjectionToolActions
 {
@@ -60,21 +59,21 @@ public:
 	void Initialize(UUVProjectionTool* ParentToolIn) { ParentTool = ParentToolIn; }
 	void PostAction(EUVProjectionToolActions Action);
 
-	/** Automatically fit the Projection Dimensions based on the current projection orientation */
+	/** Automatically fit the UV Projection Dimensions based on the current projection orientation */
 	UFUNCTION(CallInEditor, Category = Actions, meta = (DisplayName = "AutoFit", DisplayPriority = 1))
 	void AutoFit()
 	{
 		PostAction(EUVProjectionToolActions::AutoFit);
 	}
 
-	/** Automatically align the projection orientation and then automatically fit the Dimensions */
+	/** Automatically orient the projection and then automatically fit the UV Projection Dimensions */
 	UFUNCTION(CallInEditor, Category = Actions, meta = (DisplayName = "AutoFitAlign", DisplayPriority = 2))
 	void AutoFitAlign()
 	{
 		PostAction(EUVProjectionToolActions::AutoFitAlign);
 	}
 
-	/** Re-initialize the Projection based on the Initialization setting */
+	/** Re-initialize the projection based on the UV Projection Initialization property */
 	UFUNCTION(CallInEditor, Category = Actions, meta = (DisplayName = "Reset", DisplayPriority = 3))
 	void Reset()
 	{
@@ -83,33 +82,20 @@ public:
 };
 
 
-
 UENUM()
 enum class EUVProjectionToolInitializationMode
 {
 	/** Initialize projection to bounding box center */
 	Default,
-	/** Initialize projection based on previous usage of the UV Projection Tool */
+	/** Initialize projection based on previous usage of the Project tool */
 	UsePrevious,
-	/** Initialize projection using AutoFitting, for the initial projection type */
+	/** Initialize projection using Auto Fitting for the initial projection type */
 	AutoFit,
-	/** Initialize projection using AutoFitting with Alignment, for the initial projection type */
+	/** Initialize projection using Auto Fitting with Alignment for the initial projection type */
 	AutoFitAlign
 };
 
 
-UENUM()
-enum class EUVProjectionToolDimensionMode
-{
-	/** All Dimensions of the projection are used */
-	AllFree,
-	/** The Dimensions.X value is used to define all projection Dimensions */
-	UseFirst
-};
-
-
-
-
 /**
  * Standard properties
  */
@@ -119,89 +105,79 @@ class MESHMODELINGTOOLSEXP_API UUVProjectionToolProperties : public UInteractive
 	GENERATED_BODY()
 
 public:
-	/** Choose the UV projection method (cube, cylinder, plane) */
-	UPROPERTY(EditAnywhere, Category = ProjectionSettings)
+	/** Shape and/or algorithm to use for UV projection */
+	UPROPERTY(EditAnywhere, Category = "UV Projection")
 	EUVProjectionMethod ProjectionType = EUVProjectionMethod::Plane;
 
-	/** Per-axis Dimensions of the Projection. Z is height of box/cylinder. */
-	UPROPERTY(EditAnywhere, Category = ProjectionSettings)
+	/** Width, length, and height of the projection shape before rotation */
+	UPROPERTY(EditAnywhere, Category = "UV Projection")
 	FVector Dimensions = FVector(100.0f, 100.0f, 100.0f);
 
-	/** How the Dimensions are interpreted for Box/Plane/ExpMap Projections */
-	UPROPERTY(EditAnywhere, Category = ProjectionSettings)
-	EUVProjectionToolDimensionMode DimensionMode = EUVProjectionToolDimensionMode::AllFree;
+	/** Only use the Dimensions X value to uniformly define all projection shape dimensions */
+	UPROPERTY(EditAnywhere, Category = "UV Projection")
+	bool bUniformDimensions = false;
 
-	/** Determines how projection settings will be initialized. Only affects Tool until changes are made to Dimensions/Position. */
-	UPROPERTY(EditAnywhere, Category = ProjectionSettings)
+	/** Determines how projection settings will be initialized; this only takes effect if the projection shape dimensions or position are unchanged */
+	UPROPERTY(EditAnywhere, Category = "UV Projection")
 	EUVProjectionToolInitializationMode Initialization = EUVProjectionToolInitializationMode::Default;
 
 	//
-	// UV-space transform options
+	// Cylinder projection options
 	//
 
-	/** UV-Space Rotation Angle in Degrees, applied after computing projected UVs */
-	UPROPERTY(EditAnywhere, Category = "UV Transform", meta = (DisplayName = "UV Rotation") )
-	float UVRotation = 0.0;
-
-	/** UV-Space Scaling applied after computing projected UVs */
-	UPROPERTY(EditAnywhere, Category = "UV Transform", meta = (DisplayName = "UV Scale"))
-	FVector2D UVScale = FVector2D::UnitVector;
-
-	/** UV-Space Translation applied after computing projected UVs */
-	UPROPERTY(EditAnywhere, Category = "UV Transform", meta = (DisplayName = "UV Translation"))
-	FVector2D UVTranslate = FVector2D::ZeroVector;
-
+	/** Angle in degrees to determine whether faces should be assigned to the cylinder or the flat end caps */
+	UPROPERTY(EditAnywhere, Category = CylinderProjection, meta = (DisplayName = "Split Angle", UIMin = "0", UIMax = "90",
+		EditCondition = "ProjectionType == EUVProjectionMethod::Cylinder", EditConditionHides))
+	float CylinderSplitAngle = 45.0f;
 
 	//
-	// Saved State. These are used internally to support UsePrevious initialization mode
+	// ExpMap projection options
 	//
-	UPROPERTY()
-	FVector SavedDimensions = FVector::ZeroVector;
-
-	UPROPERTY()
-	FTransform SavedTransform;
-};
 
+	/** Blend between surface normals and projection normal; ExpMap projection becomes Plane projection when this value is 1 */
+	UPROPERTY(EditAnywhere, Category = "ExpMap Projection", meta = (DisplayName = "Normal Blending", UIMin = "0", UIMax = "1",
+		EditCondition = "ProjectionType == EUVProjectionMethod::ExpMap", EditConditionHides))
+	float ExpMapNormalBlending = 0.0f;
 
+	/** Number of smoothing steps to apply; this slightly increases distortion but produces more stable results. */
+	UPROPERTY(EditAnywhere, Category = "ExpMap Projection", meta = (DisplayName = "Smoothing Steps", UIMin = "0", UIMax = "100",
+		EditCondition = "ProjectionType == EUVProjectionMethod::ExpMap", EditConditionHides))
+	int ExpMapSmoothingSteps = 0;
 
+	/** Smoothing parameter, larger values result in faster smoothing in each step. */
+	UPROPERTY(EditAnywhere, Category = "ExpMap Projection", meta = (DisplayName = "Smoothing Alpha", UIMin = "0", UIMax = "1",
+		EditCondition = "ProjectionType == EUVProjectionMethod::ExpMap", EditConditionHides))
+	float ExpMapSmoothingAlpha = 0.25f;
 
-/**
- * ExpMap Projection properties
- */
-UCLASS()
-class MESHMODELINGTOOLSEXP_API UUVProjectionToolExpMapProperties : public UInteractiveToolPropertySet
-{
-	GENERATED_BODY()
-public:
-	/** Blend between surface normals and projection normal - ExpMap becomes Planar Projection when value is 1 */
-	UPROPERTY(EditAnywhere, Category = ExpMap, meta = (UIMin = "0", UIMax = "1"))
-	float NormalBlending = 0.0;
+	//
+	// UV-space transform options
+	//
 
-	/** Rounds of Smoothing to apply to surface normals before computing ExpMap */
-	UPROPERTY(EditAnywhere, Category = ExpMap, AdvancedDisplay, meta = (UIMin = "0", UIMax = "100"))
-	int SmoothingRounds = 0;
+	/** Rotation in degrees applied after computing projection */
+	UPROPERTY(EditAnywhere, Category = "UV Transform")
+	float Rotation = 0.0;
 
-	/** Smoothing Strength for ExpMap Smoothing */
-	UPROPERTY(EditAnywhere, Category = ExpMap, AdvancedDisplay, meta = (UIMin = "0", UIMax = "1"))
-	float SmoothingAlpha = 0.25f;
-};
+	/** Scaling applied after computing projection */
+	UPROPERTY(EditAnywhere, Category = "UV Transform")
+	FVector2D Scale = FVector2D::UnitVector;
 
+	/** Translation applied after computing projection */
+	UPROPERTY(EditAnywhere, Category = "UV Transform")
+	FVector2D Translation = FVector2D::ZeroVector;
 
+	//
+	// Saved State. These are used internally to support UsePrevious initialization mode
+	//
 
-/**
- * Cylinder Projection properties
- */
-UCLASS()
-class MESHMODELINGTOOLSEXP_API UUVProjectionToolCylinderProperties : public UInteractiveToolPropertySet
-{
-	GENERATED_BODY()
-public:
-	/** Angle used to determine whether faces should be assigned to cylinder or flat endcaps */
-	UPROPERTY(EditAnywhere, Category = Cylinder, meta = (UIMin = "0", UIMax = "90"))
-	float SplitAngle = 45.0f;
-};
+	UPROPERTY()
+	FVector SavedDimensions = FVector::ZeroVector;
 
+	UPROPERTY()
+	bool bSavedUniformDimensions = false;
 
+	UPROPERTY()
+	FTransform SavedTransform;
+};
 
 
 /**
@@ -222,7 +198,7 @@ public:
 
 
 /**
- * Simple Mesh Normal Updating Tool
+ * UV projection tool
  */
 UCLASS()
 class MESHMODELINGTOOLSEXP_API UUVProjectionTool : public USingleSelectionMeshEditingTool
@@ -253,12 +229,6 @@ protected:
 	UPROPERTY()
 	TObjectPtr<UUVProjectionToolProperties> BasicProperties = nullptr;
 
-	UPROPERTY()
-	TObjectPtr<UUVProjectionToolExpMapProperties> ExpMapProperties = nullptr;
-
-	UPROPERTY()
-	TObjectPtr<UUVProjectionToolCylinderProperties> CylinderProperties = nullptr;
-
 	UPROPERTY()
 	TObjectPtr<UUVProjectionToolEditActions> EditActions = nullptr;
 
@@ -268,7 +238,6 @@ protected:
 	UPROPERTY()
 	TObjectPtr<UMeshOpPreviewWithBackgroundCompute> Preview = nullptr;
 
-
 	UPROPERTY()
 	TObjectPtr<UMaterialInstanceDynamic> CheckerMaterial = nullptr;
 
@@ -284,7 +253,6 @@ protected:
 	UPROPERTY()
 	TObjectPtr<UPreviewGeometry> EdgeRenderer = nullptr;
 
-protected:
 	TSharedPtr<UE::Geometry::FDynamicMesh3, ESPMode::ThreadSafe> InputMesh;
 	TSharedPtr<TArray<int32>, ESPMode::ThreadSafe> TriangleROI;
 	TSharedPtr<TArray<int32>, ESPMode::ThreadSafe> VertexROI;
@@ -295,8 +263,10 @@ protected:
 	UE::Geometry::FAxisAlignedBox3d WorldBounds;
 
 	FVector InitialDimensions;
+	bool bInitialUniformDimensions;
 	FTransform InitialTransform;
 	int32 DimensionsWatcher = -1;
+	int32 DimensionsModeWatcher = -1;
 	bool bTransformModified = false;
 	void OnInitializationModeChanged();
 	void ApplyInitializationMode();
@@ -310,11 +280,10 @@ protected:
 
 	void TransformChanged(UTransformProxy* Proxy, FTransform Transform);
 
-	void OnProjectionTypeChanged();
 	void OnMaterialSettingsChanged();
 	void OnMeshUpdated(UMeshOpPreviewWithBackgroundCompute* PreviewCompute);
 
-	UE::Geometry::FOrientedBox3d GetProjectionBox();
+	UE::Geometry::FOrientedBox3d GetProjectionBox() const;
 
 
 	//
diff --git a/Engine/Plugins/Runtime/MeshModelingToolset/Source/ModelingOperators/Public/ParameterizationOps/UVProjectionOp.h b/Engine/Plugins/Runtime/MeshModelingToolset/Source/ModelingOperators/Public/ParameterizationOps/UVProjectionOp.h
--- a/Engine/Plugins/Runtime/MeshModelingToolset/Source/ModelingOperators/Public/ParameterizationOps/UVProjectionOp.h
+++ b/Engine/Plugins/Runtime/MeshModelingToolset/Source/ModelingOperators/Public/ParameterizationOps/UVProjectionOp.h
@@ -12,14 +12,14 @@
 UENUM()
 enum class EUVProjectionMethod : uint8
 {
-	/** Calculate UVs by assigning triangles to faces of a Box, and then apply per-box-face planar projection */
+	/** Calculate UVs by assigning triangles to faces of a box, and then apply per-box-face planar projection */
 	Box,
-	/** Calculate UVs by assigning triangles to Cylinder (radial projection) or Endcaps (planar projection) */
+	/** Calculate UVs by assigning triangles to a cylinder using radial projection or cylinder end caps using planar projection */
 	Cylinder,
-	/** Calculate UVs by projecting to a 3D plane */
+	/** Calculate UVs by projecting to a plane */
 	Plane,
-	/** Calculate UVs by Exponential Map Projection centered at nearest surface point to input 3D plane */
-	ExpMap
+	/** Calculate UVs by Exponential Map projection centered at the surface point nearest to the center of the input plane */
+	ExpMap UMETA(DisplayName = "ExpMap")
 };
 
 namespace UE
