ModelingTools:
Reduce surface area of MeshDescriptionProvider/Committer, replace with UE::ToolTarget:: calls where possible.

Add new UE::ToolTarget::CommitMeshDescriptionUpdateViaDynamicMesh() function. This is being used for now to avoid potential regressions as UE::ToolTarget::CommitDynamicMeshUpdate will preferentially use DynamicMeshCommitter, and I am not certain it is functionally equivalent in all cases.
Add new UE::ToolTarget::CommitDynamicMeshNormalsUpdate(), similar to existing UV version
Add new Move-variant of UE::ToolTarget::CommitMeshDescriptionUpdate(), uses new Move-variant of IMeshDescriptionCommitter::CommitMeshDescription.
Make existing IMeshDescriptionCommitter::CommitMeshDescription callback interface protected, to prevent usage of this function at public API level (will be removed in future).

Tool updates should not change, just using cleaner APIs.
EditNormalsTool now uses CommitDynamicMeshNormalsUpdate(), which does go via DynamicMeshCommitter preferentially, where it previously went via MeshDescriptionCommitter. In light testing the results appear equivalent.
AttributeEditorTool now operates on MeshDescription copies in various update functions. These are not performance-critical.

#rnx

diff --git a/Engine/Plugins/Experimental/HairModelingToolset/Source/HairModelingToolset/Private/GenerateLODMeshesTool.cpp b/Engine/Plugins/Experimental/HairModelingToolset/Source/HairModelingToolset/Private/GenerateLODMeshesTool.cpp
--- a/Engine/Plugins/Experimental/HairModelingToolset/Source/HairModelingToolset/Private/GenerateLODMeshesTool.cpp
+++ b/Engine/Plugins/Experimental/HairModelingToolset/Source/HairModelingToolset/Private/GenerateLODMeshesTool.cpp
@@ -21,8 +21,8 @@
 #include "Modules/ModuleManager.h"
 
 #include "TargetInterfaces/MaterialProvider.h"
-#include "TargetInterfaces/MeshDescriptionProvider.h"
 #include "TargetInterfaces/PrimitiveComponentBackedTarget.h"
+#include "ModelingToolTargetUtil.h"
 
 #include "IMeshReductionManagerModule.h"
 #include "IMeshReductionInterfaces.h"
@@ -99,15 +99,15 @@ void UGenerateLODMeshesTool::Setup()
 #else
 		auto EnterProgressFrame = [](int Progress) {};
 #endif
-		IMeshDescriptionProvider* TargetMeshProvider = Cast<IMeshDescriptionProvider>(Target);
-		OriginalMeshDescription = MakeShared<FMeshDescription, ESPMode::ThreadSafe>(*TargetMeshProvider->GetMeshDescription());
+		const FMeshDescription* TargetMeshDescription = UE::ToolTarget::GetMeshDescription(Target);
+		OriginalMeshDescription = MakeShared<FMeshDescription, ESPMode::ThreadSafe>(*TargetMeshDescription);
 		EnterProgressFrame(1);
 		// aux-data isn't deep copied - by default it is built during initial evaluation (not thread safe)  
 		// so force aux-data rebuild now before multiple UE4 simplifiers try to use it in parallel.
 		OriginalMeshDescription->BuildIndexers();   
 		OriginalMesh = MakeShared<FDynamicMesh3, ESPMode::ThreadSafe>();
 		FMeshDescriptionToDynamicMesh Converter;
-		Converter.Convert(TargetMeshProvider->GetMeshDescription(), *OriginalMesh);
+		Converter.Convert(TargetMeshDescription, *OriginalMesh);
 		EnterProgressFrame(2);
 		OriginalMeshSpatial = MakeShared<FDynamicMeshAABBTree3, ESPMode::ThreadSafe>(OriginalMesh.Get(), true);
 	}
diff --git a/Engine/Plugins/Experimental/MeshModelingToolsetExp/Source/MeshModelingToolsEditorOnly/Private/AttributeEditorTool.cpp b/Engine/Plugins/Experimental/MeshModelingToolsetExp/Source/MeshModelingToolsEditorOnly/Private/AttributeEditorTool.cpp
--- a/Engine/Plugins/Experimental/MeshModelingToolsetExp/Source/MeshModelingToolsEditorOnly/Private/AttributeEditorTool.cpp
+++ b/Engine/Plugins/Experimental/MeshModelingToolsetExp/Source/MeshModelingToolsEditorOnly/Private/AttributeEditorTool.cpp
@@ -20,9 +20,7 @@
 #include "Components/PrimitiveComponent.h"
 
 #include "TargetInterfaces/MeshDescriptionCommitter.h"
-#include "TargetInterfaces/MeshDescriptionProvider.h"
 #include "TargetInterfaces/PrimitiveComponentBackedTarget.h"
-#include "ToolTargetManager.h"
 #include "ModelingToolTargetUtil.h"
 
 using namespace UE::Geometry;
@@ -435,19 +433,20 @@ void UAttributeEditorTool::ClearNormals()
 
 	for (int32 ComponentIdx = 0; ComponentIdx < Targets.Num(); ComponentIdx++)
 	{
-		Cast<IMeshDescriptionCommitter>(Targets[ComponentIdx])->CommitMeshDescription([&](const IMeshDescriptionCommitter::FCommitterParams& CommitParams)
+		FMeshDescription EditedMesh = UE::ToolTarget::GetMeshDescriptionCopy(Targets[ComponentIdx]);
+
+		TEdgeAttributesRef<bool> EdgeHardnesses = EditedMesh.EdgeAttributes().GetAttributesRef<bool>(MeshAttribute::Edge::IsHard);
+		if (EdgeHardnesses.IsValid())
 		{
-			TEdgeAttributesRef<bool> EdgeHardnesses = CommitParams.MeshDescriptionOut->EdgeAttributes().GetAttributesRef<bool>(MeshAttribute::Edge::IsHard);
-			if (EdgeHardnesses.IsValid())
+			for (FEdgeID ElID : EditedMesh.Edges().GetElementIDs())
 			{
-				for (FEdgeID ElID : CommitParams.MeshDescriptionOut->Edges().GetElementIDs())
-				{
-					EdgeHardnesses[ElID] = false;
-				}
+				EdgeHardnesses[ElID] = false;
 			}
-			FStaticMeshOperations::ComputeTriangleTangentsAndNormals(*CommitParams.MeshDescriptionOut, FMathf::Epsilon);
-			FStaticMeshOperations::RecomputeNormalsAndTangentsIfNeeded(*CommitParams.MeshDescriptionOut, EComputeNTBsFlags::WeightedNTBs | EComputeNTBsFlags::Normals);
-		});
+		}
+		FStaticMeshOperations::ComputeTriangleTangentsAndNormals(EditedMesh, FMathf::Epsilon);
+		FStaticMeshOperations::RecomputeNormalsAndTangentsIfNeeded(EditedMesh, EComputeNTBsFlags::WeightedNTBs | EComputeNTBsFlags::Normals);
+
+		UE::ToolTarget::CommitMeshDescriptionUpdate(Targets[ComponentIdx], MoveTemp(EditedMesh));
 	}
 	GetToolManager()->EndUndoTransaction();
 }
@@ -460,28 +459,29 @@ void UAttributeEditorTool::ClearUVs()
 	GetToolManager()->BeginUndoTransaction(LOCTEXT("ClearUVsTransactionMessage", "Clear Selected UVs"));
 	for (int32 ComponentIdx = 0; ComponentIdx < Targets.Num(); ComponentIdx++)
 	{
-		Cast<IMeshDescriptionCommitter>(Targets[ComponentIdx])->CommitMeshDescription([&](const IMeshDescriptionCommitter::FCommitterParams& CommitParams)
+		FMeshDescription EditedMesh = UE::ToolTarget::GetMeshDescriptionCopy(Targets[ComponentIdx]);
+
+		TVertexInstanceAttributesRef<FVector2f> InstanceUVs =
+			EditedMesh.VertexInstanceAttributes().GetAttributesRef<FVector2f>(MeshAttribute::VertexInstance::TextureCoordinate);
+		int32 NumChannels = InstanceUVs.GetNumChannels();
+		const FVertexInstanceArray& Instances = EditedMesh.VertexInstances();
+		for (int LayerIndex = NumChannels-1; LayerIndex >= 0; LayerIndex--)
 		{
-			TVertexInstanceAttributesRef<FVector2f> InstanceUVs =
-				CommitParams.MeshDescriptionOut->VertexInstanceAttributes().GetAttributesRef<FVector2f>(MeshAttribute::VertexInstance::TextureCoordinate);
-			int32 NumChannels = InstanceUVs.GetNumChannels();
-			const FVertexInstanceArray& Instances = CommitParams.MeshDescriptionOut->VertexInstances();
-			for (int LayerIndex = NumChannels-1; LayerIndex >= 0; LayerIndex--)
+			if (!FStaticMeshOperations::RemoveUVChannel(EditedMesh, LayerIndex))
 			{
-				if (!FStaticMeshOperations::RemoveUVChannel(*CommitParams.MeshDescriptionOut, LayerIndex))
+				for (const FVertexInstanceID& ElID : Instances.GetElementIDs())
 				{
-					for (const FVertexInstanceID& ElID : Instances.GetElementIDs())
-					{
-						InstanceUVs.Set(ElID, LayerIndex, FVector2D::ZeroVector);
-					}
+					InstanceUVs.Set(ElID, LayerIndex, FVector2D::ZeroVector);
 				}
 			}
+		}
 
-			if (bHaveAutoGeneratedLightmapUVSet)
-			{
-				UpdateAutoGeneratedLightmapUVChannel(UE::ToolTarget::GetTargetComponent(Targets[ComponentIdx]), InstanceUVs.GetNumChannels());
-			}
-		});
+		UE::ToolTarget::CommitMeshDescriptionUpdate(Targets[ComponentIdx], MoveTemp(EditedMesh));
+
+		if (bHaveAutoGeneratedLightmapUVSet)
+		{
+			UpdateAutoGeneratedLightmapUVChannel(UE::ToolTarget::GetTargetComponent(Targets[ComponentIdx]), InstanceUVs.GetNumChannels());
+		}
 	}
 
 	EmitAttributesChange();
@@ -511,24 +511,25 @@ void UAttributeEditorTool::DeleteSelectedUVSet()
 	GetToolManager()->BeginUndoTransaction(LOCTEXT("ClearUVsTransactionMessage", "Clear Selected UVs"));
 	for (int32 ComponentIdx = 0; ComponentIdx < Targets.Num(); ComponentIdx++)
 	{
-		Cast<IMeshDescriptionCommitter>(Targets[ComponentIdx])->CommitMeshDescription([&](const IMeshDescriptionCommitter::FCommitterParams& CommitParams)
+		FMeshDescription EditedMesh = UE::ToolTarget::GetMeshDescriptionCopy(Targets[ComponentIdx]);
+
+		TVertexInstanceAttributesRef<FVector2f> InstanceUVs =
+			EditedMesh.VertexInstanceAttributes().GetAttributesRef<FVector2f>(MeshAttribute::VertexInstance::TextureCoordinate);
+		if (!FStaticMeshOperations::RemoveUVChannel(EditedMesh, DeleteIndex))
 		{
-			TVertexInstanceAttributesRef<FVector2f> InstanceUVs =
-				CommitParams.MeshDescriptionOut->VertexInstanceAttributes().GetAttributesRef<FVector2f>(MeshAttribute::VertexInstance::TextureCoordinate);
-			if (!FStaticMeshOperations::RemoveUVChannel(*CommitParams.MeshDescriptionOut, DeleteIndex))
+			const FVertexInstanceArray& Instances = EditedMesh.VertexInstances();
+			for (const FVertexInstanceID& InstanceID : Instances.GetElementIDs())
 			{
-				const FVertexInstanceArray& Instances = CommitParams.MeshDescriptionOut->VertexInstances();
-				for (const FVertexInstanceID& InstanceID : Instances.GetElementIDs())
-				{
-					InstanceUVs.Set(InstanceID, DeleteIndex, FVector2D::ZeroVector);
-				}
+				InstanceUVs.Set(InstanceID, DeleteIndex, FVector2D::ZeroVector);
 			}
+		}
 
-			if (bHaveAutoGeneratedLightmapUVSet)
-			{
-				UpdateAutoGeneratedLightmapUVChannel(UE::ToolTarget::GetTargetComponent(Targets[ComponentIdx]), InstanceUVs.GetNumChannels());
-			}
-		});
+		UE::ToolTarget::CommitMeshDescriptionUpdate(Targets[ComponentIdx], MoveTemp(EditedMesh));
+
+		if (bHaveAutoGeneratedLightmapUVSet)
+		{
+			UpdateAutoGeneratedLightmapUVChannel(UE::ToolTarget::GetTargetComponent(Targets[ComponentIdx]), InstanceUVs.GetNumChannels());
+		}
 	}
 
 	EmitAttributesChange();
@@ -544,25 +545,25 @@ void UAttributeEditorTool::AddUVSet()
 	GetToolManager()->BeginUndoTransaction(LOCTEXT("AddUVSetMessage", "Add UV Set"));
 	for (int32 ComponentIdx = 0; ComponentIdx < Targets.Num(); ComponentIdx++)
 	{
-		Cast<IMeshDescriptionCommitter>(Targets[ComponentIdx])->CommitMeshDescription([&](const IMeshDescriptionCommitter::FCommitterParams& CommitParams)
+		FMeshDescription EditedMesh = UE::ToolTarget::GetMeshDescriptionCopy(Targets[ComponentIdx]);
+
+		TVertexInstanceAttributesRef<FVector2f> InstanceUVs =
+			EditedMesh.VertexInstanceAttributes().GetAttributesRef<FVector2f>(MeshAttribute::VertexInstance::TextureCoordinate);
+		int32 NewChannelIndex = InstanceUVs.GetNumChannels();
+		if (FStaticMeshOperations::AddUVChannel(EditedMesh))
 		{
-			TVertexInstanceAttributesRef<FVector2f> InstanceUVs =
-				CommitParams.MeshDescriptionOut->VertexInstanceAttributes().GetAttributesRef<FVector2f>(MeshAttribute::VertexInstance::TextureCoordinate);
-			int32 NewChannelIndex = InstanceUVs.GetNumChannels();
-			if (!FStaticMeshOperations::AddUVChannel(*CommitParams.MeshDescriptionOut))
-			{
-				GetToolManager()->DisplayMessage(LOCTEXT("FailedToAddUVSet", "Adding UV Set Failed"), EToolMessageLevel::UserWarning);
-			}
-			else
-			{
-				GetToolManager()->DisplayMessage(FText::Format(LOCTEXT("AddedNewUVSet", "Added UV{0}"), FText::FromString(FString::FromInt(NewChannelIndex))), EToolMessageLevel::UserWarning);
+			GetToolManager()->DisplayMessage(FText::Format(LOCTEXT("AddedNewUVSet", "Added UV{0}"), FText::FromString(FString::FromInt(NewChannelIndex))), EToolMessageLevel::UserWarning);
+			UE::ToolTarget::CommitMeshDescriptionUpdate(Targets[ComponentIdx], MoveTemp(EditedMesh));
 
-				if (bHaveAutoGeneratedLightmapUVSet)
-				{
-					UpdateAutoGeneratedLightmapUVChannel(UE::ToolTarget::GetTargetComponent(Targets[ComponentIdx]), InstanceUVs.GetNumChannels());
-				}
+			if (bHaveAutoGeneratedLightmapUVSet)
+			{
+				UpdateAutoGeneratedLightmapUVChannel(UE::ToolTarget::GetTargetComponent(Targets[ComponentIdx]), InstanceUVs.GetNumChannels());
 			}
-		});
+		}
+		else
+		{
+			GetToolManager()->DisplayMessage(LOCTEXT("FailedToAddUVSet", "Adding UV Set Failed"), EToolMessageLevel::UserWarning);
+		}
 	}
 
 	EmitAttributesChange();
@@ -586,33 +587,34 @@ void UAttributeEditorTool::DuplicateSelectedUVSet()
 	GetToolManager()->BeginUndoTransaction(LOCTEXT("DuplicateUVSetMessage", "Duplicate UV Set"));
 	for (int32 ComponentIdx = 0; ComponentIdx < Targets.Num(); ComponentIdx++)
 	{
-		Cast<IMeshDescriptionCommitter>(Targets[ComponentIdx])->CommitMeshDescription([&](const IMeshDescriptionCommitter::FCommitterParams& CommitParams)
+		FMeshDescription EditedMesh = UE::ToolTarget::GetMeshDescriptionCopy(Targets[ComponentIdx]);
+
+		TVertexInstanceAttributesRef<FVector2f> InstanceUVs =
+			EditedMesh.VertexInstanceAttributes().GetAttributesRef<FVector2f>(MeshAttribute::VertexInstance::TextureCoordinate);
+		int32 NewChannelIndex = InstanceUVs.GetNumChannels();
+		if (FStaticMeshOperations::AddUVChannel(EditedMesh))
 		{
-			TVertexInstanceAttributesRef<FVector2f> InstanceUVs =
-				CommitParams.MeshDescriptionOut->VertexInstanceAttributes().GetAttributesRef<FVector2f>(MeshAttribute::VertexInstance::TextureCoordinate);
-			int32 NewChannelIndex = InstanceUVs.GetNumChannels();
-			if (!FStaticMeshOperations::AddUVChannel(*CommitParams.MeshDescriptionOut))
+			const FVertexInstanceArray& Instances = EditedMesh.VertexInstances();
+			for (const FVertexInstanceID& InstanceID : Instances.GetElementIDs())
 			{
-				GetToolManager()->DisplayMessage(LOCTEXT("FailedToAddUVSet", "Adding UV Set Failed"), EToolMessageLevel::UserWarning);
+				FVector2D SourceUV = InstanceUVs.Get(InstanceID, SourceIndex);
+				InstanceUVs.Set(InstanceID, NewChannelIndex, SourceUV);
 			}
-			else
-			{
-				const FVertexInstanceArray& Instances = CommitParams.MeshDescriptionOut->VertexInstances();
-				for (const FVertexInstanceID& InstanceID : Instances.GetElementIDs())
-				{
-					FVector2D SourceUV = InstanceUVs.Get(InstanceID, SourceIndex);
-					InstanceUVs.Set(InstanceID, NewChannelIndex, SourceUV);
-				}
 
-				if (bHaveAutoGeneratedLightmapUVSet)
-				{
-					UpdateAutoGeneratedLightmapUVChannel(UE::ToolTarget::GetTargetComponent(Targets[ComponentIdx]), InstanceUVs.GetNumChannels());
-				}
+			UE::ToolTarget::CommitMeshDescriptionUpdate(Targets[ComponentIdx], MoveTemp(EditedMesh));
 
-				GetToolManager()->DisplayMessage(FText::Format(LOCTEXT("Copied UV Set", "Copied UV{0} to UV{1}"), 
-					FText::FromString(FString::FromInt(SourceIndex)), FText::FromString(FString::FromInt(NewChannelIndex))), EToolMessageLevel::UserWarning);
+			if (bHaveAutoGeneratedLightmapUVSet)
+			{
+				UpdateAutoGeneratedLightmapUVChannel(UE::ToolTarget::GetTargetComponent(Targets[ComponentIdx]), InstanceUVs.GetNumChannels());
 			}
-		});
+
+			GetToolManager()->DisplayMessage(FText::Format(LOCTEXT("Copied UV Set", "Copied UV{0} to UV{1}"), 
+				FText::FromString(FString::FromInt(SourceIndex)), FText::FromString(FString::FromInt(NewChannelIndex))), EToolMessageLevel::UserWarning);
+		}
+		else
+		{
+			GetToolManager()->DisplayMessage(LOCTEXT("FailedToAddUVSet", "Adding UV Set Failed"), EToolMessageLevel::UserWarning);
+		}
 	}
 
 	EmitAttributesChange();
@@ -786,8 +788,9 @@ void UAttributeEditorTool::ResetLightmapUVsChannels()
 	GetToolManager()->BeginUndoTransaction(LOCTEXT("ResetLightmapUVs", "Reset Lightmap UVs"));
 	for (int32 ComponentIdx = 0; ComponentIdx < Targets.Num(); ComponentIdx++)
 	{
+		const FMeshDescription* SourceMesh = UE::ToolTarget::GetMeshDescription(Targets[ComponentIdx]);
 		TVertexInstanceAttributesConstRef<FVector2f> InstanceUVs = 
-			Cast<IMeshDescriptionProvider>(Targets[ComponentIdx])->GetMeshDescription()->VertexInstanceAttributes().GetAttributesRef<FVector2f>(MeshAttribute::VertexInstance::TextureCoordinate);
+			SourceMesh->VertexInstanceAttributes().GetAttributesRef<FVector2f>(MeshAttribute::VertexInstance::TextureCoordinate);
 		int32 SetChannel = FMath::Max(InstanceUVs.GetNumChannels(), 1);
 
 		UStaticMeshComponent* StaticMeshComponent = Cast<UStaticMeshComponent>(UE::ToolTarget::GetTargetComponent(Targets[ComponentIdx]));
diff --git a/Engine/Plugins/Experimental/MeshModelingToolsetExp/Source/MeshModelingToolsEditorOnly/Private/PolygonOnMeshTool.cpp b/Engine/Plugins/Experimental/MeshModelingToolsetExp/Source/MeshModelingToolsEditorOnly/Private/PolygonOnMeshTool.cpp
--- a/Engine/Plugins/Experimental/MeshModelingToolsetExp/Source/MeshModelingToolsEditorOnly/Private/PolygonOnMeshTool.cpp
+++ b/Engine/Plugins/Experimental/MeshModelingToolsetExp/Source/MeshModelingToolsEditorOnly/Private/PolygonOnMeshTool.cpp
@@ -16,9 +16,8 @@
 #include "DynamicMeshToMeshDescription.h"
 
 #include "TargetInterfaces/MaterialProvider.h"
-#include "TargetInterfaces/MeshDescriptionCommitter.h"
-#include "TargetInterfaces/MeshDescriptionProvider.h"
 #include "TargetInterfaces/PrimitiveComponentBackedTarget.h"
+#include "ModelingToolTargetUtil.h"
 
 using namespace UE::Geometry;
 
@@ -130,7 +129,7 @@ void UPolygonOnMeshTool::Setup()
 	// Convert input mesh description to dynamic mesh
 	OriginalDynamicMesh = MakeShared<FDynamicMesh3, ESPMode::ThreadSafe>();
 	FMeshDescriptionToDynamicMesh Converter;
-	Converter.Convert(Cast<IMeshDescriptionProvider>(Target)->GetMeshDescription(), *OriginalDynamicMesh);
+	Converter.Convert( UE::ToolTarget::GetMeshDescription(Target), *OriginalDynamicMesh);
 	// TODO: consider adding an AABB tree construction here?  tradeoff vs doing a raycast against full every time a param change happens ...
 
 	LastDrawnPolygon = FPolygon2d();
@@ -406,11 +405,7 @@ void UPolygonOnMeshTool::GenerateAsset(const TArray<FDynamicMeshOpResult>& Resul
 	
 	check(Results.Num() > 0);
 	check(Results[0].Mesh.Get() != nullptr);
-	Cast<IMeshDescriptionCommitter>(Target)->CommitMeshDescription([&Results](const IMeshDescriptionCommitter::FCommitterParams& CommitParams)
-	{
-		FDynamicMeshToMeshDescription Converter;
-		Converter.Convert(Results[0].Mesh.Get(), *CommitParams.MeshDescriptionOut);
-	});
+	UE::ToolTarget::CommitMeshDescriptionUpdateViaDynamicMesh(Target, *Results[0].Mesh.Get(), true);
 
 	GetToolManager()->EndUndoTransaction();
 }
diff --git a/Engine/Plugins/Experimental/MeshModelingToolsetExp/Source/MeshModelingToolsEditorOnly/Private/SubdividePolyTool.cpp b/Engine/Plugins/Experimental/MeshModelingToolsetExp/Source/MeshModelingToolsEditorOnly/Private/SubdividePolyTool.cpp
--- a/Engine/Plugins/Experimental/MeshModelingToolsetExp/Source/MeshModelingToolsEditorOnly/Private/SubdividePolyTool.cpp
+++ b/Engine/Plugins/Experimental/MeshModelingToolsetExp/Source/MeshModelingToolsEditorOnly/Private/SubdividePolyTool.cpp
@@ -13,9 +13,8 @@
 #include "Drawing/PreviewGeometryActor.h"
 
 #include "TargetInterfaces/MaterialProvider.h"
-#include "TargetInterfaces/MeshDescriptionCommitter.h"
-#include "TargetInterfaces/MeshDescriptionProvider.h"
 #include "TargetInterfaces/PrimitiveComponentBackedTarget.h"
+#include "ModelingToolTargetUtil.h"
 
 using namespace UE::Geometry;
 
@@ -161,7 +160,7 @@ void USubdividePolyTool::Setup()
 	bool bWantVertexNormals = false;
 	OriginalMesh = MakeShared<FDynamicMesh3, ESPMode::ThreadSafe>(bWantVertexNormals, false, false, false);
 	FMeshDescriptionToDynamicMesh Converter;
-	Converter.Convert(Cast<IMeshDescriptionProvider>(Target)->GetMeshDescription(), *OriginalMesh);
+	Converter.Convert(UE::ToolTarget::GetMeshDescription(Target), *OriginalMesh);
 
 	FText ErrorMessage;
 	bool bCatmullClarkOK = CheckGroupTopology(ErrorMessage);
@@ -396,11 +395,7 @@ void USubdividePolyTool::Shutdown(EToolShutdownType ShutdownType)
 			UDynamicMeshComponent* PreviewDynamicMeshComponent = (UDynamicMeshComponent*)PreviewMesh->GetRootComponent();
 			FDynamicMesh3* DynamicMeshResult = PreviewDynamicMeshComponent->GetRenderMesh();
 
-			Cast<IMeshDescriptionCommitter>(Target)->CommitMeshDescription([DynamicMeshResult](const IMeshDescriptionCommitter::FCommitterParams& CommitParams)
-			{
-				FDynamicMeshToMeshDescription Converter;
-				Converter.Convert(DynamicMeshResult, *CommitParams.MeshDescriptionOut);
-			});
+			UE::ToolTarget::CommitMeshDescriptionUpdateViaDynamicMesh(Target, *DynamicMeshResult, true);
 
 			GetToolManager()->EndUndoTransaction();
 		}
diff --git a/Engine/Plugins/Experimental/MeshModelingToolsetExp/Source/MeshModelingToolsExp/Private/BakeMeshAttributeVertexTool.cpp b/Engine/Plugins/Experimental/MeshModelingToolsetExp/Source/MeshModelingToolsExp/Private/BakeMeshAttributeVertexTool.cpp
--- a/Engine/Plugins/Experimental/MeshModelingToolsetExp/Source/MeshModelingToolsExp/Private/BakeMeshAttributeVertexTool.cpp
+++ b/Engine/Plugins/Experimental/MeshModelingToolsetExp/Source/MeshModelingToolsExp/Private/BakeMeshAttributeVertexTool.cpp
@@ -18,8 +18,6 @@
 #include "Sampling/MeshResampleImageEvaluator.h"
 
 #include "TargetInterfaces/MaterialProvider.h"
-#include "TargetInterfaces/MeshDescriptionProvider.h"
-#include "TargetInterfaces/MeshDescriptionCommitter.h"
 #include "TargetInterfaces/PrimitiveComponentBackedTarget.h"
 #include "TargetInterfaces/StaticMeshBackedTarget.h"
 #include "TargetInterfaces/SkeletalMeshBackedTarget.h"
@@ -274,20 +272,6 @@ void UBakeMeshAttributeVertexTool::Setup()
 
 	// Setup tool property sets
 
-	Settings = NewObject<UBakeMeshAttributeVertexToolProperties>(this);
-	Settings->RestoreProperties(this);
-	AddToolPropertySource(Settings);
-
-	Settings->WatchProperty(Settings->OutputMode, [this](EBakeVertexOutput) { OpState |= EBakeOpState::Evaluate; UpdateOnModeChange(); });
-	Settings->WatchProperty(Settings->OutputType, [this](int32) { OpState |= EBakeOpState::Evaluate; UpdateOnModeChange(); });
-	Settings->WatchProperty(Settings->OutputTypeR, [this](int32) { OpState |= EBakeOpState::Evaluate; UpdateOnModeChange(); });
-	Settings->WatchProperty(Settings->OutputTypeG, [this](int32) { OpState |= EBakeOpState::Evaluate; UpdateOnModeChange(); });
-	Settings->WatchProperty(Settings->OutputTypeB, [this](int32) { OpState |= EBakeOpState::Evaluate; UpdateOnModeChange(); });
-	Settings->WatchProperty(Settings->OutputTypeA, [this](int32) { OpState |= EBakeOpState::Evaluate; UpdateOnModeChange(); });
-	Settings->WatchProperty(Settings->PreviewMode, [this](EBakeVertexChannel) { UpdateVisualization(); });
-	Settings->WatchProperty(Settings->bSplitAtNormalSeams, [this](bool) { bColorTopologyValid = false; OpState |= EBakeOpState::Evaluate; });
-	Settings->WatchProperty(Settings->bSplitAtUVSeams, [this](bool) { bColorTopologyValid = false; OpState |= EBakeOpState::Evaluate; });
-
 	InputMeshSettings = NewObject<UBakeInputMeshProperties>(this);
 	InputMeshSettings->RestoreProperties(this);
 	AddToolPropertySource(InputMeshSettings);
@@ -303,6 +287,20 @@ void UBakeMeshAttributeVertexTool::Setup()
 	InputMeshSettings->SourceNormalMap = nullptr;
 	InputMeshSettings->WatchProperty(InputMeshSettings->ProjectionDistance, [this](float) { OpState |= EBakeOpState::Evaluate; });
 	InputMeshSettings->WatchProperty(InputMeshSettings->bProjectionInWorldSpace, [this](bool) { OpState |= EBakeOpState::EvaluateDetailMesh; });
+	
+	Settings = NewObject<UBakeMeshAttributeVertexToolProperties>(this);
+	Settings->RestoreProperties(this);
+	AddToolPropertySource(Settings);
+
+	Settings->WatchProperty(Settings->OutputMode, [this](EBakeVertexOutput) { OpState |= EBakeOpState::Evaluate; UpdateOnModeChange(); });
+	Settings->WatchProperty(Settings->OutputType, [this](int32) { OpState |= EBakeOpState::Evaluate; UpdateOnModeChange(); });
+	Settings->WatchProperty(Settings->OutputTypeR, [this](int32) { OpState |= EBakeOpState::Evaluate; UpdateOnModeChange(); });
+	Settings->WatchProperty(Settings->OutputTypeG, [this](int32) { OpState |= EBakeOpState::Evaluate; UpdateOnModeChange(); });
+	Settings->WatchProperty(Settings->OutputTypeB, [this](int32) { OpState |= EBakeOpState::Evaluate; UpdateOnModeChange(); });
+	Settings->WatchProperty(Settings->OutputTypeA, [this](int32) { OpState |= EBakeOpState::Evaluate; UpdateOnModeChange(); });
+	Settings->WatchProperty(Settings->PreviewMode, [this](EBakeVertexChannel) { UpdateVisualization(); });
+	Settings->WatchProperty(Settings->bSplitAtNormalSeams, [this](bool) { bColorTopologyValid = false; OpState |= EBakeOpState::Evaluate; });
+	Settings->WatchProperty(Settings->bSplitAtUVSeams, [this](bool) { bColorTopologyValid = false; OpState |= EBakeOpState::Evaluate; });
 
 	OcclusionSettings = NewObject<UBakeOcclusionMapToolProperties>(this);
 	OcclusionSettings->RestoreProperties(this);
@@ -454,11 +452,11 @@ void UBakeMeshAttributeVertexTool::UpdateDetailMesh()
 {
 	IPrimitiveComponentBackedTarget* TargetComponent = Cast<IPrimitiveComponentBackedTarget>(Targets[0]);
 	IPrimitiveComponentBackedTarget* DetailComponent = Cast<IPrimitiveComponentBackedTarget>(Targets[bIsBakeToSelf ? 0 : 1]);
-	IMeshDescriptionProvider* DetailMeshProvider = Cast<IMeshDescriptionProvider>(Targets[bIsBakeToSelf ? 0 : 1]);
+	UToolTarget* DetailTargetMesh = Targets[bIsBakeToSelf ? 0 : 1];
 
 	DetailMesh = MakeShared<FDynamicMesh3, ESPMode::ThreadSafe>();
 	FMeshDescriptionToDynamicMesh Converter;
-	Converter.Convert(DetailMeshProvider->GetMeshDescription(), *DetailMesh);
+	Converter.Convert( UE::ToolTarget::GetMeshDescription(DetailTargetMesh), *DetailMesh);
 	if (InputMeshSettings->bProjectionInWorldSpace && bIsBakeToSelf == false)
 	{
 		const UE::Geometry::FTransform3d DetailToWorld(DetailComponent->GetWorldTransform());
diff --git a/Engine/Plugins/Experimental/MeshModelingToolsetExp/Source/MeshModelingToolsExp/Private/EditNormalsTool.cpp b/Engine/Plugins/Experimental/MeshModelingToolsetExp/Source/MeshModelingToolsExp/Private/EditNormalsTool.cpp
--- a/Engine/Plugins/Experimental/MeshModelingToolsetExp/Source/MeshModelingToolsExp/Private/EditNormalsTool.cpp
+++ b/Engine/Plugins/Experimental/MeshModelingToolsetExp/Source/MeshModelingToolsExp/Private/EditNormalsTool.cpp
@@ -17,13 +17,11 @@
 #include "InteractiveGizmoManager.h"
 
 #include "AssetUtils/MeshDescriptionUtil.h"
-#include "Engine/Classes/Engine/StaticMesh.h"
-#include "Engine/Classes/Components/StaticMeshComponent.h"
+#include "Engine/StaticMesh.h"
+#include "Components/StaticMeshComponent.h"
 
-#include "TargetInterfaces/MeshDescriptionCommitter.h"
-#include "TargetInterfaces/MeshDescriptionProvider.h"
 #include "TargetInterfaces/PrimitiveComponentBackedTarget.h"
-#include "ToolTargetManager.h"
+#include "ModelingToolTargetUtil.h"
 
 using namespace UE::Geometry;
 
@@ -262,25 +260,9 @@ void UEditNormalsTool::GenerateAsset(const TArray<FDynamicMeshOpResult>& Results
 			}
 		}
 
-		check(Results[ComponentIdx].Mesh.Get() != nullptr);
-		Cast<IMeshDescriptionCommitter>(Targets[ComponentIdx])->CommitMeshDescription([&Results, &ComponentIdx, this](const IMeshDescriptionCommitter::FCommitterParams& CommitParams)
-		{
-			FDynamicMeshToMeshDescription Converter;
-			
-			if (BasicProperties->WillTopologyChange() || !FDynamicMeshToMeshDescription::HaveMatchingElementCounts(Results[ComponentIdx].Mesh.Get(), CommitParams.MeshDescriptionOut, false, true))
-			{
-				// full conversion if normal topology changed or faces were inverted
-				Converter.Convert(Results[ComponentIdx].Mesh.Get(), *CommitParams.MeshDescriptionOut);
-			}
-			else
-			{
-				// otherwise just copy attributes
-				const bool bUpdateNormals = true;
-				const bool bCopyOverlayTangents = false;
-				const bool bCopyOverlayUVs = false;
-				Converter.UpdateAttributes(Results[ComponentIdx].Mesh.Get(), *CommitParams.MeshDescriptionOut, bUpdateNormals, bCopyOverlayTangents, bCopyOverlayUVs);
-			}
-		});
+		const FDynamicMesh3* NewDynamicMesh = Results[ComponentIdx].Mesh.Get();
+		check(NewDynamicMesh != nullptr);
+		UE::ToolTarget::CommitDynamicMeshNormalsUpdate(Targets[ComponentIdx], NewDynamicMesh, false);
 	}
 
 	GetToolManager()->EndUndoTransaction();
diff --git a/Engine/Plugins/Experimental/MeshModelingToolsetExp/Source/MeshModelingToolsExp/Private/HoleFillTool.cpp b/Engine/Plugins/Experimental/MeshModelingToolsetExp/Source/MeshModelingToolsExp/Private/HoleFillTool.cpp
--- a/Engine/Plugins/Experimental/MeshModelingToolsetExp/Source/MeshModelingToolsExp/Private/HoleFillTool.cpp
+++ b/Engine/Plugins/Experimental/MeshModelingToolsetExp/Source/MeshModelingToolsExp/Private/HoleFillTool.cpp
@@ -15,9 +15,8 @@
 #include "Selection/PolygonSelectionMechanic.h"
 
 #include "TargetInterfaces/MaterialProvider.h"
-#include "TargetInterfaces/MeshDescriptionCommitter.h"
-#include "TargetInterfaces/MeshDescriptionProvider.h"
 #include "TargetInterfaces/PrimitiveComponentBackedTarget.h"
+#include "ModelingToolTargetUtil.h"
 
 using namespace UE::Geometry;
 
@@ -125,7 +124,7 @@ void UHoleFillTool::Setup()
 	// create mesh to operate on
 	OriginalMesh = MakeShared<FDynamicMesh3, ESPMode::ThreadSafe>();
 	FMeshDescriptionToDynamicMesh Converter;
-	Converter.Convert(Cast<IMeshDescriptionProvider>(Target)->GetMeshDescription(), *OriginalMesh);
+	Converter.Convert(UE::ToolTarget::GetMeshDescription(Target), *OriginalMesh);
 
 	// initialize properties
 	Properties = NewObject<UHoleFillToolProperties>(this, TEXT("Hole Fill Settings"));
@@ -271,13 +270,7 @@ void UHoleFillTool::Shutdown(EToolShutdownType ShutdownType)
 		GetToolManager()->BeginUndoTransaction(LOCTEXT("HoleFillToolTransactionName", "Hole Fill Tool"));
 
 		check(Result.Mesh.Get() != nullptr);
-		Cast<IMeshDescriptionCommitter>(Target)->CommitMeshDescription([&Result](const IMeshDescriptionCommitter::FCommitterParams& CommitParams)
-		{
-			FDynamicMeshToMeshDescription Converter;
-
-			// full conversion if normal topology changed or faces were inverted
-			Converter.Convert(Result.Mesh.Get(), *CommitParams.MeshDescriptionOut);
-		});
+		UE::ToolTarget::CommitMeshDescriptionUpdateViaDynamicMesh(Target, *Result.Mesh.Get(), true);
 
 		GetToolManager()->EndUndoTransaction();
 	}
diff --git a/Engine/Plugins/Experimental/MeshModelingToolsetExp/Source/MeshModelingToolsExp/Private/LatticeDeformerTool.cpp b/Engine/Plugins/Experimental/MeshModelingToolsetExp/Source/MeshModelingToolsExp/Private/LatticeDeformerTool.cpp
--- a/Engine/Plugins/Experimental/MeshModelingToolsetExp/Source/MeshModelingToolsExp/Private/LatticeDeformerTool.cpp
+++ b/Engine/Plugins/Experimental/MeshModelingToolsetExp/Source/MeshModelingToolsExp/Private/LatticeDeformerTool.cpp
@@ -23,6 +23,7 @@
 #include "TargetInterfaces/MeshDescriptionCommitter.h"
 #include "TargetInterfaces/MeshDescriptionProvider.h"
 #include "TargetInterfaces/PrimitiveComponentBackedTarget.h"
+#include "ModelingToolTargetUtil.h"
 
 using namespace UE::Geometry;
 
@@ -170,7 +171,7 @@ void ULatticeDeformerTool::Setup()
 
 	OriginalMesh = MakeShared<FDynamicMesh3, ESPMode::ThreadSafe>();
 	FMeshDescriptionToDynamicMesh Converter;
-	Converter.Convert(Cast<IMeshDescriptionProvider>(Target)->GetMeshDescription(), *OriginalMesh);
+	Converter.Convert(UE::ToolTarget::GetMeshDescription(Target), *OriginalMesh);
 
 	Settings = NewObject<ULatticeDeformerToolProperties>(this, TEXT("Lattice Deformer Tool Settings"));
 	Settings->Initialize(this);
@@ -359,11 +360,7 @@ void ULatticeDeformerTool::Shutdown(EToolShutdownType ShutdownType)
 			UE::Geometry::FTransform3d LocalToWorld(TargetComponent->GetWorldTransform());
 			MeshTransforms::ApplyTransformInverse(*DynamicMeshResult, LocalToWorld);
 
-			Cast<IMeshDescriptionCommitter>(Target)->CommitMeshDescription([DynamicMeshResult](const IMeshDescriptionCommitter::FCommitterParams& CommitParams)
-			{
-				FDynamicMeshToMeshDescription Converter;
-				Converter.Convert(DynamicMeshResult, *CommitParams.MeshDescriptionOut);
-			});
+			UE::ToolTarget::CommitMeshDescriptionUpdateViaDynamicMesh(Target, *DynamicMeshResult, true);
 
 			GetToolManager()->EndUndoTransaction();
 		}
diff --git a/Engine/Plugins/Experimental/MeshModelingToolsetExp/Source/MeshModelingToolsExp/Private/MeshAttributePaintTool.cpp b/Engine/Plugins/Experimental/MeshModelingToolsetExp/Source/MeshModelingToolsExp/Private/MeshAttributePaintTool.cpp
--- a/Engine/Plugins/Experimental/MeshModelingToolsetExp/Source/MeshModelingToolsExp/Private/MeshAttributePaintTool.cpp
+++ b/Engine/Plugins/Experimental/MeshModelingToolsetExp/Source/MeshModelingToolsExp/Private/MeshAttributePaintTool.cpp
@@ -9,9 +9,8 @@
 
 #include "MeshDescription.h"
 
-#include "TargetInterfaces/MeshDescriptionCommitter.h"
-#include "TargetInterfaces/MeshDescriptionProvider.h"
 #include "TargetInterfaces/PrimitiveComponentBackedTarget.h"
+#include "ModelingToolTargetUtil.h"
 
 using namespace UE::Geometry;
 
@@ -227,7 +226,7 @@ void UMeshAttributePaintTool::Setup()
 	ColorMapper = MakeUnique<FFloatAttributeColorMapper>();
 
 	EditedMesh = MakeUnique<FMeshDescription>();
-	*EditedMesh = *Cast<IMeshDescriptionProvider>(Target)->GetMeshDescription();
+	*EditedMesh = *UE::ToolTarget::GetMeshDescription(Target);
 
 	AttributeSource = MakeUnique<FMeshDescriptionVertexAttributeSource>(EditedMesh.Get());
 	AttribProps->Initialize(AttributeSource->GetAttributeList(), true);
@@ -684,10 +683,8 @@ void UMeshAttributePaintTool::OnShutdown(EToolShutdownType ShutdownType)
 		// this block bakes the modified DynamicMeshComponent back into the StaticMeshComponent inside an undo transaction
 		GetToolManager()->BeginUndoTransaction(LOCTEXT("MeshAttributePaintTool", "Edit Attributes"));
 
-		Cast<IMeshDescriptionCommitter>(Target)->CommitMeshDescription([=](const IMeshDescriptionCommitter::FCommitterParams& CommitParams)
-		{
-			*CommitParams.MeshDescriptionOut = *EditedMesh;
-		});
+		UE::ToolTarget::CommitMeshDescriptionUpdate(Target, EditedMesh.Get());
+
 		GetToolManager()->EndUndoTransaction();
 	}
 }
diff --git a/Engine/Plugins/Experimental/MeshModelingToolsetExp/Source/MeshModelingToolsExp/Private/MeshBoundaryToolBase.cpp b/Engine/Plugins/Experimental/MeshModelingToolsetExp/Source/MeshModelingToolsExp/Private/MeshBoundaryToolBase.cpp
--- a/Engine/Plugins/Experimental/MeshModelingToolsetExp/Source/MeshModelingToolsExp/Private/MeshBoundaryToolBase.cpp
+++ b/Engine/Plugins/Experimental/MeshModelingToolsetExp/Source/MeshModelingToolsExp/Private/MeshBoundaryToolBase.cpp
@@ -9,8 +9,8 @@
 #include "MeshDescriptionToDynamicMesh.h"
 #include "Selection/PolygonSelectionMechanic.h"
 
-#include "TargetInterfaces/MeshDescriptionProvider.h"
 #include "TargetInterfaces/PrimitiveComponentBackedTarget.h"
+#include "ModelingToolTargetUtil.h"
 
 using namespace UE::Geometry;
 
@@ -28,7 +28,7 @@ void UMeshBoundaryToolBase::Setup()
 	// create mesh to operate on
 	OriginalMesh = MakeShared<FDynamicMesh3, ESPMode::ThreadSafe>();
 	FMeshDescriptionToDynamicMesh Converter;
-	Converter.Convert(Cast<IMeshDescriptionProvider>(Target)->GetMeshDescription(), *OriginalMesh);
+	Converter.Convert( UE::ToolTarget::GetMeshDescription(Target), *OriginalMesh);
 
 	// initialize hit query
 	MeshSpatial.SetMesh(OriginalMesh.Get());
diff --git a/Engine/Plugins/Experimental/MeshModelingToolsetExp/Source/MeshModelingToolsExp/Private/MeshSpaceDeformerTool.cpp b/Engine/Plugins/Experimental/MeshModelingToolsetExp/Source/MeshModelingToolsExp/Private/MeshSpaceDeformerTool.cpp
--- a/Engine/Plugins/Experimental/MeshModelingToolsetExp/Source/MeshModelingToolsExp/Private/MeshSpaceDeformerTool.cpp
+++ b/Engine/Plugins/Experimental/MeshModelingToolsetExp/Source/MeshModelingToolsExp/Private/MeshSpaceDeformerTool.cpp
@@ -22,14 +22,12 @@
 #include "BaseGizmos/TransformGizmoUtil.h"
 #include "BaseGizmos/IntervalGizmo.h"
 #include "MeshDescriptionToDynamicMesh.h"
-#include "ModelingToolTargetUtil.h"
 #include "DynamicMeshToMeshDescription.h"
 #include "CoreMinimal.h"
 #include "Math/Matrix.h"
 
-#include "TargetInterfaces/MeshDescriptionCommitter.h"
-#include "TargetInterfaces/MeshDescriptionProvider.h"
 #include "TargetInterfaces/PrimitiveComponentBackedTarget.h"
+#include "ModelingToolTargetUtil.h"
 
 using namespace UE::Geometry;
 
@@ -140,7 +138,7 @@ void UMeshSpaceDeformerTool::Setup()
 	{
 		OriginalDynamicMesh = MakeShared<FDynamicMesh3, ESPMode::ThreadSafe>();
 		FMeshDescriptionToDynamicMesh Converter;
-		Converter.Convert(Cast<IMeshDescriptionProvider>(Target)->GetMeshDescription(), *OriginalDynamicMesh);
+		Converter.Convert(UE::ToolTarget::GetMeshDescription(Target), *OriginalDynamicMesh);
 	}
 
 	IPrimitiveComponentBackedTarget* TargetComponent = Cast<IPrimitiveComponentBackedTarget>(Target);
@@ -335,12 +333,7 @@ void UMeshSpaceDeformerTool::Shutdown(EToolShutdownType ShutdownType)
 			FDynamicMesh3* DynamicMeshResult = Result.Mesh.Get();
 			check(DynamicMeshResult != nullptr);
 
-			Cast<IMeshDescriptionCommitter>(Target)->CommitMeshDescription([DynamicMeshResult](const IMeshDescriptionCommitter::FCommitterParams& CommitParams)
-			{
-				FDynamicMeshToMeshDescription Converter;
-				Converter.Convert(DynamicMeshResult, *CommitParams.MeshDescriptionOut);
-			});
-
+			UE::ToolTarget::CommitMeshDescriptionUpdateViaDynamicMesh(Target, *DynamicMeshResult, true);
 
 			GetToolManager()->EndUndoTransaction();
 		}
diff --git a/Engine/Plugins/Experimental/MeshModelingToolsetExp/Source/MeshModelingToolsExp/Private/MirrorTool.cpp b/Engine/Plugins/Experimental/MeshModelingToolsetExp/Source/MeshModelingToolsExp/Private/MirrorTool.cpp
--- a/Engine/Plugins/Experimental/MeshModelingToolsetExp/Source/MeshModelingToolsExp/Private/MirrorTool.cpp
+++ b/Engine/Plugins/Experimental/MeshModelingToolsetExp/Source/MeshModelingToolsExp/Private/MirrorTool.cpp
@@ -16,10 +16,8 @@
 #include "ToolSetupUtil.h"
 
 #include "TargetInterfaces/MaterialProvider.h"
-#include "TargetInterfaces/MeshDescriptionCommitter.h"
-#include "TargetInterfaces/MeshDescriptionProvider.h"
 #include "TargetInterfaces/PrimitiveComponentBackedTarget.h"
-#include "ToolTargetManager.h"
+#include "ModelingToolTargetUtil.h"
 
 using namespace UE::Geometry;
 
@@ -418,11 +416,7 @@ void UMirrorTool::GenerateAsset(const TArray<FDynamicMeshOpResult>& Results)
 		{
 			NewSelection.Actors.Add(UE::ToolTarget::GetTargetActor(Targets[OrigMeshIdx]));
 
-			Cast<IMeshDescriptionCommitter>(Targets[OrigMeshIdx])->CommitMeshDescription([&Mesh](const IMeshDescriptionCommitter::FCommitterParams& CommitParams)
-				{
-					FDynamicMeshToMeshDescription Converter;
-					Converter.Convert(Mesh, *CommitParams.MeshDescriptionOut);
-				});
+			UE::ToolTarget::CommitMeshDescriptionUpdateViaDynamicMesh(Targets[OrigMeshIdx], *Mesh, true);
 		}
 		else
 		{
diff --git a/Engine/Plugins/Experimental/MeshModelingToolsetExp/Source/MeshModelingToolsExp/Private/PlaneCutTool.cpp b/Engine/Plugins/Experimental/MeshModelingToolsetExp/Source/MeshModelingToolsExp/Private/PlaneCutTool.cpp
--- a/Engine/Plugins/Experimental/MeshModelingToolsetExp/Source/MeshModelingToolsExp/Private/PlaneCutTool.cpp
+++ b/Engine/Plugins/Experimental/MeshModelingToolsetExp/Source/MeshModelingToolsExp/Private/PlaneCutTool.cpp
@@ -33,7 +33,7 @@
 #include "TargetInterfaces/MeshDescriptionCommitter.h"
 #include "TargetInterfaces/MeshDescriptionProvider.h"
 #include "TargetInterfaces/PrimitiveComponentBackedTarget.h"
-#include "ToolTargetManager.h"
+#include "ModelingToolTargetUtil.h"
 
 using namespace UE::Geometry;
 
@@ -474,11 +474,7 @@ void UPlaneCutTool::GenerateAsset(const TArray<FDynamicMeshOpResult>& Results)
 			}
 		}
 
-		Cast<IMeshDescriptionCommitter>(Targets[OrigMeshIdx])->CommitMeshDescription([&UseMesh](const IMeshDescriptionCommitter::FCommitterParams& CommitParams)
-		{
-			FDynamicMeshToMeshDescription Converter;
-			Converter.Convert(UseMesh, *CommitParams.MeshDescriptionOut);
-		});
+		UE::ToolTarget::CommitMeshDescriptionUpdateViaDynamicMesh(Targets[OrigMeshIdx], *UseMesh, true);
 	}
 
 	if (bNeedToAdd)
diff --git a/Engine/Plugins/Experimental/MeshModelingToolsetExp/Source/MeshModelingToolsExp/Private/RemoveOccludedTrianglesTool.cpp b/Engine/Plugins/Experimental/MeshModelingToolsetExp/Source/MeshModelingToolsExp/Private/RemoveOccludedTrianglesTool.cpp
--- a/Engine/Plugins/Experimental/MeshModelingToolsetExp/Source/MeshModelingToolsExp/Private/RemoveOccludedTrianglesTool.cpp
+++ b/Engine/Plugins/Experimental/MeshModelingToolsetExp/Source/MeshModelingToolsExp/Private/RemoveOccludedTrianglesTool.cpp
@@ -28,7 +28,7 @@
 #include "TargetInterfaces/MeshDescriptionCommitter.h"
 #include "TargetInterfaces/MeshDescriptionProvider.h"
 #include "TargetInterfaces/PrimitiveComponentBackedTarget.h"
-#include "ToolTargetManager.h"
+#include "ModelingToolTargetUtil.h"
 
 using namespace UE::Geometry;
 
@@ -505,11 +505,7 @@ void URemoveOccludedTrianglesTool::GenerateAsset(const TArray<FDynamicMeshOpResu
 			continue;
 		}
 
-		Cast<IMeshDescriptionCommitter>(Targets[ComponentIdx])->CommitMeshDescription([&Results, &PreviewIdx, this](const IMeshDescriptionCommitter::FCommitterParams& CommitParams)
-		{
-			FDynamicMeshToMeshDescription Converter;
-			Converter.Convert(Results[PreviewIdx].Mesh.Get(), *CommitParams.MeshDescriptionOut);
-		});
+		UE::ToolTarget::CommitMeshDescriptionUpdateViaDynamicMesh(Targets[ComponentIdx], *Results[PreviewIdx].Mesh, true);
 	}
 
 	GetToolManager()->EndUndoTransaction();
diff --git a/Engine/Plugins/Experimental/MeshModelingToolsetExp/Source/MeshModelingToolsExp/Private/SeamSculptTool.cpp b/Engine/Plugins/Experimental/MeshModelingToolsetExp/Source/MeshModelingToolsExp/Private/SeamSculptTool.cpp
--- a/Engine/Plugins/Experimental/MeshModelingToolsetExp/Source/MeshModelingToolsExp/Private/SeamSculptTool.cpp
+++ b/Engine/Plugins/Experimental/MeshModelingToolsetExp/Source/MeshModelingToolsExp/Private/SeamSculptTool.cpp
@@ -24,6 +24,7 @@
 
 #include "TargetInterfaces/MeshDescriptionCommitter.h"
 #include "TargetInterfaces/PrimitiveComponentBackedTarget.h"
+#include "ModelingToolTargetUtil.h"
 
 using namespace UE::Geometry;
 
@@ -113,11 +114,7 @@ void USeamSculptTool::OnShutdown(EToolShutdownType ShutdownType)
 
 		GetToolManager()->BeginUndoTransaction(LOCTEXT("SeamSculptTransactionName", "UV Seam Edit"));
 
-		Cast<IMeshDescriptionCommitter>(Target)->CommitMeshDescription([&ResultMesh](const IMeshDescriptionCommitter::FCommitterParams& CommitParams)
-		{
-			FDynamicMeshToMeshDescription Converter;
-			Converter.Convert(&ResultMesh, *CommitParams.MeshDescriptionOut);
-		});
+		UE::ToolTarget::CommitMeshDescriptionUpdateViaDynamicMesh(Target, ResultMesh, true);
 
 		GetToolManager()->EndUndoTransaction();
 	}
diff --git a/Engine/Plugins/Experimental/MeshModelingToolsetExp/Source/MeshModelingToolsExp/Private/SkinWeightsBindingTool.cpp b/Engine/Plugins/Experimental/MeshModelingToolsetExp/Source/MeshModelingToolsExp/Private/SkinWeightsBindingTool.cpp
--- a/Engine/Plugins/Experimental/MeshModelingToolsetExp/Source/MeshModelingToolsExp/Private/SkinWeightsBindingTool.cpp
+++ b/Engine/Plugins/Experimental/MeshModelingToolsetExp/Source/MeshModelingToolsExp/Private/SkinWeightsBindingTool.cpp
@@ -11,15 +11,14 @@
 #include "SkeletalMeshAttributes.h"
 #include "ToolSetupUtil.h"
 #include "ToolTargetManager.h"
-#include "ModelingToolTargetUtil.h"
 #include "Components/SkeletalMeshComponent.h"
 #include "Engine/SkeletalMesh.h"
 #include "Spatial/FastWinding.h"
 #include "Spatial/MeshWindingNumberGrid.h"
 
 #include "TargetInterfaces/PrimitiveComponentBackedTarget.h"
-#include "TargetInterfaces/MeshDescriptionProvider.h"
-#include "TargetInterfaces/MeshDescriptionCommitter.h"
+#include "ModelingToolTargetUtil.h"
+
 
 // #pragma optimize( "", off )
 
@@ -1007,13 +1006,7 @@ void USkinWeightsBindingTool::GenerateAsset(const FDynamicMeshOpResult& Result)
 	GetToolManager()->BeginUndoTransaction(LOCTEXT("SkinWeightsBindingToolTransactionName", "Create Rigid Binding"));
 
 	check(Result.Mesh.Get() != nullptr);
-	Cast<IMeshDescriptionCommitter>(Targets[0])->CommitMeshDescription([&Result](const IMeshDescriptionCommitter::FCommitterParams& CommitParams)
-	{
-		FDynamicMeshToMeshDescription Converter;
-
-		// full conversion if normal topology changed or faces were inverted
-		Converter.Convert(Result.Mesh.Get(), *CommitParams.MeshDescriptionOut);
-	});
+	UE::ToolTarget::CommitMeshDescriptionUpdateViaDynamicMesh(Targets[0], *Result.Mesh.Get(), true);
 
 	GetToolManager()->EndUndoTransaction();	
 }
diff --git a/Engine/Plugins/Experimental/MeshModelingToolsetExp/Source/MeshModelingToolsExp/Private/SkinWeightsPaintTool.cpp b/Engine/Plugins/Experimental/MeshModelingToolsetExp/Source/MeshModelingToolsExp/Private/SkinWeightsPaintTool.cpp
--- a/Engine/Plugins/Experimental/MeshModelingToolsetExp/Source/MeshModelingToolsExp/Private/SkinWeightsPaintTool.cpp
+++ b/Engine/Plugins/Experimental/MeshModelingToolsetExp/Source/MeshModelingToolsExp/Private/SkinWeightsPaintTool.cpp
@@ -8,8 +8,9 @@
 #include "SkeletalDebugRendering.h"
 #include "Math/UnrealMathUtility.h"
 #include "Components/SkeletalMeshComponent.h"
-#include "TargetInterfaces/MeshDescriptionProvider.h"
+
 #include "TargetInterfaces/PrimitiveComponentBackedTarget.h"
+#include "ModelingToolTargetUtil.h"
 
 #include "MeshDescription.h"
 
@@ -128,7 +129,7 @@ void USkinWeightsPaintTool::Setup()
 		EToolMessageLevel::UserNotification);
 
 	EditedMesh = MakeUnique<FMeshDescription>();
-	*EditedMesh = *Cast<IMeshDescriptionProvider>(Target)->GetMeshDescription();
+	*EditedMesh = *UE::ToolTarget::GetMeshDescription(Target);
 
 	InitializeSkinWeights();
 
@@ -448,13 +449,13 @@ void USkinWeightsPaintTool::OnShutdown(EToolShutdownType ShutdownType)
 
 	if (ShutdownType == EToolShutdownType::Accept)
 	{
+		UpdateEditedSkinWeightsMesh();
+
 		// this block bakes the modified DynamicMeshComponent back into the StaticMeshComponent inside an undo transaction
 		GetToolManager()->BeginUndoTransaction(LOCTEXT("SkinWeightsPaintTool", "Paint Skin Weights"));
 
-		Cast<IMeshDescriptionCommitter>(Target)->CommitMeshDescription([this](const IMeshDescriptionCommitter::FCommitterParams& CommitParams)
-			{
-				CommitSkinWeights(CommitParams);
-			});
+		UE::ToolTarget::CommitMeshDescriptionUpdate(Target, EditedMesh.Get());
+
 		GetToolManager()->EndUndoTransaction();
 	}
 }
@@ -680,7 +681,7 @@ void USkinWeightsPaintTool::InitializeSkinWeights()
 }
 
 
-void USkinWeightsPaintTool::CommitSkinWeights(const IMeshDescriptionCommitter::FCommitterParams& CommitParams)
+void USkinWeightsPaintTool::UpdateEditedSkinWeightsMesh()
 {
 	using namespace UE::AnimationCore;
 	
@@ -717,8 +718,6 @@ void USkinWeightsPaintTool::CommitSkinWeights(const IMeshDescriptionCommitter::F
 
 		VertexSkinWeights.Set(FVertexID(VertexIndex), FBoneWeights::Create(SourceBoneWeights, Settings));
 	}
-
-	*CommitParams.MeshDescriptionOut = *EditedMesh;
 }
 
 
diff --git a/Engine/Plugins/Experimental/MeshModelingToolsetExp/Source/MeshModelingToolsExp/Private/UVLayoutTool.cpp b/Engine/Plugins/Experimental/MeshModelingToolsetExp/Source/MeshModelingToolsExp/Private/UVLayoutTool.cpp
--- a/Engine/Plugins/Experimental/MeshModelingToolsetExp/Source/MeshModelingToolsExp/Private/UVLayoutTool.cpp
+++ b/Engine/Plugins/Experimental/MeshModelingToolsetExp/Source/MeshModelingToolsExp/Private/UVLayoutTool.cpp
@@ -16,10 +16,8 @@
 #include "Properties/UVLayoutProperties.h"
 
 #include "TargetInterfaces/MaterialProvider.h"
-#include "TargetInterfaces/MeshDescriptionCommitter.h"
-#include "TargetInterfaces/MeshDescriptionProvider.h"
 #include "TargetInterfaces/PrimitiveComponentBackedTarget.h"
-#include "ToolTargetManager.h"
+#include "ModelingToolTargetUtil.h"
 
 using namespace UE::Geometry;
 
@@ -292,26 +290,9 @@ void UUVLayoutTool::GenerateAsset(const TArray<FDynamicMeshOpResult>& Results)
 	
 	for (int32 ComponentIdx = 0; ComponentIdx < Targets.Num(); ComponentIdx++)
 	{
-		check(Results[ComponentIdx].Mesh.Get() != nullptr);
-		Cast<IMeshDescriptionCommitter>(Targets[ComponentIdx])->CommitMeshDescription([&Results, &ComponentIdx, this](const IMeshDescriptionCommitter::FCommitterParams& CommitParams)
-		{
-			const FDynamicMesh3* DynamicMesh = Results[ComponentIdx].Mesh.Get();
-			FMeshDescription* MeshDescription = CommitParams.MeshDescriptionOut;
-			
-			constexpr bool bVerticesOnly = false;
-			constexpr bool bAttributesOnly = true;
-			if (FDynamicMeshToMeshDescription::HaveMatchingElementCounts(DynamicMesh, MeshDescription, bVerticesOnly, bAttributesOnly))
-			{
-				FDynamicMeshToMeshDescription Converter;
-				Converter.UpdateAttributes(DynamicMesh, *MeshDescription, false, false, true/*update uvs*/);
-			}
-			else
-			{
-				// must have been duplicate tris in the mesh description; we can't count on 1-to-1 mapping of TriangleIDs.  Just convert
-				FDynamicMeshToMeshDescription Converter;
-				Converter.Convert(DynamicMesh, *MeshDescription);
-			}
-		});
+		const FDynamicMesh3* DynamicMesh = Results[ComponentIdx].Mesh.Get();
+		check(DynamicMesh != nullptr);
+		UE::ToolTarget::CommitDynamicMeshUVUpdate(Targets[ComponentIdx], DynamicMesh);
 	}
 
 	GetToolManager()->EndUndoTransaction();
diff --git a/Engine/Plugins/Experimental/MeshModelingToolsetExp/Source/MeshModelingToolsExp/Public/SkinWeightsPaintTool.h b/Engine/Plugins/Experimental/MeshModelingToolsetExp/Source/MeshModelingToolsExp/Public/SkinWeightsPaintTool.h
--- a/Engine/Plugins/Experimental/MeshModelingToolsetExp/Source/MeshModelingToolsExp/Public/SkinWeightsPaintTool.h
+++ b/Engine/Plugins/Experimental/MeshModelingToolsetExp/Source/MeshModelingToolsExp/Public/SkinWeightsPaintTool.h
@@ -167,5 +167,5 @@ protected:
 	friend class FMeshSkinWeightsChange;
 	void ExternalUpdateValues(const FName &BoneName, const TMap<int32, float>& IndexValues);
 
-	void CommitSkinWeights(const IMeshDescriptionCommitter::FCommitterParams& CommitParams);
+	void UpdateEditedSkinWeightsMesh();
 };
diff --git a/Engine/Plugins/Runtime/MeshModelingToolset/Source/MeshModelingTools/Private/CutMeshWithMeshTool.cpp b/Engine/Plugins/Runtime/MeshModelingToolset/Source/MeshModelingTools/Private/CutMeshWithMeshTool.cpp
--- a/Engine/Plugins/Runtime/MeshModelingToolset/Source/MeshModelingTools/Private/CutMeshWithMeshTool.cpp
+++ b/Engine/Plugins/Runtime/MeshModelingToolset/Source/MeshModelingTools/Private/CutMeshWithMeshTool.cpp
@@ -17,7 +17,6 @@
 #include "TargetInterfaces/MeshDescriptionProvider.h"
 #include "TargetInterfaces/PrimitiveComponentBackedTarget.h"
 #include "TargetInterfaces/AssetBackedTarget.h"
-#include "ToolTargetManager.h"
 #include "ModelingToolTargetUtil.h"
 
 using namespace UE::Geometry;
@@ -356,11 +355,7 @@ void UCutMeshWithMeshTool::Shutdown(EToolShutdownType ShutdownType)
 			{
 				MeshTransforms::ApplyTransform(*OpResult.Mesh, OpResult.Transform);
 				MeshTransforms::ApplyTransformInverse(*OpResult.Mesh, TargetToWorld);
-				Cast<IMeshDescriptionCommitter>(Targets[0])->CommitMeshDescription([&](const IMeshDescriptionCommitter::FCommitterParams& CommitParams)
-				{
-					FDynamicMeshToMeshDescription Converter;
-					Converter.Convert(OpResult.Mesh.Get(), *CommitParams.MeshDescriptionOut);
-				});
+				UE::ToolTarget::CommitMeshDescriptionUpdateViaDynamicMesh(Targets[0], *OpResult.Mesh, true);
 				Cast<IMaterialProvider>(Targets[0])->CommitMaterialSetUpdate(MaterialSet, true);
 			}
 		}
diff --git a/Engine/Plugins/Runtime/MeshModelingToolset/Source/ModelingComponents/Private/BaseTools/BaseCreateFromSelectedTool.cpp b/Engine/Plugins/Runtime/MeshModelingToolset/Source/ModelingComponents/Private/BaseTools/BaseCreateFromSelectedTool.cpp
--- a/Engine/Plugins/Runtime/MeshModelingToolset/Source/ModelingComponents/Private/BaseTools/BaseCreateFromSelectedTool.cpp
+++ b/Engine/Plugins/Runtime/MeshModelingToolset/Source/ModelingComponents/Private/BaseTools/BaseCreateFromSelectedTool.cpp
@@ -13,8 +13,6 @@
 #include "ModelingObjectsCreationAPI.h"
 #include "Selection/ToolSelectionUtil.h"
 
-#include "TargetInterfaces/MeshDescriptionCommitter.h"
-#include "TargetInterfaces/MeshDescriptionProvider.h"
 #include "TargetInterfaces/PrimitiveComponentBackedTarget.h"
 #include "TargetInterfaces/MaterialProvider.h"
 #include "TargetInterfaces/AssetBackedTarget.h"
@@ -245,11 +243,7 @@ void UBaseCreateFromSelectedTool::UpdateAsset(const FDynamicMeshOpResult& Result
 	MeshTransforms::ApplyTransform(*Result.Mesh, ResultTransform);
 	MeshTransforms::ApplyTransformInverse(*Result.Mesh, TargetToWorld);
 
-	Cast<IMeshDescriptionCommitter>(UpdateTarget)->CommitMeshDescription([&](const IMeshDescriptionCommitter::FCommitterParams& CommitParams)
-	{
-		FDynamicMeshToMeshDescription Converter;
-		Converter.Convert(Result.Mesh.Get(), *CommitParams.MeshDescriptionOut);
-	});
+	UE::ToolTarget::CommitMeshDescriptionUpdateViaDynamicMesh(UpdateTarget, *Result.Mesh, true);
 
 	FComponentMaterialSet MaterialSet;
 	MaterialSet.Materials = GetOutputMaterials();
diff --git a/Engine/Plugins/Runtime/MeshModelingToolset/Source/ModelingComponents/Private/BaseTools/BaseMeshProcessingTool.cpp b/Engine/Plugins/Runtime/MeshModelingToolset/Source/ModelingComponents/Private/BaseTools/BaseMeshProcessingTool.cpp
--- a/Engine/Plugins/Runtime/MeshModelingToolset/Source/ModelingComponents/Private/BaseTools/BaseMeshProcessingTool.cpp
+++ b/Engine/Plugins/Runtime/MeshModelingToolset/Source/ModelingComponents/Private/BaseTools/BaseMeshProcessingTool.cpp
@@ -19,6 +19,7 @@
 #include "TargetInterfaces/MeshDescriptionProvider.h"
 #include "TargetInterfaces/PrimitiveComponentBackedTarget.h"
 #include "ToolTargetManager.h"
+#include "ModelingToolTargetUtil.h"
 
 #define LOCTEXT_NAMESPACE "UBaseMeshProcessingTool"
 
@@ -82,7 +83,7 @@ void UBaseMeshProcessingTool::Setup()
 
 	// populate the BaseMesh with a conversion of the input mesh.
 	FMeshDescriptionToDynamicMesh Converter;
-	Converter.Convert(Cast<IMeshDescriptionProvider>(Target)->GetMeshDescription(), InitialMesh);
+	Converter.Convert(UE::ToolTarget::GetMeshDescription(Target), InitialMesh);
 
 	if (RequiresScaleNormalization())
 	{
@@ -238,19 +239,7 @@ void UBaseMeshProcessingTool::Shutdown(EToolShutdownType ShutdownType)
 			}
 
 			bool bTopologyChanged = HasMeshTopologyChanged();
-			Cast<IMeshDescriptionCommitter>(Target)->CommitMeshDescription([DynamicMeshResult, bTopologyChanged](const IMeshDescriptionCommitter::FCommitterParams& CommitParams)
-			{
-				FDynamicMeshToMeshDescription Converter;
-				if (bTopologyChanged)
-				{
-					Converter.Convert(DynamicMeshResult, *CommitParams.MeshDescriptionOut);
-				}
-				else
-				{
-					Converter.Update(DynamicMeshResult, *CommitParams.MeshDescriptionOut);
-				}
-			});
-
+			UE::ToolTarget::CommitMeshDescriptionUpdateViaDynamicMesh(Target, *DynamicMeshResult, bTopologyChanged);
 
 			GetToolManager()->EndUndoTransaction();
 		}
@@ -372,7 +361,7 @@ void UBaseMeshProcessingTool::SetupWeightMapPropertySet(UWeightMapSetProperties*
 	WeightMapPropertySet = Properties;
 
 	// initialize property list
-	Properties->InitializeFromMesh(Cast<IMeshDescriptionProvider>(Target)->GetMeshDescription());
+	Properties->InitializeFromMesh(UE::ToolTarget::GetMeshDescription(Target));
 
 	Properties->WatchProperty(Properties->WeightMap,
 		[&](FName) { OnSelectedWeightMapChanged(true); });
@@ -388,7 +377,7 @@ void UBaseMeshProcessingTool::OnSelectedWeightMapChanged(bool bInvalidate)
 	TSharedPtr<FIndexedWeightMap1f> NewWeightMap = MakeShared<FIndexedWeightMap1f>();
 
 	// this will return all-ones weight map if None is selected
-	bool bFound = UE::WeightMaps::GetVertexWeightMap(Cast<IMeshDescriptionProvider>(Target)->GetMeshDescription(), WeightMapPropertySet->WeightMap, *NewWeightMap, 1.0f);
+	bool bFound = UE::WeightMaps::GetVertexWeightMap(UE::ToolTarget::GetMeshDescription(Target), WeightMapPropertySet->WeightMap, *NewWeightMap, 1.0f);
 	if (bFound && WeightMapPropertySet->bInvertWeightMap)
 	{
 		NewWeightMap->InvertWeightMap();
diff --git a/Engine/Plugins/Runtime/MeshModelingToolset/Source/ModelingComponents/Private/ConversionUtils/DynamicMeshViaMeshDescriptionUtil.cpp b/Engine/Plugins/Runtime/MeshModelingToolset/Source/ModelingComponents/Private/ConversionUtils/DynamicMeshViaMeshDescriptionUtil.cpp
--- a/Engine/Plugins/Runtime/MeshModelingToolset/Source/ModelingComponents/Private/ConversionUtils/DynamicMeshViaMeshDescriptionUtil.cpp
+++ b/Engine/Plugins/Runtime/MeshModelingToolset/Source/ModelingComponents/Private/ConversionUtils/DynamicMeshViaMeshDescriptionUtil.cpp
@@ -18,6 +18,7 @@ FDynamicMesh3 UE::Geometry::GetDynamicMeshViaMeshDescription(
 }
 
 void UE::Geometry::CommitDynamicMeshViaMeshDescription(
+	FMeshDescription&& CurrentMeshDescription,
 	IMeshDescriptionCommitter& MeshDescriptionCommitter, 
 	const FDynamicMesh3& Mesh, const IDynamicMeshCommitter::FDynamicMeshCommitInfo& CommitInfo)
 {
@@ -30,18 +31,16 @@ void UE::Geometry::CommitDynamicMeshViaMeshDescription(
 	ConversionOptions.bUpdateVtxColors = CommitInfo.bVertexColorsChanged;
 	ConversionOptions.bTransformVtxColorsSRGBToLinear = CommitInfo.bTransformVertexColorsSRGBToLinear;
 
-	MeshDescriptionCommitter.CommitMeshDescription([&CommitInfo, &ConversionOptions, &Mesh](const IMeshDescriptionCommitter::FCommitterParams& CommitParams)
+	FDynamicMeshToMeshDescription Converter(ConversionOptions);
+	if (!CommitInfo.bTopologyChanged)
 	{
-		FDynamicMeshToMeshDescription Converter(ConversionOptions);
+		Converter.UpdateUsingConversionOptions(&Mesh, CurrentMeshDescription);
+	}
+	else
+	{
+		// Do a full conversion.
+		Converter.Convert(&Mesh, CurrentMeshDescription);
+	}
 
-		if (!CommitInfo.bTopologyChanged)
-		{
-			Converter.UpdateUsingConversionOptions(&Mesh, *CommitParams.MeshDescriptionOut);
-		}
-		else
-		{
-			// Do a full conversion.
-			Converter.Convert(&Mesh, *CommitParams.MeshDescriptionOut);
-		}
-	});
+	MeshDescriptionCommitter.CommitMeshDescription(MoveTemp(CurrentMeshDescription));
 }
\ No newline at end of file
diff --git a/Engine/Plugins/Runtime/MeshModelingToolset/Source/ModelingComponents/Private/ModelingToolTargetUtil.cpp b/Engine/Plugins/Runtime/MeshModelingToolset/Source/ModelingComponents/Private/ModelingToolTargetUtil.cpp
--- a/Engine/Plugins/Runtime/MeshModelingToolset/Source/ModelingComponents/Private/ModelingToolTargetUtil.cpp
+++ b/Engine/Plugins/Runtime/MeshModelingToolset/Source/ModelingComponents/Private/ModelingToolTargetUtil.cpp
@@ -237,16 +237,50 @@ UE::ToolTarget::EDynamicMeshUpdateResult UE::ToolTarget::CommitMeshDescriptionUp
 		CommitMaterialSetUpdate(Target, *UpdatedMaterials, true);
 	}
 
-	EDynamicMeshUpdateResult Result = EDynamicMeshUpdateResult::Failed;
-	MeshDescriptionCommitter->CommitMeshDescription([UpdatedMesh, &Result](const IMeshDescriptionCommitter::FCommitterParams& CommitParams)
+	bool bOK = MeshDescriptionCommitter->CommitMeshDescription(*UpdatedMesh);
+	return (bOK) ? EDynamicMeshUpdateResult::Ok : EDynamicMeshUpdateResult::Failed;
+}
+
+
+UE::ToolTarget::EDynamicMeshUpdateResult UE::ToolTarget::CommitMeshDescriptionUpdate(UToolTarget* Target, FMeshDescription&& UpdatedMesh)
+{
+	if (IMeshDescriptionCommitter* MeshDescriptionCommitter = Cast<IMeshDescriptionCommitter>(Target))
 	{
-		*CommitParams.MeshDescriptionOut = *UpdatedMesh;
-		Result = EDynamicMeshUpdateResult::Ok;
-	});
-	return Result;
+		bool bOK = MeshDescriptionCommitter->CommitMeshDescription(MoveTemp(UpdatedMesh));
+		return (bOK) ? EDynamicMeshUpdateResult::Ok : EDynamicMeshUpdateResult::Failed;
+	}
+	return EDynamicMeshUpdateResult::Failed;
 }
 
 
+UE::ToolTarget::EDynamicMeshUpdateResult UE::ToolTarget::CommitMeshDescriptionUpdateViaDynamicMesh(
+	UToolTarget* Target, 
+	const UE::Geometry::FDynamicMesh3& UpdatedMesh, 
+	bool bHaveModifiedTopology)
+{
+	IMeshDescriptionCommitter* MeshDescriptionCommitter = Cast<IMeshDescriptionCommitter>(Target);
+	if (!ensure(MeshDescriptionCommitter))
+	{
+		return EDynamicMeshUpdateResult::Failed;
+	}
+
+	FMeshDescription ConvertedMesh;
+	FDynamicMeshToMeshDescription Converter;
+	if (bHaveModifiedTopology)
+	{
+		Converter.Convert(&UpdatedMesh, ConvertedMesh);
+	}
+	else
+	{
+		Converter.Update(&UpdatedMesh, ConvertedMesh);
+	}
+
+	bool bOK = MeshDescriptionCommitter->CommitMeshDescription(MoveTemp(ConvertedMesh));
+	return (bOK) ? EDynamicMeshUpdateResult::Ok : EDynamicMeshUpdateResult::Failed;
+}
+
+
+
 void UE::ToolTarget::Internal::CommitDynamicMeshViaIPersistentDynamicMeshSource(
 	IPersistentDynamicMeshSource& DynamicMeshSource, 
 	const FDynamicMesh3& UpdatedMesh, bool bHaveModifiedTopology)
@@ -311,24 +345,19 @@ UE::ToolTarget::EDynamicMeshUpdateResult UE::ToolTarget::CommitDynamicMeshUpdate
 	IMeshDescriptionCommitter* MeshDescriptionCommitter = Cast<IMeshDescriptionCommitter>(Target);
 	if (MeshDescriptionCommitter)
 	{
-		EDynamicMeshUpdateResult Result = EDynamicMeshUpdateResult::Failed;
-		MeshDescriptionCommitter->CommitMeshDescription([&](const IMeshDescriptionCommitter::FCommitterParams& CommitParams)
+		FMeshDescription ConvertedMesh;
+		FDynamicMeshToMeshDescription Converter(ConversionOptions);
+		if (!bHaveModifiedTopology)
 		{
-			FMeshDescription* MeshDescription = CommitParams.MeshDescriptionOut;
-
-			FDynamicMeshToMeshDescription Converter(ConversionOptions);
-			if (!bHaveModifiedTopology)
-			{
-				Converter.UpdateUsingConversionOptions(&UpdatedMesh, *MeshDescription);
-				Result = EDynamicMeshUpdateResult::Ok;
-			}
-			else
-			{
-				Converter.Convert(&UpdatedMesh, *MeshDescription);
-				Result = EDynamicMeshUpdateResult::Ok;
-			}
-		});
-		return Result;
+			Converter.UpdateUsingConversionOptions(&UpdatedMesh, ConvertedMesh);
+		}
+		else
+		{
+			Converter.Convert(&UpdatedMesh, ConvertedMesh);
+		}
+
+		bool bOK = MeshDescriptionCommitter->CommitMeshDescription(MoveTemp(ConvertedMesh));
+		return (bOK) ? EDynamicMeshUpdateResult::Ok : EDynamicMeshUpdateResult::Failed;
 	}
 
 	ensure(false);
@@ -366,33 +395,80 @@ UE::ToolTarget::EDynamicMeshUpdateResult UE::ToolTarget::CommitDynamicMeshUVUpda
 		return EDynamicMeshUpdateResult::Failed;
 	}
 
-	EDynamicMeshUpdateResult Result = EDynamicMeshUpdateResult::Failed;
-	MeshDescriptionCommitter->CommitMeshDescription([UpdatedMesh, &Result](const IMeshDescriptionCommitter::FCommitterParams& CommitParams)
+	FMeshDescription NewMeshDescription = UE::ToolTarget::GetMeshDescriptionCopy(Target);
+	bool bVerticesOnly = false;
+	bool bAttributesOnly = true;
+	if (FDynamicMeshToMeshDescription::HaveMatchingElementCounts(UpdatedMesh, &NewMeshDescription, bVerticesOnly, bAttributesOnly))
 	{
-		FMeshDescription* MeshDescription = CommitParams.MeshDescriptionOut;
+		FDynamicMeshToMeshDescription Converter;
+		Converter.UpdateAttributes(UpdatedMesh, NewMeshDescription, false, false, true/*update uvs*/);
+	}
+	else
+	{
+		// must have been duplicate tris in the mesh description; we can't count on 1-to-1 mapping of TriangleIDs.  Just convert 
+		FDynamicMeshToMeshDescription Converter;
+		Converter.Convert(UpdatedMesh, NewMeshDescription);
+	}
 
-		bool bVerticesOnly = false;
-		bool bAttributesOnly = true;
-		if (FDynamicMeshToMeshDescription::HaveMatchingElementCounts(UpdatedMesh, MeshDescription, bVerticesOnly, bAttributesOnly))
-		{
-			FDynamicMeshToMeshDescription Converter;
-			Converter.UpdateAttributes(UpdatedMesh, *MeshDescription, false, false, true/*update uvs*/);
-			Result = EDynamicMeshUpdateResult::Ok;
-		}
-		else
-		{
-			// must have been duplicate tris in the mesh description; we can't count on 1-to-1 mapping of TriangleIDs.  Just convert 
-			FDynamicMeshToMeshDescription Converter;
-			Converter.Convert(UpdatedMesh, *MeshDescription);
-			Result = EDynamicMeshUpdateResult::Ok_ForcedFullUpdate;
-		}
-	});
-	return Result;
+	bool bOK = MeshDescriptionCommitter->CommitMeshDescription(MoveTemp(NewMeshDescription));
+	return (bOK) ? EDynamicMeshUpdateResult::Ok : EDynamicMeshUpdateResult::Failed;
+}
+
+
+
+
+
+UE::ToolTarget::EDynamicMeshUpdateResult UE::ToolTarget::CommitDynamicMeshNormalsUpdate(
+	UToolTarget* Target, 
+	const UE::Geometry::FDynamicMesh3* UpdatedMesh,
+	bool bUpdateTangents)
+{
+	IPersistentDynamicMeshSource* DynamicMeshSource = Cast<IPersistentDynamicMeshSource>(Target);
+	if (DynamicMeshSource)
+	{
+		// just do a full mesh update for now
+		return CommitDynamicMeshUpdate(Target, *UpdatedMesh, true);
+	}
+
+	IDynamicMeshCommitter* DynamicMeshCommitter = Cast<IDynamicMeshCommitter>(Target);
+	if (DynamicMeshCommitter)
+	{
+		IDynamicMeshCommitter::FDynamicMeshCommitInfo CommitInfo(false);
+		CommitInfo.bNormalsChanged = true;
+		CommitInfo.bTangentsChanged = bUpdateTangents;
+		DynamicMeshCommitter->CommitDynamicMesh(*UpdatedMesh, CommitInfo);
+		return EDynamicMeshUpdateResult::Ok;
+	}
+
+	IMeshDescriptionCommitter* MeshDescriptionCommitter = Cast<IMeshDescriptionCommitter>(Target);
+	if (!ensure(MeshDescriptionCommitter))
+	{
+		return EDynamicMeshUpdateResult::Failed;
+	}
+
+	FMeshDescription NewMeshDescription = UE::ToolTarget::GetMeshDescriptionCopy(Target);
+	bool bVerticesOnly = false;
+	bool bAttributesOnly = true;
+	if (FDynamicMeshToMeshDescription::HaveMatchingElementCounts(UpdatedMesh, &NewMeshDescription, bVerticesOnly, bAttributesOnly))
+	{
+		FDynamicMeshToMeshDescription Converter;
+		Converter.UpdateAttributes(UpdatedMesh, NewMeshDescription, true, bUpdateTangents, false);
+	}
+	else
+	{
+		// must have been duplicate tris in the mesh description; we can't count on 1-to-1 mapping of TriangleIDs.  Just convert 
+		FDynamicMeshToMeshDescription Converter;
+		Converter.Convert(UpdatedMesh, NewMeshDescription);
+	}
+
+	bool bOK = MeshDescriptionCommitter->CommitMeshDescription(MoveTemp(NewMeshDescription));
+	return (bOK) ? EDynamicMeshUpdateResult::Ok : EDynamicMeshUpdateResult::Failed;
 }
 
 
 
 
+
 bool UE::ToolTarget::ConfigureCreateMeshObjectParams(UToolTarget* SourceTarget, FCreateMeshObjectParams& DerivedParamsOut)
 {
 	IPrimitiveComponentBackedTarget* ComponentTarget = Cast<IPrimitiveComponentBackedTarget>(SourceTarget);
diff --git a/Engine/Plugins/Runtime/MeshModelingToolset/Source/ModelingComponents/Public/ConversionUtils/DynamicMeshViaMeshDescriptionUtil.h b/Engine/Plugins/Runtime/MeshModelingToolset/Source/ModelingComponents/Public/ConversionUtils/DynamicMeshViaMeshDescriptionUtil.h
--- a/Engine/Plugins/Runtime/MeshModelingToolset/Source/ModelingComponents/Public/ConversionUtils/DynamicMeshViaMeshDescriptionUtil.h
+++ b/Engine/Plugins/Runtime/MeshModelingToolset/Source/ModelingComponents/Public/ConversionUtils/DynamicMeshViaMeshDescriptionUtil.h
@@ -18,7 +18,10 @@ namespace Geometry {
 	MODELINGCOMPONENTS_API FDynamicMesh3 GetDynamicMeshViaMeshDescription(
 		IMeshDescriptionProvider& MeshDescriptionProvider);
 
-	MODELINGCOMPONENTS_API void CommitDynamicMeshViaMeshDescription(IMeshDescriptionCommitter& MeshDescriptionCommitter,
-		const UE::Geometry::FDynamicMesh3& Mesh, const IDynamicMeshCommitter::FDynamicMeshCommitInfo& CommitInfo);
+	MODELINGCOMPONENTS_API void CommitDynamicMeshViaMeshDescription(
+		FMeshDescription&& CurrentMeshDescription,
+		IMeshDescriptionCommitter& MeshDescriptionCommitter,
+		const UE::Geometry::FDynamicMesh3& Mesh, 
+		const IDynamicMeshCommitter::FDynamicMeshCommitInfo& CommitInfo);
 
 }}
\ No newline at end of file
diff --git a/Engine/Plugins/Runtime/MeshModelingToolset/Source/ModelingComponents/Public/ModelingToolTargetUtil.h b/Engine/Plugins/Runtime/MeshModelingToolset/Source/ModelingComponents/Public/ModelingToolTargetUtil.h
--- a/Engine/Plugins/Runtime/MeshModelingToolset/Source/ModelingComponents/Public/ModelingToolTargetUtil.h
+++ b/Engine/Plugins/Runtime/MeshModelingToolset/Source/ModelingComponents/Public/ModelingToolTargetUtil.h
@@ -131,6 +131,24 @@ MODELINGCOMPONENTS_API EDynamicMeshUpdateResult CommitMeshDescriptionUpdate(
 	const FMeshDescription* UpdatedMesh, 
 	const FComponentMaterialSet* UpdatedMaterials = nullptr);
 
+/**
+* Update the Mesh in a ToolTarget based on the provided MeshDescription, and optional material set
+* @return EDynamicMeshUpdateResult::Ok on success
+*/
+MODELINGCOMPONENTS_API EDynamicMeshUpdateResult CommitMeshDescriptionUpdate(
+	UToolTarget* Target, 
+	FMeshDescription&& UpdatedMesh);
+
+/**
+* Update the Mesh in a ToolTarget based on the provided MeshDescription, and optional material set
+* @return EDynamicMeshUpdateResult::Ok on success
+*/
+MODELINGCOMPONENTS_API EDynamicMeshUpdateResult CommitMeshDescriptionUpdateViaDynamicMesh(
+	UToolTarget* Target, 
+	const UE::Geometry::FDynamicMesh3& UpdatedMesh,
+	bool bHaveModifiedTopology);
+
+
 /**
  * Update the Mesh in a ToolTarget based on the provided DynamicMesh, and optional material set
  * @param bHaveModifiedTopology If the update only changes vertex or attribute values (but not counts), then in some cases a more efficient and/or less destructive update can be applied to the Target
@@ -152,6 +170,12 @@ MODELINGCOMPONENTS_API EDynamicMeshUpdateResult CommitDynamicMeshUpdate(
  */
 MODELINGCOMPONENTS_API EDynamicMeshUpdateResult CommitDynamicMeshUVUpdate(UToolTarget* Target, const UE::Geometry::FDynamicMesh3* UpdatedMesh);
 
+/**
+* Update the Normals/Tangents of the ToolTarget's mesh (assuming it has one) based on the provided UpdatedMesh.
+* @return EDynamicMeshUpdateResult::Ok on success, or Ok_ForcedFullUpdate if any dependent mesh topology was modified
+*/
+MODELINGCOMPONENTS_API EDynamicMeshUpdateResult CommitDynamicMeshNormalsUpdate(UToolTarget* Target, const UE::Geometry::FDynamicMesh3* UpdatedMesh, bool bUpdateTangents = false);
+
 /**
  * FCreateMeshObjectParams::TypeHint is used by the ModelingObjectsCreationAPI to suggest what type of mesh object to create
  * inside various Tools. This should often be derived from the input mesh object type (eg if you plane-cut a Volume, the output
diff --git a/Engine/Plugins/Runtime/MeshModelingToolset/Source/ModelingComponentsEditorOnly/Private/ToolTargets/SkeletalMeshComponentToolTarget.cpp b/Engine/Plugins/Runtime/MeshModelingToolset/Source/ModelingComponentsEditorOnly/Private/ToolTargets/SkeletalMeshComponentToolTarget.cpp
--- a/Engine/Plugins/Runtime/MeshModelingToolset/Source/ModelingComponentsEditorOnly/Private/ToolTargets/SkeletalMeshComponentToolTarget.cpp
+++ b/Engine/Plugins/Runtime/MeshModelingToolset/Source/ModelingComponentsEditorOnly/Private/ToolTargets/SkeletalMeshComponentToolTarget.cpp
@@ -148,7 +148,8 @@ void USkeletalMeshComponentToolTarget::CommitMeshDescription(const FCommitter& C
 
 void USkeletalMeshComponentToolTarget::CommitDynamicMesh(const FDynamicMesh3& Mesh, const FDynamicMeshCommitInfo& CommitInfo)
 {
-	CommitDynamicMeshViaMeshDescription(*this, Mesh, CommitInfo);
+	FMeshDescription CurrentMeshDescription = *GetMeshDescription();
+	CommitDynamicMeshViaMeshDescription(MoveTemp(CurrentMeshDescription), *this, Mesh, CommitInfo);
 }
 
 
diff --git a/Engine/Plugins/Runtime/MeshModelingToolset/Source/ModelingComponentsEditorOnly/Private/ToolTargets/SkeletalMeshToolTarget.cpp b/Engine/Plugins/Runtime/MeshModelingToolset/Source/ModelingComponentsEditorOnly/Private/ToolTargets/SkeletalMeshToolTarget.cpp
--- a/Engine/Plugins/Runtime/MeshModelingToolset/Source/ModelingComponentsEditorOnly/Private/ToolTargets/SkeletalMeshToolTarget.cpp
+++ b/Engine/Plugins/Runtime/MeshModelingToolset/Source/ModelingComponentsEditorOnly/Private/ToolTargets/SkeletalMeshToolTarget.cpp
@@ -217,7 +217,8 @@ void USkeletalMeshToolTarget::CommitMeshDescription(USkeletalMesh* SkeletalMeshI
 
 void USkeletalMeshToolTarget::CommitDynamicMesh(const FDynamicMesh3& Mesh, const FDynamicMeshCommitInfo& CommitInfo)
 {
-	CommitDynamicMeshViaMeshDescription(*this, Mesh, CommitInfo);
+	FMeshDescription CurrentMeshDescription = *GetMeshDescription();
+	CommitDynamicMeshViaMeshDescription(MoveTemp(CurrentMeshDescription), *this, Mesh, CommitInfo);
 }
 
 
diff --git a/Engine/Plugins/Runtime/MeshModelingToolset/Source/ModelingComponentsEditorOnly/Private/ToolTargets/StaticMeshComponentToolTarget.cpp b/Engine/Plugins/Runtime/MeshModelingToolset/Source/ModelingComponentsEditorOnly/Private/ToolTargets/StaticMeshComponentToolTarget.cpp
--- a/Engine/Plugins/Runtime/MeshModelingToolset/Source/ModelingComponentsEditorOnly/Private/ToolTargets/StaticMeshComponentToolTarget.cpp
+++ b/Engine/Plugins/Runtime/MeshModelingToolset/Source/ModelingComponentsEditorOnly/Private/ToolTargets/StaticMeshComponentToolTarget.cpp
@@ -183,7 +183,8 @@ FDynamicMesh3 UStaticMeshComponentToolTarget::GetDynamicMesh()
 
 void UStaticMeshComponentToolTarget::CommitDynamicMesh(const FDynamicMesh3& Mesh, const FDynamicMeshCommitInfo& CommitInfo)
 {
-	CommitDynamicMeshViaMeshDescription(*this, Mesh, CommitInfo);
+	FMeshDescription CurrentMeshDescription = *GetMeshDescription();
+	CommitDynamicMeshViaMeshDescription(MoveTemp(CurrentMeshDescription), *this, Mesh, CommitInfo);
 }
 
 UStaticMesh* UStaticMeshComponentToolTarget::GetStaticMesh() const
diff --git a/Engine/Plugins/Runtime/MeshModelingToolset/Source/ModelingComponentsEditorOnly/Private/ToolTargets/StaticMeshToolTarget.cpp b/Engine/Plugins/Runtime/MeshModelingToolset/Source/ModelingComponentsEditorOnly/Private/ToolTargets/StaticMeshToolTarget.cpp
--- a/Engine/Plugins/Runtime/MeshModelingToolset/Source/ModelingComponentsEditorOnly/Private/ToolTargets/StaticMeshToolTarget.cpp
+++ b/Engine/Plugins/Runtime/MeshModelingToolset/Source/ModelingComponentsEditorOnly/Private/ToolTargets/StaticMeshToolTarget.cpp
@@ -260,7 +260,8 @@ FDynamicMesh3 UStaticMeshToolTarget::GetDynamicMesh()
 
 void UStaticMeshToolTarget::CommitDynamicMesh(const FDynamicMesh3& Mesh, const FDynamicMeshCommitInfo& CommitInfo)
 {
-	CommitDynamicMeshViaMeshDescription(*this, Mesh, CommitInfo);
+	FMeshDescription CurrentMeshDescription = *GetMeshDescription();
+	CommitDynamicMeshViaMeshDescription(MoveTemp(CurrentMeshDescription), *this, Mesh, CommitInfo);
 }
 
 UStaticMesh* UStaticMeshToolTarget::GetStaticMesh() const
diff --git a/Engine/Source/Runtime/InteractiveToolsFramework/Public/TargetInterfaces/MeshDescriptionCommitter.h b/Engine/Source/Runtime/InteractiveToolsFramework/Public/TargetInterfaces/MeshDescriptionCommitter.h
--- a/Engine/Source/Runtime/InteractiveToolsFramework/Public/TargetInterfaces/MeshDescriptionCommitter.h
+++ b/Engine/Source/Runtime/InteractiveToolsFramework/Public/TargetInterfaces/MeshDescriptionCommitter.h
@@ -21,39 +21,66 @@ class INTERACTIVETOOLSFRAMEWORK_API IMeshDescriptionCommitter
 
 public:
 
-	struct FCommitterParams
+	/**
+	 * Commits the given FMeshDescription.
+	 */
+	virtual bool CommitMeshDescription(const FMeshDescription& Mesh)
 	{
-		/** 
-		 * Mesh description that should be populated/updated by the passed-in function and which
-		 * will be committed to the target.
-		 */
-		FMeshDescription* MeshDescriptionOut  = nullptr;
-	};
-	using  FCommitter = TFunction<void(const FCommitterParams&)>;
+		// It seems reasonable to have this function, but we'll go ahead and give a default implementation
+		// if users want to just implement the other one.
+		bool bSuccess = false;
+		CommitMeshDescription([&](const FCommitterParams& CommitParams)
+		{
+			if (CommitParams.MeshDescriptionOut)
+			{
+				*CommitParams.MeshDescriptionOut = Mesh;
+				bSuccess = true;
+			}
+		});
+		return bSuccess;
+	}
 
-	/** 
-	 * Commit a mesh description. The mesh description to be committed will be passed to the
-	 * given function as a parameter, and it is up to the function to update it properly.
-	 *
-	 * @param Committer A function that takes in const IMeshDescriptionCommitter::FCommitParams&
-	 *  and populates the FMeshDescription pointed to by the MeshDescription pointer inside.
-	 */
-	virtual void CommitMeshDescription(const FCommitter& Committer) = 0;
 
 	/**
 	 * Commits the given FMeshDescription.
 	 */
-	virtual void CommitMeshDescription(const FMeshDescription& Mesh)
+	virtual bool CommitMeshDescription(FMeshDescription&& Mesh)
 	{
 		// It seems reasonable to have this function, but we'll go ahead and give a default implementation
 		// if users want to just implement the other one.
-
-		CommitMeshDescription([&Mesh](const FCommitterParams& CommitParams)
+		bool bSuccess = false;
+		CommitMeshDescription([&](const FCommitterParams& CommitParams)
 		{
 			if (CommitParams.MeshDescriptionOut)
 			{
-				*CommitParams.MeshDescriptionOut = Mesh;
+				*CommitParams.MeshDescriptionOut = MoveTemp(Mesh);
+				bSuccess = true;
 			}
 		});
+		return bSuccess;
 	}
+
+
+
+protected:
+
+	struct FCommitterParams
+	{
+		/** 
+		* Mesh description that should be populated/updated by the passed-in function and which
+		* will be committed to the target.
+		*/
+		FMeshDescription* MeshDescriptionOut  = nullptr;
+	};
+	using FCommitter = TFunction<void(const FCommitterParams&)>;
+
+	/** 
+	* Commit a mesh description. The mesh description to be committed will be passed to the
+	* given function as a parameter, and it is up to the function to update it properly.
+	*
+	* @param Committer A function that takes in const IMeshDescriptionCommitter::FCommitParams&
+	*  and populates the FMeshDescription pointed to by the MeshDescription pointer inside.
+	*/
+	virtual void CommitMeshDescription(const FCommitter& Committer) = 0;
+
 };
\ No newline at end of file
