ModelingTools: add alpha mask visualization to brush stamp indicator in VertexSculptTool
#rnx

diff --git a/Engine/Plugins/Experimental/MeshModelingToolsetExp/Source/MeshModelingToolsExp/Private/MeshVertexSculptTool.cpp b/Engine/Plugins/Experimental/MeshModelingToolsetExp/Source/MeshModelingToolsExp/Private/MeshVertexSculptTool.cpp
--- a/Engine/Plugins/Experimental/MeshModelingToolsetExp/Source/MeshModelingToolsExp/Private/MeshVertexSculptTool.cpp
+++ b/Engine/Plugins/Experimental/MeshModelingToolsetExp/Source/MeshModelingToolsExp/Private/MeshVertexSculptTool.cpp
@@ -12,6 +12,11 @@
 #include "DynamicMesh/MeshIndexUtil.h"
 #include "Util/BufferUtil.h"
 #include "AssetUtils/Texture2DUtil.h"
+#include "ToolSetupUtil.h"
+
+#include "BaseGizmos/BrushStampIndicator.h"
+#include "PreviewMesh.h"
+#include "Generators/RectangleMeshGenerator.h"
 
 #include "Changes/MeshVertexChange.h"
 
@@ -304,6 +309,44 @@ void UMeshVertexSculptTool::OnPropertyModified(UObject* PropertySet, FProperty*
 	CalculateBrushRadius();
 }
 
+
+
+UPreviewMesh* UMeshVertexSculptTool::MakeBrushIndicatorMesh(UObject* Parent, UWorld* World)
+{
+	UPreviewMesh* PlaneMesh = NewObject<UPreviewMesh>(Parent);
+	PlaneMesh->CreateInWorld(World, FTransform::Identity);
+
+	FRectangleMeshGenerator RectGen;
+	RectGen.Width = RectGen.Height = 2.0;
+	RectGen.WidthVertexCount = RectGen.HeightVertexCount = 1;
+	FDynamicMesh3 Mesh(&RectGen.Generate());
+	FDynamicMeshUVOverlay* UVOverlay = Mesh.Attributes()->PrimaryUV();
+	// configure UVs to be in same space as texture pixels when mapped into brush frame (??)
+	for (int32 eid : UVOverlay->ElementIndicesItr())
+	{
+		FVector2f UV = UVOverlay->GetElement(eid);
+		UV.X = 1.0 - UV.X;
+		UV.Y = 1.0 - UV.Y;
+		UVOverlay->SetElement(eid, UV);
+	}
+	PlaneMesh->UpdatePreview(&Mesh);
+
+	BrushIndicatorMaterial = ToolSetupUtil::GetDefaultBrushAlphaMaterial(GetToolManager());
+	if (BrushIndicatorMaterial)
+	{
+		PlaneMesh->SetMaterial(BrushIndicatorMaterial);
+	}
+
+	return PlaneMesh;
+}
+
+void UMeshVertexSculptTool::InitializeIndicator()
+{
+	UMeshSculptToolBase::InitializeIndicator();
+	// want to draw radius
+	BrushIndicator->bDrawRadiusCircle = true;
+}
+
 void UMeshVertexSculptTool::SetActiveBrushType(int32 Identifier)
 {
 	EMeshVertexSculptBrushType NewBrushType =  static_cast<EMeshVertexSculptBrushType>(Identifier);
@@ -749,6 +792,16 @@ bool UMeshVertexSculptTool::UpdateBrushPosition(const FRay& WorldRay)
 
 
 
+void UMeshVertexSculptTool::UpdateHoverStamp(const FFrame3d& StampFrame)
+{
+	FFrame3d HoverFrame = StampFrame;
+	if (bHaveBrushAlpha && (AlphaProperties->RotationAngle != 0))
+	{
+		HoverFrame.Rotate(FQuaterniond(HoverFrame.Z(), AlphaProperties->RotationAngle, true));
+	}
+	UMeshSculptToolBase::UpdateHoverStamp(HoverFrame);
+}
+
 bool UMeshVertexSculptTool::OnUpdateHover(const FInputDeviceRay& DevicePos)
 {
 	// 4.26 HOTFIX: update LastWorldRay position so that we have it for updating WorkPlane position
@@ -758,7 +811,51 @@ bool UMeshVertexSculptTool::OnUpdateHover(const FInputDeviceRay& DevicePos)
 	if(ensure(InStroke() == false))
 	{
 		UpdateBrushPosition(DevicePos.WorldRay);
+
+		if (BrushIndicatorMaterial)
+		{
+			BrushIndicatorMaterial->SetScalarParameterValue(TEXT("FalloffRatio"), GetCurrentBrushFalloff());
+
+			switch (SculptProperties->PrimaryFalloffType)
+			{
+			default:
+			case EMeshSculptFalloffType::Smooth:
+			case EMeshSculptFalloffType::BoxSmooth:
+				BrushIndicatorMaterial->SetScalarParameterValue(TEXT("FalloffMode"), 0.0f);
+				break;
+			case EMeshSculptFalloffType::Linear:
+			case EMeshSculptFalloffType::BoxLinear:
+				BrushIndicatorMaterial->SetScalarParameterValue(TEXT("FalloffMode"), 0.3333333f);
+				break;
+			case EMeshSculptFalloffType::Inverse:
+			case EMeshSculptFalloffType::BoxInverse:
+				BrushIndicatorMaterial->SetScalarParameterValue(TEXT("FalloffMode"), 0.6666666f);
+				break;
+			case EMeshSculptFalloffType::Round:
+			case EMeshSculptFalloffType::BoxRound:
+				BrushIndicatorMaterial->SetScalarParameterValue(TEXT("FalloffMode"), 1.0f);
+				break;
+			}
+
+			switch (SculptProperties->PrimaryFalloffType)
+			{
+			default:
+			case EMeshSculptFalloffType::Smooth:
+			case EMeshSculptFalloffType::Linear:
+			case EMeshSculptFalloffType::Inverse:
+			case EMeshSculptFalloffType::Round:
+				BrushIndicatorMaterial->SetScalarParameterValue(TEXT("FalloffShape"), 0.0f);
+				break;
+			case EMeshSculptFalloffType::BoxSmooth:
+			case EMeshSculptFalloffType::BoxLinear:
+			case EMeshSculptFalloffType::BoxInverse:
+			case EMeshSculptFalloffType::BoxRound:
+				BrushIndicatorMaterial->SetScalarParameterValue(TEXT("FalloffShape"), 1.0f);
+			}
+			
+		}
 	}
+
 	return true;
 }
 
@@ -1010,12 +1107,19 @@ void UMeshVertexSculptTool::UpdateBrushAlpha(UTexture2D* NewAlpha)
 				BrushAlphaValues = MoveTemp(AlphaValues);
 				BrushAlphaDimensions = AlphaValues.GetDimensions();
 				bHaveBrushAlpha = true;
+
+				BrushIndicatorMaterial->SetTextureParameterValue(TEXT("BrushAlpha"), NewAlpha);
+				BrushIndicatorMaterial->SetScalarParameterValue(TEXT("AlphaPower"), 1.0);
+
 				return;
 			}
 		}
 		bHaveBrushAlpha = false;
 		BrushAlphaValues = TImageBuilder<FVector4f>();
 		BrushAlphaDimensions = FImageDimensions();
+
+		BrushIndicatorMaterial->SetTextureParameterValue(TEXT("BrushAlpha"), nullptr);
+		BrushIndicatorMaterial->SetScalarParameterValue(TEXT("AlphaPower"), 0.0);
 	}
 }
 
diff --git a/Engine/Plugins/Experimental/MeshModelingToolsetExp/Source/MeshModelingToolsExp/Private/Sculpting/MeshSculptToolBase.cpp b/Engine/Plugins/Experimental/MeshModelingToolsetExp/Source/MeshModelingToolsExp/Private/Sculpting/MeshSculptToolBase.cpp
--- a/Engine/Plugins/Experimental/MeshModelingToolsetExp/Source/MeshModelingToolsExp/Private/Sculpting/MeshSculptToolBase.cpp
+++ b/Engine/Plugins/Experimental/MeshModelingToolsetExp/Source/MeshModelingToolsExp/Private/Sculpting/MeshSculptToolBase.cpp
@@ -249,10 +249,16 @@ void UMeshSculptToolBase::OnTick(float DeltaTime)
 
 	UpdateHoverStamp(GetBrushFrameWorld());
 
-	// update brush position here?
-	BrushIndicator->Update((float)GetCurrentBrushRadius(),
-		(FVector)HoverStamp.WorldFrame.Origin, (FVector)HoverStamp.WorldFrame.Z(),
-		1.0f - (float)GetCurrentBrushFalloff());
+	if (bInStroke)
+	{
+		BrushIndicator->Update((float)GetCurrentBrushRadius(),
+			CurrentStamp.WorldFrame.ToFTransform(), 1.0f - (float)GetCurrentBrushFalloff());
+	}
+	else
+	{
+		BrushIndicator->Update((float)GetCurrentBrushRadius(),
+			HoverStamp.WorldFrame.ToFTransform(), 1.0f - (float)GetCurrentBrushFalloff());
+	}
 
 	UpdateWorkPlane();
 }
@@ -1129,7 +1135,7 @@ void UMeshSculptToolBase::InitializeIndicator()
 	// register and spawn brush indicator gizmo
 	GetToolManager()->GetPairedGizmoManager()->RegisterGizmoType(VertexSculptIndicatorGizmoType, NewObject<UBrushStampIndicatorBuilder>());
 	BrushIndicator = GetToolManager()->GetPairedGizmoManager()->CreateGizmo<UBrushStampIndicator>(VertexSculptIndicatorGizmoType, FString(), this);
-	BrushIndicatorMesh = MakeDefaultIndicatorSphereMesh(this, TargetWorld);
+	BrushIndicatorMesh = MakeBrushIndicatorMesh(this, TargetWorld);
 	BrushIndicator->AttachedComponent = BrushIndicatorMesh->GetRootComponent();
 	BrushIndicator->LineThickness = 1.0;
 	BrushIndicator->bDrawIndicatorLines = true;
@@ -1174,12 +1180,12 @@ bool UMeshSculptToolBase::GetIndicatorVisibility() const
 	return BrushIndicator->bVisible;
 }
 
-UPreviewMesh* UMeshSculptToolBase::MakeDefaultIndicatorSphereMesh(UObject* Parent, UWorld* World, int Resolution /*= 32*/)
+UPreviewMesh* UMeshSculptToolBase::MakeBrushIndicatorMesh(UObject* Parent, UWorld* World)
 {
 	UPreviewMesh* SphereMesh = NewObject<UPreviewMesh>(Parent);
 	SphereMesh->CreateInWorld(World, FTransform::Identity);
 	FSphereGenerator SphereGen;
-	SphereGen.NumPhi = SphereGen.NumTheta = Resolution;
+	SphereGen.NumPhi = SphereGen.NumTheta = 32;
 	SphereGen.Generate();
 	FDynamicMesh3 Mesh(&SphereGen);
 	SphereMesh->UpdatePreview(&Mesh);
diff --git a/Engine/Plugins/Experimental/MeshModelingToolsetExp/Source/MeshModelingToolsExp/Public/MeshVertexSculptTool.h b/Engine/Plugins/Experimental/MeshModelingToolsetExp/Source/MeshModelingToolsExp/Public/MeshVertexSculptTool.h
--- a/Engine/Plugins/Experimental/MeshModelingToolsetExp/Source/MeshModelingToolsExp/Public/MeshVertexSculptTool.h
+++ b/Engine/Plugins/Experimental/MeshModelingToolsetExp/Source/MeshModelingToolsExp/Public/MeshVertexSculptTool.h
@@ -217,14 +217,19 @@ public:
 
 protected:
 	// UMeshSculptToolBase API
-	virtual UBaseDynamicMeshComponent* GetSculptMeshComponent() { return DynamicMeshComponent; }
-	virtual FDynamicMesh3* GetBaseMesh() { return &BaseMesh; }
-	virtual const FDynamicMesh3* GetBaseMesh() const { return &BaseMesh; }
+	virtual void InitializeIndicator() override;
+	virtual UPreviewMesh* MakeBrushIndicatorMesh(UObject* Parent, UWorld* World) override;
+
+	virtual UBaseDynamicMeshComponent* GetSculptMeshComponent() override { return DynamicMeshComponent; }
+	virtual FDynamicMesh3* GetBaseMesh() override{ return &BaseMesh; }
+	virtual const FDynamicMesh3* GetBaseMesh() const override{ return &BaseMesh; }
 
 	virtual int32 FindHitSculptMeshTriangle(const FRay3d& LocalRay) override;
 	virtual int32 FindHitTargetMeshTriangle(const FRay3d& LocalRay) override;
 	bool IsHitTriangleBackFacing(int32 TriangleID, const FDynamicMesh3* QueryMesh);
 
+	virtual void UpdateHoverStamp(const FFrame3d& StampFrame) override;
+
 	virtual void OnBeginStroke(const FRay& WorldRay) override;
 	virtual void OnEndStroke() override;
 	// end UMeshSculptToolBase API
diff --git a/Engine/Plugins/Experimental/MeshModelingToolsetExp/Source/MeshModelingToolsExp/Public/Sculpting/MeshSculptToolBase.h b/Engine/Plugins/Experimental/MeshModelingToolsetExp/Source/MeshModelingToolsExp/Public/Sculpting/MeshSculptToolBase.h
--- a/Engine/Plugins/Experimental/MeshModelingToolsetExp/Source/MeshModelingToolsExp/Public/Sculpting/MeshSculptToolBase.h
+++ b/Engine/Plugins/Experimental/MeshModelingToolsetExp/Source/MeshModelingToolsExp/Public/Sculpting/MeshSculptToolBase.h
@@ -522,9 +522,11 @@ protected:
 	// subclasses should call this to create indicator in their ::Setup()
 	virtual void InitializeIndicator();
 
-	virtual bool GetIsVolumetricIndicator();
+	// Called by InitializeIndicator to create a mesh for the brush ROI indicator. Default is sphere.
+	virtual UPreviewMesh* MakeBrushIndicatorMesh(UObject* Parent, UWorld* World);
 
 	virtual void ConfigureIndicator(bool bVolumetric);
+	virtual bool GetIsVolumetricIndicator();
 
 	virtual void SetIndicatorVisibility(bool bVisible);
 	virtual bool GetIndicatorVisibility() const;
@@ -543,10 +545,6 @@ protected:
 	UPROPERTY()
 	TObjectPtr<UPreviewMesh> BrushIndicatorMesh;
 
-	// creates default sphere indicator
-	UPreviewMesh* MakeDefaultIndicatorSphereMesh(UObject* Parent, UWorld* World, int Resolution = 32);
-
-
 
 
 	//
diff --git a/Engine/Plugins/Runtime/MeshModelingToolset/Source/ModelingComponents/Private/ToolSetupUtil.cpp b/Engine/Plugins/Runtime/MeshModelingToolset/Source/ModelingComponents/Private/ToolSetupUtil.cpp
--- a/Engine/Plugins/Runtime/MeshModelingToolset/Source/ModelingComponents/Private/ToolSetupUtil.cpp
+++ b/Engine/Plugins/Runtime/MeshModelingToolset/Source/ModelingComponents/Private/ToolSetupUtil.cpp
@@ -109,6 +109,16 @@ UMaterialInstanceDynamic* ToolSetupUtil::GetDefaultBrushVolumeMaterial(UInteract
 	return nullptr;
 }
 
+UMaterialInstanceDynamic* ToolSetupUtil::GetDefaultBrushAlphaMaterial(UInteractiveToolManager* ToolManager)
+{
+	UMaterial* Material = LoadObject<UMaterial>(nullptr, TEXT("/MeshModelingToolsetExp/Materials/BrushAlphaIndicatorMaterial"));
+	if (Material != nullptr)
+	{
+		UMaterialInstanceDynamic* MatInstance = UMaterialInstanceDynamic::Create(Material, ToolManager);
+		return MatInstance;
+	}
+	return nullptr;
+}
 
 
 UMaterialInterface* ToolSetupUtil::GetDefaultSculptMaterial(UInteractiveToolManager* ToolManager)
diff --git a/Engine/Plugins/Runtime/MeshModelingToolset/Source/ModelingComponents/Public/ToolSetupUtil.h b/Engine/Plugins/Runtime/MeshModelingToolset/Source/ModelingComponents/Public/ToolSetupUtil.h
--- a/Engine/Plugins/Runtime/MeshModelingToolset/Source/ModelingComponents/Public/ToolSetupUtil.h
+++ b/Engine/Plugins/Runtime/MeshModelingToolset/Source/ModelingComponents/Public/ToolSetupUtil.h
@@ -66,6 +66,20 @@ namespace ToolSetupUtil
 	 */
 	MODELINGCOMPONENTS_API UMaterialInstanceDynamic* GetDefaultBrushVolumeMaterial(UInteractiveToolManager* ToolManager);
 
+	/**
+	 * Get a Material to use with a square brush plane/ROI indicator that has an analytic falloff function combined w/ an alpha mask texture.
+	 * Falloff value will be used for both color and opacity, rendered as Additive with depth-testing disabled
+	 * Material Parameters:
+	 *  "BrushAlpha" - Texture2D Parameter for brush alpha texture map
+	 *  "AlphaPower" - configure whether BrushAlpha texture is multiplied w/ analytic falloff (0 -> disabled, 1 -> enabled)
+	 *  "FalloffRatio" - float in range [0,1] that corresponds to width of brush falloff region
+	 *  "FalloffMode" - float in range [0,1] used to interpolate between Smooth (0) -> Linear (0.3333) -> Inverse (0.6666) -> Round (1.0)
+	 *  "FalloffShape" - float in range [0,1] used to interpolate between Circular (0.0) and Square (1.0) 
+	 *  "IntensityScale" - float in range [0,1], multiplier on brightness of falloff highlight (default 0.2)
+	 */
+	MODELINGCOMPONENTS_API UMaterialInstanceDynamic* GetDefaultBrushAlphaMaterial(UInteractiveToolManager* ToolManager);
+
+
 
 	/**
 	 * @return Sculpt Material 1
diff --git a/Engine/Source/Runtime/InteractiveToolsFramework/Private/BaseGizmos/BrushStampIndicator.cpp b/Engine/Source/Runtime/InteractiveToolsFramework/Private/BaseGizmos/BrushStampIndicator.cpp
--- a/Engine/Source/Runtime/InteractiveToolsFramework/Private/BaseGizmos/BrushStampIndicator.cpp
+++ b/Engine/Source/Runtime/InteractiveToolsFramework/Private/BaseGizmos/BrushStampIndicator.cpp
@@ -83,3 +83,29 @@ void UBrushStampIndicator::Update(float Radius, const FVector& Position, const F
 		AttachedComponent->SetWorldTransform(Transform);
 	}
 }
+
+
+
+void UBrushStampIndicator::Update(float Radius, const FTransform& WorldTransform, float Falloff)
+{
+	BrushRadius = Radius;
+	BrushPosition = WorldTransform.GetLocation();
+	BrushNormal = WorldTransform.GetRotation().GetAxisZ();
+	BrushFalloff = Falloff;
+
+	if (AttachedComponent != nullptr)
+	{
+		FTransform Transform = WorldTransform;
+
+		if (ScaleInitializedComponent != AttachedComponent)
+		{
+			InitialComponentScale = AttachedComponent->GetComponentTransform().GetScale3D();
+			InitialComponentScale *= 1.0f / InitialComponentScale.Z;
+			ScaleInitializedComponent = AttachedComponent;
+		}
+
+		Transform.SetScale3D(Radius * InitialComponentScale);
+
+		AttachedComponent->SetWorldTransform(Transform);
+	}
+}
diff --git a/Engine/Source/Runtime/InteractiveToolsFramework/Public/BaseGizmos/BrushStampIndicator.h b/Engine/Source/Runtime/InteractiveToolsFramework/Public/BaseGizmos/BrushStampIndicator.h
--- a/Engine/Source/Runtime/InteractiveToolsFramework/Public/BaseGizmos/BrushStampIndicator.h
+++ b/Engine/Source/Runtime/InteractiveToolsFramework/Public/BaseGizmos/BrushStampIndicator.h
@@ -42,6 +42,11 @@ public:
 	 */
 	virtual void Update(float Radius, const FVector& Position, const FVector& Normal, float Falloff);
 
+	/**
+	* Update the Radius, Transform and Falloff of the Stamp Indicator
+	*/
+	virtual void Update(float Radius, const FTransform& WorldTransform, float Falloff);
+
 public:
 
 	/** Controls whether Gizmo will draw visual elements. Does not currently affect AttachedComponent. */
