MeshModelingToolset: Cleanup Tangents tool

diff --git a/Engine/Plugins/Experimental/MeshModelingToolsetExp/Source/MeshModelingToolsEditorOnly/Private/MeshTangentsTool.cpp b/Engine/Plugins/Experimental/MeshModelingToolsetExp/Source/MeshModelingToolsEditorOnly/Private/MeshTangentsTool.cpp
--- a/Engine/Plugins/Experimental/MeshModelingToolsetExp/Source/MeshModelingToolsEditorOnly/Private/MeshTangentsTool.cpp
+++ b/Engine/Plugins/Experimental/MeshModelingToolsetExp/Source/MeshModelingToolsEditorOnly/Private/MeshTangentsTool.cpp
@@ -7,7 +7,6 @@
 #include "DynamicMesh/DynamicMesh3.h"
 #include "DynamicMeshToMeshDescription.h"
 #include "MeshDescription.h"
-#include "MeshDescriptionToDynamicMesh.h"
 #include "ToolSetupUtil.h"
 #include "ToolDataVisualizer.h"
 
@@ -90,7 +89,7 @@ void UMeshTangentsTool::Setup()
 	Settings->RestoreProperties(this);
 	AddToolPropertySource(Settings);
 
-	Settings->WatchProperty(Settings->TangentType, [this](EMeshTangentsType) { Compute->InvalidateResult(); });
+	Settings->WatchProperty(Settings->CalculationMethod, [this](EMeshTangentsType) { Compute->InvalidateResult(); });
 	Settings->WatchProperty(Settings->LineLength, [this](float) { bLengthDirty = true; });
 	Settings->WatchProperty(Settings->LineThickness, [this](float) { bThicknessDirty = true; });
 	Settings->WatchProperty(Settings->bShowTangents, [this](float) { bVisibilityChanged = true; });
@@ -170,14 +169,14 @@ void UMeshTangentsTool::Render(IToolsContextRenderAPI* RenderAPI)
 		Visualizer.SetTransform(PreviewMesh->GetTransform());
 		for (const FMikktDeviation& ErrorPt : Deviations)
 		{
-			if (ErrorPt.MaxAngleDeg > Settings->AngleThreshDeg)
+			if (ErrorPt.MaxAngleDeg > Settings->CompareWithMikktThreshold)
 			{
-				Visualizer.DrawPoint(ErrorPt.VertexPos, FLinearColor(0.95, 0.05, 0.05), 6 * Settings->LineThickness, false);
-				Visualizer.DrawLine<FVector3f>(ErrorPt.VertexPos, ErrorPt.VertexPos + Settings->LineLength * ErrorPt.MikktTangent, FLinearColor(0.95, 0.05, 0.05), 2 * Settings->LineThickness, false);
-				Visualizer.DrawLine<FVector3f>(ErrorPt.VertexPos, ErrorPt.VertexPos + Settings->LineLength * ErrorPt.MikktBitangent, FLinearColor(0.05, 0.95, 0.05), 2 * Settings->LineThickness, false);
+				Visualizer.DrawPoint(ErrorPt.VertexPos, FLinearColor(0.95f, 0.05f, 0.05f), 6.0f * Settings->LineThickness, false);
+				Visualizer.DrawLine<FVector3f>(ErrorPt.VertexPos, ErrorPt.VertexPos + Settings->LineLength * ErrorPt.MikktTangent, FLinearColor(0.95f, 0.05f, 0.05f), 2 * Settings->LineThickness, false);
+				Visualizer.DrawLine<FVector3f>(ErrorPt.VertexPos, ErrorPt.VertexPos + Settings->LineLength * ErrorPt.MikktBitangent, FLinearColor(0.05f, 0.95f, 0.05f), 2 * Settings->LineThickness, false);
 
-				Visualizer.DrawLine<FVector3f>(ErrorPt.VertexPos, ErrorPt.VertexPos + (1.1f*Settings->LineLength) * ErrorPt.OtherTangent, FLinearColor(0.95, 0.50, 0.05), Settings->LineThickness, false);
-				Visualizer.DrawLine<FVector3f>(ErrorPt.VertexPos, ErrorPt.VertexPos + (1.1f*Settings->LineLength) * ErrorPt.OtherBitangent, FLinearColor(0.05, 0.95, 0.95), Settings->LineThickness, false);
+				Visualizer.DrawLine<FVector3f>(ErrorPt.VertexPos, ErrorPt.VertexPos + (1.1f * Settings->LineLength) * ErrorPt.OtherTangent, FLinearColor(0.95f, 0.50f, 0.05f), Settings->LineThickness, false);
+				Visualizer.DrawLine<FVector3f>(ErrorPt.VertexPos, ErrorPt.VertexPos + (1.1f * Settings->LineLength) * ErrorPt.OtherBitangent, FLinearColor(0.05f, 0.95f, 0.95f), Settings->LineThickness, false);
 			}
 		}
 		Visualizer.EndFrame();
@@ -197,7 +196,7 @@ TUniquePtr<TGenericDataOperator<FMeshTangentsd>> UMeshTangentsTool::MakeNewOpera
 
 	TangentsOp->SourceMesh = InputMesh;
 	TangentsOp->SourceTangents = InitialTangents;
-	TangentsOp->CalculationMethod = Settings->TangentType;
+	TangentsOp->CalculationMethod = Settings->CalculationMethod;
 
 	return TangentsOp;
 }
@@ -214,7 +213,7 @@ void UMeshTangentsTool::UpdateVisualization(bool bThicknessChanged, bool bLength
 
 	if (bThicknessChanged)
 	{
-		float Thickness = (Settings->LineThickness-2.0) / 10.0;
+		float Thickness = Settings->LineThickness;
 		TangentLines->SetAllLinesThickness(Thickness);
 		NormalLines->SetAllLinesThickness(Thickness);
 	}
@@ -233,11 +232,11 @@ void UMeshTangentsTool::UpdateVisualization(bool bThicknessChanged, bool bLength
 
 void UMeshTangentsTool::OnTangentsUpdated(const TUniquePtr<FMeshTangentsd>& NewResult)
 {
-	float LineLength = Settings->LineLength;
-	float Thickness = (Settings->LineThickness-2.0) / 10.0;
+	const float LineLength = Settings->LineLength;
+	const float Thickness = Settings->LineThickness;
 
 	TSet<int32> DegenerateTris;
-	if (Settings->bHideDegenerates || (Settings->bCompareWithMikkt && Settings->TangentType != EMeshTangentsType::MikkTSpace) )
+	if (!Settings->bShowDegenerates || (Settings->bCompareWithMikkt && Settings->CalculationMethod != EMeshTangentsType::MikkTSpace) )
 	{
 		FMeshTangentsd DegenTangents(InputMesh.Get());
 		DegenTangents.ComputeTriangleTangents(InputMesh->Attributes()->GetUVLayer(0));
@@ -282,7 +281,7 @@ void UMeshTangentsTool::OnTangentsUpdated(const TUniquePtr<FMeshTangentsd>& NewR
 
 	// calculate deviation between what we have and MikkT, if necessary
 	Deviations.Reset();
-	if (Settings->bCompareWithMikkt && Settings->TangentType != EMeshTangentsType::MikkTSpace)
+	if (Settings->bCompareWithMikkt && Settings->CalculationMethod != EMeshTangentsType::MikkTSpace)
 	{
 		FProgressCancel TmpCancel;
 		FCalculateTangentsOp MikktOp;
@@ -307,7 +306,8 @@ void UMeshTangentsTool::OnTangentsUpdated(const TUniquePtr<FMeshTangentsd>& NewR
 				{
 					FVector3f TangentMikkt, BitangentMikkt;
 					MikktTangents->GetPerTriangleTangent<FVector3f, float>(Index, j, TangentMikkt, BitangentMikkt);
-					UE::Geometry::Normalize(TangentMikkt); UE::Geometry::Normalize(BitangentMikkt);
+					UE::Geometry::Normalize(TangentMikkt);
+					UE::Geometry::Normalize(BitangentMikkt);
 					FVector3f TangentNew, BitangentNew;
 					NewTangents->GetPerTriangleTangent<FVector3f, float>(Index, j, TangentNew, BitangentNew);
 					UE::Geometry::Normalize(TangentNew); 
diff --git a/Engine/Plugins/Experimental/MeshModelingToolsetExp/Source/MeshModelingToolsEditorOnly/Public/MeshTangentsTool.h b/Engine/Plugins/Experimental/MeshModelingToolsetExp/Source/MeshModelingToolsEditorOnly/Public/MeshTangentsTool.h
--- a/Engine/Plugins/Experimental/MeshModelingToolsetExp/Source/MeshModelingToolsEditorOnly/Public/MeshTangentsTool.h
+++ b/Engine/Plugins/Experimental/MeshModelingToolsetExp/Source/MeshModelingToolsEditorOnly/Public/MeshTangentsTool.h
@@ -3,8 +3,6 @@
 #pragma once
 
 #include "CoreMinimal.h"
-#include "SingleSelectionTool.h"
-#include "InteractiveToolBuilder.h"
 #include "PreviewMesh.h"
 #include "Drawing/PreviewGeometryActor.h"
 #include "MeshOpPreviewHelpers.h"
@@ -15,7 +13,7 @@
 #include "MeshTangentsTool.generated.h"
 
 
-// predeclarations
+// Forward declarations
 struct FMeshDescription;
 class UDynamicMeshComponent;
 class UMaterialInterface;
@@ -45,31 +43,41 @@ class MESHMODELINGTOOLSEDITORONLY_API UMeshTangentsToolProperties : public UInte
 	GENERATED_BODY()
 public:
 
-	UPROPERTY(EditAnywhere, Category = Options)
-	EMeshTangentsType TangentType = EMeshTangentsType::FastMikkTSpace;
+	/** Method used for calculating the tangents */
+	UPROPERTY(EditAnywhere, Category = TangentsCalculation)
+	EMeshTangentsType CalculationMethod = EMeshTangentsType::FastMikkTSpace;
 
-
-	UPROPERTY(EditAnywhere, Category = Visualization)
+	/** Display the mesh tangents */
+	UPROPERTY(EditAnywhere, Category = Display)
 	bool bShowTangents = true;
 
-	UPROPERTY(EditAnywhere, Category = Visualization)
+	/** Display the mesh normals */
+	UPROPERTY(EditAnywhere, Category = Display)
 	bool bShowNormals = false;
 
-	UPROPERTY(EditAnywhere, Category = Visualization, AdvancedDisplay)
-	bool bHideDegenerates = true;
-
-	UPROPERTY(EditAnywhere, Category = Visualization, AdvancedDisplay, meta = (UIMin = "0.01", UIMax = "25.0", ClampMin = "0.01", ClampMax = "10000000.0"))
-	float LineLength = 2.0;
+	/** Length of lines used for displaying tangents and/or normals */
+	UPROPERTY(EditAnywhere, Category = Display, meta = (UIMin = "0.01", UIMax = "25.0", ClampMin = "0.01", ClampMax = "10000000.0"))
+	float LineLength = 2.0f;
 
-	UPROPERTY(EditAnywhere, Category = Visualization, AdvancedDisplay, meta = (UIMin = "0", UIMax = "25.0", ClampMin = "0", ClampMax = "1000.0"))
-	float LineThickness = 3.0;
+	/** Thickness of lines used for displaying tangents and/or normals */
+	UPROPERTY(EditAnywhere, Category = Display, meta = (UIMin = "0", UIMax = "25.0", ClampMin = "0", ClampMax = "1000.0"))
+	float LineThickness = 1.0f;
 
+	/** Display tangents and/or normals for degenerate triangles */
+	UPROPERTY(EditAnywhere, Category = Display, AdvancedDisplay)
+	bool bShowDegenerates = false;
 
-	UPROPERTY(EditAnywhere, Category = Visualization, AdvancedDisplay)
+	/** Display difference between the current tangents and tangents calculated with standard MikkTSpace calculation.
+	 * This is only available if a Calculation Method other than MikkTSpace is selected. */
+	UPROPERTY(EditAnywhere, Category = Display, AdvancedDisplay, meta = (DisplayName = "Compare with MikkT",
+		EditCondition = "CalculationMethod != EMeshTangentsType::MikkTSpace"))
 	bool bCompareWithMikkt = false;
 
-	UPROPERTY(EditAnywhere, Category = Visualization, AdvancedDisplay, meta = (UIMin = "0.5", UIMax = "90.0"))
-	float AngleThreshDeg = 5.0;
+	/** Minimum angle difference in degrees for a tangent to be considered different to the standard MikkTSpace calculation
+	 * This is only available if a Compare with MikkT is enabled and a Calculation Method other than MikkTSpace is selected. */
+	UPROPERTY(EditAnywhere, Category = Display, AdvancedDisplay, meta = (DisplayName = "Compare Threshold", UIMin = "0.5", UIMax = "90.0",
+		EditCondition = "CalculationMethod != EMeshTangentsType::MikkTSpace && bCompareWithMikkt"))
+	float CompareWithMikktThreshold = 5.0f;
 
 };
 
@@ -104,7 +112,6 @@ protected:
 	UPROPERTY()
 	TObjectPtr<UMeshTangentsToolProperties> Settings = nullptr;
 
-protected:
 	UPROPERTY()
 	TObjectPtr<UMaterialInterface> DefaultMaterial = nullptr;
 
@@ -116,7 +123,6 @@ protected:
 
 	TUniquePtr<TGenericDataBackgroundCompute<UE::Geometry::FMeshTangentsd>> Compute = nullptr;
 
-protected:
 	TSharedPtr<FMeshDescription, ESPMode::ThreadSafe> InputMeshDescription;
 	TSharedPtr<UE::Geometry::FMeshTangentsf, ESPMode::ThreadSafe> InitialTangents;
 	TSharedPtr<UE::Geometry::FDynamicMesh3, ESPMode::ThreadSafe> InputMesh;
diff --git a/Engine/Plugins/Runtime/MeshModelingToolset/Source/MeshModelingTools/Public/Properties/RevolveProperties.h b/Engine/Plugins/Runtime/MeshModelingToolset/Source/MeshModelingTools/Public/Properties/RevolveProperties.h
--- a/Engine/Plugins/Runtime/MeshModelingToolset/Source/MeshModelingTools/Public/Properties/RevolveProperties.h
+++ b/Engine/Plugins/Runtime/MeshModelingToolset/Source/MeshModelingTools/Public/Properties/RevolveProperties.h
@@ -133,8 +133,8 @@ public:
 	UPROPERTY(EditAnywhere, Category = Revolve, AdvancedDisplay, meta = (DisplayName = "Path at Midpoint"))
 	bool bPathAtMidpointOfStep = false;
 
-	/** How Polygroups are assigned to shape primitives. If caps are generated, they will always be in separate groups. */
-	UPROPERTY(EditAnywhere, Category = Revolve, AdvancedDisplay)
+	/** How PolyGroups are assigned to shape primitives. If caps are generated, they will always be in separate groups. */
+	UPROPERTY(EditAnywhere, Category = Revolve, AdvancedDisplay, meta = (DisplayName = "PolyGroup Mode"))
 	ERevolvePropertiesPolygroupMode PolygroupMode = ERevolvePropertiesPolygroupMode::PerFace;
 
 	/** How generated quads are split into triangles. */
diff --git a/Engine/Plugins/Runtime/MeshModelingToolset/Source/ModelingOperatorsEditorOnly/Public/ParameterizationOps/CalculateTangentsOp.h b/Engine/Plugins/Runtime/MeshModelingToolset/Source/ModelingOperatorsEditorOnly/Public/ParameterizationOps/CalculateTangentsOp.h
--- a/Engine/Plugins/Runtime/MeshModelingToolset/Source/ModelingOperatorsEditorOnly/Public/ParameterizationOps/CalculateTangentsOp.h
+++ b/Engine/Plugins/Runtime/MeshModelingToolset/Source/ModelingOperatorsEditorOnly/Public/ParameterizationOps/CalculateTangentsOp.h
@@ -9,17 +9,16 @@
 #include "CalculateTangentsOp.generated.h"
 
 
-
 UENUM()
 enum class EMeshTangentsType : uint8
 {
-	/** Standard MikkTSpace Tangent Calculation */
+	/** Standard MikkTSpace tangent calculation */
 	MikkTSpace = 0,
-	/** MikkT-like Blended Per-Triangle Tangents, with Blending based on existing Mesh/UV/Normal Topology */
+	/** MikkTSpace-like blended per-triangle tangents, with the blending being based on existing mesh, normals, and UV topology */
 	FastMikkTSpace = 1,
-	/** Project per-triangle Tangents projected onto Normals */
+	/** Project per-triangle tangents onto normals */
 	PerTriangle = 2,
-	/** Use Source Mesh Tangents */
+	/** Use existing source mesh tangents */
 	CopyExisting = 3
 };
 
@@ -28,15 +27,14 @@ namespace UE
 namespace Geometry
 {
 
-
-class MODELINGOPERATORSEDITORONLY_API FCalculateTangentsOp : public TGenericDataOperator<UE::Geometry::FMeshTangentsd>
+class MODELINGOPERATORSEDITORONLY_API FCalculateTangentsOp : public TGenericDataOperator<FMeshTangentsd>
 {
 public:
 	virtual ~FCalculateTangentsOp() {}
 
 	// inputs
 	TSharedPtr<FDynamicMesh3, ESPMode::ThreadSafe> SourceMesh;
-	TSharedPtr<UE::Geometry::FMeshTangentsf, ESPMode::ThreadSafe> SourceTangents;
+	TSharedPtr<FMeshTangentsf, ESPMode::ThreadSafe> SourceTangents;
 
 	// parameters
 	EMeshTangentsType CalculationMethod;
@@ -50,12 +48,10 @@ public:
 
 	virtual void CalculateResult(FProgressCancel* Progress) override;
 
-
-	virtual void CalculateStandard(FProgressCancel* Progress, TUniquePtr<UE::Geometry::FMeshTangentsd>& Tangents);
-	virtual void CalculateMikkT(FProgressCancel* Progress, TUniquePtr<UE::Geometry::FMeshTangentsd>& Tangents);
-	virtual void CopyFromSource(FProgressCancel* Progress, TUniquePtr<UE::Geometry::FMeshTangentsd>& Tangents);
+	virtual void CalculateStandard(FProgressCancel* Progress, TUniquePtr<FMeshTangentsd>& Tangents);
+	virtual void CalculateMikkT(FProgressCancel* Progress, TUniquePtr<FMeshTangentsd>& Tangents);
+	virtual void CopyFromSource(FProgressCancel* Progress, TUniquePtr<FMeshTangentsd>& Tangents);
 };
 
 } // end namespace UE::Geometry
 } // end namespace UE
-
