ModelingTools: fix bad conversion to new mesh-commit paths in EditNormalsTool
#rnx

diff --git a/Engine/Plugins/Experimental/MeshModelingToolsetExp/Source/MeshModelingToolsExp/Private/EditNormalsTool.cpp b/Engine/Plugins/Experimental/MeshModelingToolsetExp/Source/MeshModelingToolsExp/Private/EditNormalsTool.cpp
--- a/Engine/Plugins/Experimental/MeshModelingToolsetExp/Source/MeshModelingToolsExp/Private/EditNormalsTool.cpp
+++ b/Engine/Plugins/Experimental/MeshModelingToolsetExp/Source/MeshModelingToolsExp/Private/EditNormalsTool.cpp
@@ -261,8 +261,11 @@ void UEditNormalsTool::GenerateAsset(const TArray<FDynamicMeshOpResult>& Results
 		}
 
 		const FDynamicMesh3* NewDynamicMesh = Results[ComponentIdx].Mesh.Get();
-		check(NewDynamicMesh != nullptr);
-		UE::ToolTarget::CommitDynamicMeshNormalsUpdate(Targets[ComponentIdx], NewDynamicMesh, false);
+		if (NewDynamicMesh)
+		{
+			bool bTopologyChanged = BasicProperties->WillTopologyChange();
+			UE::ToolTarget::CommitMeshDescriptionUpdateViaDynamicMesh(Targets[ComponentIdx], *NewDynamicMesh, bTopologyChanged);
+		}
 	}
 
 	GetToolManager()->EndUndoTransaction();
