fix brush size visualization not updating in DSclpt tool when brush size is changed via hotkeys

diff --git a/Engine/Plugins/Experimental/MeshModelingToolsetExp/Source/MeshModelingToolsExp/Private/DynamicMeshSculptTool.cpp b/Engine/Plugins/Experimental/MeshModelingToolsetExp/Source/MeshModelingToolsExp/Private/DynamicMeshSculptTool.cpp
--- a/Engine/Plugins/Experimental/MeshModelingToolsetExp/Source/MeshModelingToolsExp/Private/DynamicMeshSculptTool.cpp
+++ b/Engine/Plugins/Experimental/MeshModelingToolsetExp/Source/MeshModelingToolsExp/Private/DynamicMeshSculptTool.cpp
@@ -2198,24 +2198,28 @@ void UDynamicMeshSculptTool::IncreaseBrushRadiusAction()
 {	
 	BrushProperties->BrushSize.IncreaseRadius(false);
 	NotifyOfPropertyChangeByTool(BrushProperties);
+	CalculateBrushRadius();
 }
 
 void UDynamicMeshSculptTool::DecreaseBrushRadiusAction()
 {
 	BrushProperties->BrushSize.DecreaseRadius(false);
 	NotifyOfPropertyChangeByTool(BrushProperties);
+	CalculateBrushRadius();
 }
 
 void UDynamicMeshSculptTool::IncreaseBrushRadiusSmallStepAction()
 {
 	BrushProperties->BrushSize.IncreaseRadius(true);
 	NotifyOfPropertyChangeByTool(BrushProperties);
+	CalculateBrushRadius();
 }
 
 void UDynamicMeshSculptTool::DecreaseBrushRadiusSmallStepAction()
 {
 	BrushProperties->BrushSize.DecreaseRadius(true);
 	NotifyOfPropertyChangeByTool(BrushProperties);
+	CalculateBrushRadius();
 }
 
 void UDynamicMeshSculptTool::IncreaseBrushSpeedAction()
