ModelingTools: Fix the MeshInspector's material override settings to be functional with the generated checkerboard material.

#rnx

diff --git a/Engine/Plugins/Experimental/MeshModelingToolsetExp/Source/MeshModelingToolsExp/Private/EditUVIslandsTool.cpp b/Engine/Plugins/Experimental/MeshModelingToolsetExp/Source/MeshModelingToolsExp/Private/EditUVIslandsTool.cpp
--- a/Engine/Plugins/Experimental/MeshModelingToolsetExp/Source/MeshModelingToolsExp/Private/EditUVIslandsTool.cpp
+++ b/Engine/Plugins/Experimental/MeshModelingToolsetExp/Source/MeshModelingToolsExp/Private/EditUVIslandsTool.cpp
@@ -112,6 +112,13 @@ void UEditUVIslandsTool::Setup()
 
 
 	MaterialSettings = NewObject<UExistingMeshMaterialProperties>(this);
+	const FDynamicMesh3* TargetMesh = DynamicMeshComponent->GetMesh();
+	TArray<FString> UVChannelNamesList;
+	for (int32 k = 0; k < TargetMesh->Attributes()->NumUVLayers(); ++k)
+	{
+		UVChannelNamesList.Add(FString::Printf(TEXT("UV %d"), k));
+	}
+	MaterialSettings->UpdateUVChannels(0, UVChannelNamesList);
 	MaterialSettings->RestoreProperties(this);
 	AddToolPropertySource(MaterialSettings);
 	MaterialSettings->GetOnModified().AddLambda([this](UObject*, FProperty*)
diff --git a/Engine/Plugins/Experimental/MeshModelingToolsetExp/Source/MeshModelingToolsExp/Private/MeshInspectorTool.cpp b/Engine/Plugins/Experimental/MeshModelingToolsetExp/Source/MeshModelingToolsExp/Private/MeshInspectorTool.cpp
--- a/Engine/Plugins/Experimental/MeshModelingToolsetExp/Source/MeshModelingToolsExp/Private/MeshInspectorTool.cpp
+++ b/Engine/Plugins/Experimental/MeshModelingToolsetExp/Source/MeshModelingToolsExp/Private/MeshInspectorTool.cpp
@@ -79,6 +79,13 @@ void UMeshInspectorTool::Setup()
 	AddToolPropertySource(Settings);
 
 	MaterialSettings = NewObject<UExistingMeshMaterialProperties>(this);
+	const FDynamicMesh3* TargetMesh = PreviewMesh->GetPreviewDynamicMesh();
+	TArray<FString> UVChannelNamesList;
+	for (int32 k = 0; k < TargetMesh->Attributes()->NumUVLayers(); ++k)
+	{
+		UVChannelNamesList.Add(FString::Printf(TEXT("UV %d"), k));
+	}
+	MaterialSettings->UpdateUVChannels(0, UVChannelNamesList);
 	MaterialSettings->RestoreProperties(this);
 	AddToolPropertySource(MaterialSettings);
 
