UniformTesselate.h/cpp
Added computation of the vertex/triangles mappings between the input mesh and the new tessellated mesh. Specifically, we compute [new vertex -> original triangles] and [new triangle -> original triangle] mappings.
Added a constructor for tesselating mesh in place.

DisplaceMeshTool.cpp
Removed recursive SubdivideMesh function and replaced it with the call to FUniformTesselate.
Adjusted how the max allowed Subdivision number is calculated since we are not doing it recursively anymore.

PNTriangles.h/cpp
Use FUniformTesselate instead of the recursive subdivision.

OperationsTest.cpp/CubeExpectedResult.obj
Updated the PNTriangles unit test to reflect the change in the tesselation.

#rnx

diff --git a/Engine/Plugins/Experimental/MeshModelingToolsetExp/Source/MeshModelingToolsExp/Private/DisplaceMeshTool.cpp b/Engine/Plugins/Experimental/MeshModelingToolsetExp/Source/MeshModelingToolsExp/Private/DisplaceMeshTool.cpp
--- a/Engine/Plugins/Experimental/MeshModelingToolsetExp/Source/MeshModelingToolsExp/Private/DisplaceMeshTool.cpp
+++ b/Engine/Plugins/Experimental/MeshModelingToolsetExp/Source/MeshModelingToolsExp/Private/DisplaceMeshTool.cpp
@@ -17,6 +17,8 @@
 #include "MeshDescription.h"
 #include "ModelingToolTargetUtil.h"
 #include "Operations/PNTriangles.h"
+#include "Operations/UniformTesselate.h"
+
 
 // needed to disable normals recalculation on the underlying asset
 #include "AssetUtils/MeshDescriptionUtil.h"
@@ -29,56 +31,6 @@ using namespace UE::Geometry;
 
 namespace DisplaceMeshToolLocals{
 
-	void SubdivideMesh(FDynamicMesh3& Mesh, FProgressCancel* ProgressCancel)
-	{
-		TArray<int> EdgesToProcess;
-		for (int tid : Mesh.EdgeIndicesItr())
-		{
-			EdgesToProcess.Add(tid);
-		}
-		int MaxTriangleID = Mesh.MaxTriangleID();
-
-		if (ProgressCancel && ProgressCancel->Cancelled()) return;
-
-		TArray<int> TriSplitEdges;
-		TriSplitEdges.Init(-1, Mesh.MaxTriangleID());
-
-		for (int eid : EdgesToProcess)
-		{
-			FIndex2i EdgeTris = Mesh.GetEdgeT(eid);
-
-			FDynamicMesh3::FEdgeSplitInfo SplitInfo;
-			EMeshResult result = Mesh.SplitEdge(eid, SplitInfo);
-			if (result == EMeshResult::Ok)
-			{
-				if (EdgeTris.A < MaxTriangleID && TriSplitEdges[EdgeTris.A] == -1)
-				{
-					TriSplitEdges[EdgeTris.A] = SplitInfo.NewEdges.B;
-				}
-				if (EdgeTris.B != FDynamicMesh3::InvalidID)
-				{
-					if (EdgeTris.B < MaxTriangleID && TriSplitEdges[EdgeTris.B] == -1)
-					{
-						TriSplitEdges[EdgeTris.B] = SplitInfo.NewEdges.C;
-					}
-				}
-			}
-
-			if (ProgressCancel && ProgressCancel->Cancelled()) return;
-		}
-
-		for (int eid : TriSplitEdges)
-		{
-			if (eid != -1)
-			{
-				FDynamicMesh3::FEdgeFlipInfo FlipInfo;
-				Mesh.FlipEdge(eid, FlipInfo);
-
-				if (ProgressCancel && ProgressCancel->Cancelled()) return;
-			}
-		}
-	}
-
 	namespace ComputeDisplacement 
 	{
 		/// Directional Filter: Scale displacement for a given vertex based on how well 
@@ -292,21 +244,20 @@ namespace DisplaceMeshToolLocals{
 				ResultMesh->SetVertexUV(vid, FVector2f::One());
 			}
 		}
-
 	}
 
 	void FSubdivideMeshOp::CalculateResult(FProgressCancel* ProgressCancel)
 	{
-		using namespace DisplaceMeshToolLocals;
-
-		// calculate subdivisions (todo: move to elsewhere)
 		if (SubdivisionType == ESubdivisionType::Loop) 
 		{
-			for (int ri = 0; ri < SubdivisionsCount; ri++)
+			FUniformTesselate Tesselator(ResultMesh.Get());
+			Tesselator.Progress = ProgressCancel;
+			Tesselator.TesselationNum = SubdivisionsCount;
+						
+			if (Tesselator.Validate() == EOperationValidationResult::Ok) 
 			{
-				if (ProgressCancel && ProgressCancel->Cancelled()) return;
-				SubdivideMesh(*ResultMesh, ProgressCancel);
-			} 
+				Tesselator.Compute();
+			}
 		}
 		else if (SubdivisionType == ESubdivisionType::PNTriangles) 
 		{
@@ -973,7 +924,7 @@ void UDisplaceMeshTool::ValidateSubdivisions()
 
 	constexpr int MaxTriangles = 3000000;
 	double NumTriangles = OriginalMesh.MaxTriangleID();
-	int MaxSubdivisions = (int)floor(log2(MaxTriangles / NumTriangles) / 2.0);
+	int MaxSubdivisions = (int)(FMath::Sqrt(MaxTriangles/NumTriangles) - 1);
 	if (CommonProperties->Subdivisions > MaxSubdivisions)
 	{
 		if (bIsInitialized)		// only show warning after initial tool startup
