ModelingTools: Add error message if target mesh does not have valid tangents when baking Tangent Normals or Bent Normals.

#rnx

diff --git a/Engine/Plugins/Experimental/MeshModelingToolsetExp/Source/MeshModelingToolsExp/Private/BakeMeshAttributeMapsTool.cpp b/Engine/Plugins/Experimental/MeshModelingToolsetExp/Source/MeshModelingToolsExp/Private/BakeMeshAttributeMapsTool.cpp
--- a/Engine/Plugins/Experimental/MeshModelingToolsetExp/Source/MeshModelingToolsExp/Private/BakeMeshAttributeMapsTool.cpp
+++ b/Engine/Plugins/Experimental/MeshModelingToolsetExp/Source/MeshModelingToolsExp/Private/BakeMeshAttributeMapsTool.cpp
@@ -487,6 +487,21 @@ void UBakeMeshAttributeMapsTool::Shutdown(EToolShutdownType ShutdownType)
 }
 
 
+bool UBakeMeshAttributeMapsTool::ValidDetailMeshTangents()
+{
+	if (!DetailMesh)
+	{
+		return false;
+	}
+	if (bCheckDetailMeshTangents)
+	{
+		bValidDetailMeshTangents = FDynamicMeshTangents(DetailMesh.Get()).HasValidTangents(true);
+		bCheckDetailMeshTangents = false;
+	}
+	return bValidDetailMeshTangents;
+}
+
+
 void UBakeMeshAttributeMapsTool::UpdateDetailMesh()
 {
 	UToolTarget* DetailTarget = Targets[bIsBakeToSelf ? 0 : 1];
@@ -577,6 +592,8 @@ void UBakeMeshAttributeMapsTool::UpdateResult()
 	// Clear our invalid bitflag to check again for valid inputs.
 	OpState &= ~EBakeOpState::Invalid;
 
+	OpState |= UpdateResult_TargetMeshTangents(CachedBakeSettings.BakeMapTypes);
+
 	OpState |= UpdateResult_DetailNormalMap();
 
 	// Update map type settings
@@ -623,6 +640,19 @@ void UBakeMeshAttributeMapsTool::UpdateResult()
 }
 
 
+EBakeOpState UBakeMeshAttributeMapsTool::UpdateResult_DetailMeshTangents(EBakeMapType BakeType)
+{
+	EBakeOpState ResultState = EBakeOpState::Clean;
+	const bool bNeedDetailMeshTangents = (bool)(BakeType & (EBakeMapType::TangentSpaceNormal | EBakeMapType::BentNormal)); 
+	if (bNeedDetailMeshTangents && !ValidDetailMeshTangents())
+	{
+		GetToolManager()->DisplayMessage(LOCTEXT("InvalidSourceTangentsWarning", "The Source Mesh does not have valid tangents."), EToolMessageLevel::UserWarning);
+		ResultState = EBakeOpState::Invalid;
+	}
+	return ResultState;
+}
+
+
 EBakeOpState UBakeMeshAttributeMapsTool::UpdateResult_DetailNormalMap()
 {
 	EBakeOpState ResultState = EBakeOpState::Clean;
@@ -631,19 +661,21 @@ EBakeOpState UBakeMeshAttributeMapsTool::UpdateResult_DetailNormalMap()
 	const FDynamicMeshUVOverlay* UVOverlay = DetailMesh->Attributes()->GetUVLayer(DetailUVLayer);
 	if (UVOverlay == nullptr)
 	{
-		GetToolManager()->DisplayMessage(LOCTEXT("InvalidUVWarning", "The Detail Mesh does not have the selected UV layer"), EToolMessageLevel::UserWarning);
+		GetToolManager()->DisplayMessage(LOCTEXT("InvalidUVWarning", "The Source Mesh does not have the selected UV layer"), EToolMessageLevel::UserWarning);
 		return EBakeOpState::Invalid;
 	}
 	
-	UTexture2D* DetailMeshNormalMap = InputMeshSettings->SourceNormalMap; 
-	if (DetailMeshNormalMap)
+	if (UTexture2D* DetailMeshNormalMap = InputMeshSettings->SourceNormalMap)
 	{
 		CachedDetailNormalMap = MakeShared<UE::Geometry::TImageBuilder<FVector4f>, ESPMode::ThreadSafe>();
 		if (!UE::AssetUtils::ReadTexture(DetailMeshNormalMap, *CachedDetailNormalMap, bPreferPlatformData))
 		{
 			// Report the failed texture read as a warning, but permit the bake to continue.
-			GetToolManager()->DisplayMessage(LOCTEXT("CannotReadTextureWarning", "Cannot read from the detail normal map texture"), EToolMessageLevel::UserWarning);
+			GetToolManager()->DisplayMessage(LOCTEXT("CannotReadTextureWarning", "Cannot read from the source normal map texture"), EToolMessageLevel::UserWarning);
 		}
+
+		// Validate detail mesh tangents for detail normal map
+		ResultState |= UpdateResult_DetailMeshTangents(CachedBakeSettings.BakeMapTypes);
 	}
 	else
 	{
@@ -656,7 +688,7 @@ EBakeOpState UBakeMeshAttributeMapsTool::UpdateResult_DetailNormalMap()
 	if (!(CachedDetailMeshSettings == DetailMeshSettings))
 	{
 		CachedDetailMeshSettings = DetailMeshSettings;
-		ResultState = EBakeOpState::Evaluate;
+		ResultState |= EBakeOpState::Evaluate;
 	}
 	return ResultState;
 }
diff --git a/Engine/Plugins/Experimental/MeshModelingToolsetExp/Source/MeshModelingToolsExp/Private/BakeMeshAttributeTool.cpp b/Engine/Plugins/Experimental/MeshModelingToolsetExp/Source/MeshModelingToolsExp/Private/BakeMeshAttributeTool.cpp
--- a/Engine/Plugins/Experimental/MeshModelingToolsetExp/Source/MeshModelingToolsExp/Private/BakeMeshAttributeTool.cpp
+++ b/Engine/Plugins/Experimental/MeshModelingToolsetExp/Source/MeshModelingToolsExp/Private/BakeMeshAttributeTool.cpp
@@ -23,6 +23,30 @@ void UBakeMeshAttributeTool::Setup()
 }
 
 
+bool UBakeMeshAttributeTool::ValidTargetMeshTangents()
+{
+	if (bCheckTargetMeshTangents)
+	{
+		bValidTargetMeshTangents = TargetMeshTangents ? FDynamicMeshTangents(&TargetMesh).HasValidTangents(true) : false;
+		bCheckTargetMeshTangents = false;
+	}
+	return bValidTargetMeshTangents;
+}
+
+
+EBakeOpState UBakeMeshAttributeTool::UpdateResult_TargetMeshTangents(EBakeMapType BakeType)
+{
+	EBakeOpState ResultState = EBakeOpState::Clean;
+	const bool bNeedTargetMeshTangents = (bool)(BakeType & (EBakeMapType::TangentSpaceNormal | EBakeMapType::BentNormal));
+	if (bNeedTargetMeshTangents && !ValidTargetMeshTangents())
+	{
+		GetToolManager()->DisplayMessage(LOCTEXT("InvalidTargetTangentsWarning", "The Target Mesh does not have valid tangents."), EToolMessageLevel::UserWarning);
+		ResultState = EBakeOpState::Invalid;
+	}
+	return ResultState;
+}
+
+
 EBakeOpState UBakeMeshAttributeTool::UpdateResult_Normal(const FImageDimensions& Dimensions)
 {
 	EBakeOpState ResultState = EBakeOpState::Clean;
diff --git a/Engine/Plugins/Experimental/MeshModelingToolsetExp/Source/MeshModelingToolsExp/Private/BakeMeshAttributeVertexTool.cpp b/Engine/Plugins/Experimental/MeshModelingToolsetExp/Source/MeshModelingToolsExp/Private/BakeMeshAttributeVertexTool.cpp
--- a/Engine/Plugins/Experimental/MeshModelingToolsetExp/Source/MeshModelingToolsExp/Private/BakeMeshAttributeVertexTool.cpp
+++ b/Engine/Plugins/Experimental/MeshModelingToolsetExp/Source/MeshModelingToolsExp/Private/BakeMeshAttributeVertexTool.cpp
@@ -272,6 +272,20 @@ void UBakeMeshAttributeVertexTool::Setup()
 
 	// Setup tool property sets
 
+	Settings = NewObject<UBakeMeshAttributeVertexToolProperties>(this);
+	Settings->RestoreProperties(this);
+	AddToolPropertySource(Settings);
+
+	Settings->WatchProperty(Settings->OutputMode, [this](EBakeVertexOutput) { OpState |= EBakeOpState::Evaluate; UpdateOnModeChange(); });
+	Settings->WatchProperty(Settings->OutputType, [this](int32) { OpState |= EBakeOpState::Evaluate; UpdateOnModeChange(); });
+	Settings->WatchProperty(Settings->OutputTypeR, [this](int32) { OpState |= EBakeOpState::Evaluate; UpdateOnModeChange(); });
+	Settings->WatchProperty(Settings->OutputTypeG, [this](int32) { OpState |= EBakeOpState::Evaluate; UpdateOnModeChange(); });
+	Settings->WatchProperty(Settings->OutputTypeB, [this](int32) { OpState |= EBakeOpState::Evaluate; UpdateOnModeChange(); });
+	Settings->WatchProperty(Settings->OutputTypeA, [this](int32) { OpState |= EBakeOpState::Evaluate; UpdateOnModeChange(); });
+	Settings->WatchProperty(Settings->PreviewMode, [this](EBakeVertexChannel) { UpdateVisualization(); });
+	Settings->WatchProperty(Settings->bSplitAtNormalSeams, [this](bool) { bColorTopologyValid = false; OpState |= EBakeOpState::Evaluate; });
+	Settings->WatchProperty(Settings->bSplitAtUVSeams, [this](bool) { bColorTopologyValid = false; OpState |= EBakeOpState::Evaluate; });
+
 	InputMeshSettings = NewObject<UBakeInputMeshProperties>(this);
 	InputMeshSettings->RestoreProperties(this);
 	AddToolPropertySource(InputMeshSettings);
@@ -287,20 +301,6 @@ void UBakeMeshAttributeVertexTool::Setup()
 	InputMeshSettings->SourceNormalMap = nullptr;
 	InputMeshSettings->WatchProperty(InputMeshSettings->ProjectionDistance, [this](float) { OpState |= EBakeOpState::Evaluate; });
 	InputMeshSettings->WatchProperty(InputMeshSettings->bProjectionInWorldSpace, [this](bool) { OpState |= EBakeOpState::EvaluateDetailMesh; });
-	
-	Settings = NewObject<UBakeMeshAttributeVertexToolProperties>(this);
-	Settings->RestoreProperties(this);
-	AddToolPropertySource(Settings);
-
-	Settings->WatchProperty(Settings->OutputMode, [this](EBakeVertexOutput) { OpState |= EBakeOpState::Evaluate; UpdateOnModeChange(); });
-	Settings->WatchProperty(Settings->OutputType, [this](int32) { OpState |= EBakeOpState::Evaluate; UpdateOnModeChange(); });
-	Settings->WatchProperty(Settings->OutputTypeR, [this](int32) { OpState |= EBakeOpState::Evaluate; UpdateOnModeChange(); });
-	Settings->WatchProperty(Settings->OutputTypeG, [this](int32) { OpState |= EBakeOpState::Evaluate; UpdateOnModeChange(); });
-	Settings->WatchProperty(Settings->OutputTypeB, [this](int32) { OpState |= EBakeOpState::Evaluate; UpdateOnModeChange(); });
-	Settings->WatchProperty(Settings->OutputTypeA, [this](int32) { OpState |= EBakeOpState::Evaluate; UpdateOnModeChange(); });
-	Settings->WatchProperty(Settings->PreviewMode, [this](EBakeVertexChannel) { UpdateVisualization(); });
-	Settings->WatchProperty(Settings->bSplitAtNormalSeams, [this](bool) { bColorTopologyValid = false; OpState |= EBakeOpState::Evaluate; });
-	Settings->WatchProperty(Settings->bSplitAtUVSeams, [this](bool) { bColorTopologyValid = false; OpState |= EBakeOpState::Evaluate; });
 
 	OcclusionSettings = NewObject<UBakeOcclusionMapToolProperties>(this);
 	OcclusionSettings->RestoreProperties(this);
@@ -649,12 +649,14 @@ void UBakeMeshAttributeVertexTool::UpdateResult()
 	const FImageDimensions Dimensions(NumColorElements, 1);
 	if (CachedBakeSettings.OutputMode == EBakeVertexOutput::RGBA)
 	{
-		switch(BakeSettings.OutputType)
+		switch(CachedBakeSettings.OutputType)
 		{
 		case EBakeMapType::TangentSpaceNormal:
 			OpState |= UpdateResult_Normal(Dimensions);
 			break;
 		case EBakeMapType::AmbientOcclusion:
+			OpState |= UpdateResult_Occlusion(Dimensions);
+			break;
 		case EBakeMapType::BentNormal:
 			OpState |= UpdateResult_Occlusion(Dimensions);
 			break;
@@ -676,6 +678,8 @@ void UBakeMeshAttributeVertexTool::UpdateResult()
 		default:
 			break;
 		}
+
+		OpState |= UpdateResult_TargetMeshTangents(CachedBakeSettings.OutputType);
 	}
 	else // CachedBakeSettings.VertexMode == EBakeVertexOutput::PerChannel
 	{
diff --git a/Engine/Plugins/Experimental/MeshModelingToolsetExp/Source/MeshModelingToolsExp/Private/BakeMultiMeshAttributeMapsTool.cpp b/Engine/Plugins/Experimental/MeshModelingToolsetExp/Source/MeshModelingToolsExp/Private/BakeMultiMeshAttributeMapsTool.cpp
--- a/Engine/Plugins/Experimental/MeshModelingToolsetExp/Source/MeshModelingToolsExp/Private/BakeMultiMeshAttributeMapsTool.cpp
+++ b/Engine/Plugins/Experimental/MeshModelingToolsetExp/Source/MeshModelingToolsExp/Private/BakeMultiMeshAttributeMapsTool.cpp
@@ -587,14 +587,11 @@ void UBakeMultiMeshAttributeMapsTool::UpdateResult()
 	// Clear our invalid bitflag to check again for valid inputs.
 	OpState &= ~EBakeOpState::Invalid;
 
+	OpState |= UpdateResult_TargetMeshTangents(CachedBakeSettings.BakeMapTypes);
+
 	// Update map type settings
 	OpState |= UpdateResult_DetailMeshes();
 
-	if ((bool)(CachedBakeSettings.BakeMapTypes & EBakeMapType::TangentSpaceNormal))
-	{
-		OpState |= UpdateResult_Normal(CachedBakeSettings.Dimensions);
-	}
-
 	// Early exit if op input parameters are invalid.
 	if (static_cast<bool>(OpState & EBakeOpState::Invalid))
 	{
diff --git a/Engine/Plugins/Experimental/MeshModelingToolsetExp/Source/MeshModelingToolsExp/Public/BakeMeshAttributeMapsTool.h b/Engine/Plugins/Experimental/MeshModelingToolsetExp/Source/MeshModelingToolsExp/Public/BakeMeshAttributeMapsTool.h
--- a/Engine/Plugins/Experimental/MeshModelingToolsetExp/Source/MeshModelingToolsExp/Public/BakeMeshAttributeMapsTool.h
+++ b/Engine/Plugins/Experimental/MeshModelingToolsetExp/Source/MeshModelingToolsExp/Public/BakeMeshAttributeMapsTool.h
@@ -130,6 +130,10 @@ protected:
 
 	bool bIsBakeToSelf = false;
 
+	bool ValidDetailMeshTangents();
+	bool bCheckDetailMeshTangents = true;
+	bool bValidDetailMeshTangents = false;
+	
 	TSharedPtr<UE::Geometry::FDynamicMesh3, ESPMode::ThreadSafe> DetailMesh;
 	TSharedPtr<UE::Geometry::FDynamicMeshAABBTree3, ESPMode::ThreadSafe> DetailSpatial;
 	TSharedPtr<UE::Geometry::TMeshTangents<double>, ESPMode::ThreadSafe> DetailMeshTangents;
@@ -140,6 +144,8 @@ protected:
 
 	void InvalidateResults();
 
+	EBakeOpState UpdateResult_DetailMeshTangents(EBakeMapType BakeType);
+
 	FDetailMeshSettings CachedDetailMeshSettings;
 	TSharedPtr<UE::Geometry::TImageBuilder<FVector4f>, ESPMode::ThreadSafe> CachedDetailNormalMap;
 	EBakeOpState UpdateResult_DetailNormalMap();
diff --git a/Engine/Plugins/Experimental/MeshModelingToolsetExp/Source/MeshModelingToolsExp/Public/BakeMeshAttributeTool.h b/Engine/Plugins/Experimental/MeshModelingToolsetExp/Source/MeshModelingToolsExp/Public/BakeMeshAttributeTool.h
--- a/Engine/Plugins/Experimental/MeshModelingToolsetExp/Source/MeshModelingToolsExp/Public/BakeMeshAttributeTool.h
+++ b/Engine/Plugins/Experimental/MeshModelingToolsetExp/Source/MeshModelingToolsExp/Public/BakeMeshAttributeTool.h
@@ -129,10 +129,23 @@ protected:
 	TSharedPtr<UE::Geometry::TMeshTangents<double>, ESPMode::ThreadSafe> TargetMeshTangents;
 	UE::Geometry::FDynamicMeshAABBTree3 TargetSpatial;
 
+	/**
+	 * Compute validity of the Target Mesh tangents. Only checks validity
+	 * once and then caches the result for successive calls.
+	 * 
+	 * @return true if the TargetMesh tangents are valid
+	 */
+	bool ValidTargetMeshTangents();
+	bool bCheckTargetMeshTangents = true;
+	bool bValidTargetMeshTangents = false;
+
 	//
 	// Bake results update management
 	//
 	const bool bPreferPlatformData = false;
+
+	/** Update the OpState based on the validity of the target mesh tangents */
+	EBakeOpState UpdateResult_TargetMeshTangents(EBakeMapType BakeType);
 	
 	EBakeOpState UpdateResult_Normal(const FImageDimensions& Dimensions);
 	FNormalMapSettings CachedNormalMapSettings;
diff --git a/Engine/Source/Runtime/GeometryCore/Private/DynamicMesh/MeshTangents.cpp b/Engine/Source/Runtime/GeometryCore/Private/DynamicMesh/MeshTangents.cpp
--- a/Engine/Source/Runtime/GeometryCore/Private/DynamicMesh/MeshTangents.cpp
+++ b/Engine/Source/Runtime/GeometryCore/Private/DynamicMesh/MeshTangents.cpp
@@ -23,6 +23,42 @@ FDynamicMeshTangents::FDynamicMeshTangents(const FDynamicMesh3* MeshIn)
 }
 
 
+bool FDynamicMeshTangents::HasValidTangents(const bool bCheckValues /*=false*/) const
+{
+	if (!Mesh || !Tangents || !Bitangents)
+	{
+		return false;
+	}
+
+	if (!bCheckValues)
+	{
+		// If not checking values, then we are done. Overlays are valid. 
+		return true;
+	}
+
+	for (const int TriId : Mesh->TriangleIndicesItr())
+	{
+		if (!Tangents->IsSetTriangle(TriId))
+		{
+			return false;
+		}
+		
+		for (int TriVtxId = 0; TriVtxId < 2; ++TriVtxId)
+		{
+			FVector3f T, B;
+			Tangents->GetTriElement(TriId, TriVtxId, T);
+			Bitangents->GetTriElement(TriId, TriVtxId, B);
+			if (T.IsNearlyZero() || T.ContainsNaN() || B.IsNearlyZero() || B.ContainsNaN())
+			{
+				return false;
+			}
+		}
+	}
+	return true;
+}
+
+
+
 void FDynamicMeshTangents::GetTangentFrame(int32 TriangleID, int32 TriVertexIndex, FVector3f& Normal, FVector3f& Tangent, FVector3f& Bitangent) const
 {
 	if (Normals && Normals->IsSetTriangle(TriangleID))
diff --git a/Engine/Source/Runtime/GeometryCore/Public/DynamicMesh/MeshTangents.h b/Engine/Source/Runtime/GeometryCore/Public/DynamicMesh/MeshTangents.h
--- a/Engine/Source/Runtime/GeometryCore/Public/DynamicMesh/MeshTangents.h
+++ b/Engine/Source/Runtime/GeometryCore/Public/DynamicMesh/MeshTangents.h
@@ -26,8 +26,14 @@ public:
 
 	FDynamicMeshTangents(const FDynamicMesh3* MeshIn);
 
-	/** @return true if the mesh has valid tangents */
-	bool HasValidTangents() const { return Tangents != nullptr && Bitangents != nullptr; }
+	/**
+	 * Checks the mesh for valid tangents. When bCheckValues == true,
+	 * inspects the tangents for invalid values (ex. zero, NaN).
+	 * 
+	 * @param bCheckValues inspect tangent values for zero/NaN
+	 * @return true if the mesh has valid tangents
+	 */
+	bool HasValidTangents(bool bCheckValues=false) const;
 
 	/**
 	 * If tangents are available in the overlays, returns them. If only Normal is available, computes orthogonal basis. Falls back to unit axes if no overlays are available.
