ApproximateActors: add support for PatchBuilder AutoUV method
#rnx

diff --git a/Engine/Plugins/Experimental/MeshModelingToolsetExp/Source/GeometryProcessingAdapters/Private/ApproximateActorsImpl.cpp b/Engine/Plugins/Experimental/MeshModelingToolsetExp/Source/GeometryProcessingAdapters/Private/ApproximateActorsImpl.cpp
--- a/Engine/Plugins/Experimental/MeshModelingToolsetExp/Source/GeometryProcessingAdapters/Private/ApproximateActorsImpl.cpp
+++ b/Engine/Plugins/Experimental/MeshModelingToolsetExp/Source/GeometryProcessingAdapters/Private/ApproximateActorsImpl.cpp
@@ -438,7 +438,7 @@ static TSharedPtr<FApproximationMeshData> GenerateApproximationMesh(
 
 	if (Options.bVerbose)
 	{
-		UE_LOG(LogApproximateActors, Warning, TEXT("Solidify mesh has %d triangles"), CurResultMesh->TriangleCount());
+		UE_LOG(LogApproximateActors, Log, TEXT("Solidify mesh has %d triangles"), CurResultMesh->TriangleCount());
 	}
 
 	Progress.EnterProgressFrame(1.f, LOCTEXT("ClosingMesh", "Topological Operations..."));
@@ -462,7 +462,7 @@ static TSharedPtr<FApproximationMeshData> GenerateApproximationMesh(
 
 		if (Options.bVerbose)
 		{
-			UE_LOG(LogApproximateActors, Warning, TEXT("Morphology mesh has %d triangles"), CurResultMesh->TriangleCount());
+			UE_LOG(LogApproximateActors, Log, TEXT("Morphology mesh has %d triangles"), CurResultMesh->TriangleCount());
 		}
 	}
 
@@ -531,7 +531,7 @@ static TSharedPtr<FApproximationMeshData> GenerateApproximationMesh(
 		int32 AfterCount = CurResultMesh->TriangleCount();
 		if (Options.bVerbose)
 		{
-			UE_LOG(LogApproximateActors, Warning, TEXT("Simplified mesh from %d to %d triangles"), BeforeCount, AfterCount);
+			UE_LOG(LogApproximateActors, Log, TEXT("Simplified mesh from %d to %d triangles"), BeforeCount, AfterCount);
 		}
 	}
 
@@ -621,7 +621,7 @@ static TSharedPtr<FApproximationMeshData> GenerateApproximationMesh(
 
 		if (Options.bVerbose)
 		{
-			UE_LOG(LogApproximateActors, Warning, TEXT("Occlusion-Filtered mesh has %d triangles (removed %d)"), CurResultMesh->TriangleCount(), NumRemoved);
+			UE_LOG(LogApproximateActors, Log, TEXT("Occlusion-Filtered mesh has %d triangles (removed %d)"), CurResultMesh->TriangleCount(), NumRemoved);
 		}
 	}
 
@@ -678,16 +678,36 @@ static TSharedPtr<FApproximationMeshData> GenerateApproximationMesh(
 	// compute UVs
 	bool bHaveValidUVs = true;
 	TSharedPtr<FDynamicMesh3, ESPMode::ThreadSafe> UVInputMesh = MakeShared<FDynamicMesh3, ESPMode::ThreadSafe>();
-	*UVInputMesh = MoveTemp(*CurResultMesh);
+	// PatchBuilder AutoUV currently requires compact input mesh. TODO: fix that, then we can just re-use CurResultMesh here.
+	UVInputMesh->CompactCopy(*CurResultMesh);
+	//*UVInputMesh = MoveTemp(*CurResultMesh);
 	FParameterizeMeshOp ParameterizeMeshOp;
 	ParameterizeMeshOp.Stretch = Options.UVAtlasStretchTarget;
+	
+	// UVAtlas parameters
 	ParameterizeMeshOp.NumCharts = 0;
+	
+	// PatchBuilder generation parameters
+	ParameterizeMeshOp.InitialPatchCount = Options.PatchBuilderInitialPatchCount;
+	ParameterizeMeshOp.bRespectInputGroups = false;
+	ParameterizeMeshOp.PatchCurvatureAlignmentWeight = Options.PatchBuilderCurvatureAlignment;
+	ParameterizeMeshOp.PatchMergingMetricThresh = Options.PatchBuilderMergingThreshold;
+	ParameterizeMeshOp.PatchMergingAngleThresh = Options.PatchBuilderMaxNormalDeviationDeg;
+	ParameterizeMeshOp.ExpMapNormalSmoothingSteps = 5;
+	ParameterizeMeshOp.ExpMapNormalSmoothingAlpha = 0.25;
+
 	ParameterizeMeshOp.InputMesh = UVInputMesh;
 	ParameterizeMeshOp.Method = UE::Geometry::EParamOpBackend::XAtlas;
 	if (Options.UVPolicy == IGeometryProcessing_ApproximateActors::EUVGenerationPolicy::PreferUVAtlas)
 	{
 		ParameterizeMeshOp.Method = UE::Geometry::EParamOpBackend::UVAtlas;
 	}
+	else if (Options.UVPolicy == IGeometryProcessing_ApproximateActors::EUVGenerationPolicy::PreferPatchBuilder)
+	{
+		ParameterizeMeshOp.Method = UE::Geometry::EParamOpBackend::PatchBuilder;
+	}
+
+
 	FProgressCancel UVProgressCancel;
 	{
 		TRACE_CPUPROFILER_EVENT_SCOPE(ApproximateActorsImpl_Generate_GenerateUVs);
@@ -885,8 +905,19 @@ IGeometryProcessing_ApproximateActors::FOptions FApproximateActorsImpl::Construc
 		Options.MeshSimplificationPolicy = IGeometryProcessing_ApproximateActors::ESimplificationPolicy::FixedTriangleCount;
 	}
 
-	Options.UVPolicy = (UseSettings.UVGenerationMethod == EMeshApproximationUVGenerationPolicy::PreferUVAtlas) ?
-		IGeometryProcessing_ApproximateActors::EUVGenerationPolicy::PreferUVAtlas : IGeometryProcessing_ApproximateActors::EUVGenerationPolicy::PreferXAtlas;
+	Options.UVPolicy = IGeometryProcessing_ApproximateActors::EUVGenerationPolicy::PreferXAtlas;
+	if (UseSettings.UVGenerationMethod == EMeshApproximationUVGenerationPolicy::PreferUVAtlas)
+	{
+		Options.UVPolicy = IGeometryProcessing_ApproximateActors::EUVGenerationPolicy::PreferUVAtlas;
+	}
+	else if (UseSettings.UVGenerationMethod == EMeshApproximationUVGenerationPolicy::PreferPatchBuilder)
+	{
+		Options.UVPolicy = IGeometryProcessing_ApproximateActors::EUVGenerationPolicy::PreferPatchBuilder;
+	}
+	Options.PatchBuilderInitialPatchCount =  FMath::Clamp(UseSettings.InitialPatchCount, 1, 99999);
+	Options.PatchBuilderCurvatureAlignment = FMath::Clamp(UseSettings.CurvatureAlignment, 0.001, 1000.0);
+	Options.PatchBuilderMergingThreshold = FMath::Clamp(UseSettings.MergingThreshold, 1.0, 9999.0);
+	Options.PatchBuilderMaxNormalDeviationDeg = FMath::Clamp(UseSettings.MaxAngleDeviation, 0.0, 180);
 
 	Options.bCalculateHardNormals = UseSettings.bEstimateHardNormals;
 	Options.HardNormalsAngleDeg = FMath::Clamp(UseSettings.HardNormalAngle, 0.001f, 89.99f);
@@ -977,7 +1008,7 @@ void FApproximateActorsImpl::GenerateApproximationForActorSet(const TArray<AActo
 	{
 		FMeshSceneAdapter::FStatistics Stats;
 		Scene.GetGeometryStatistics(Stats);
-		UE_LOG(LogApproximateActors, Warning, TEXT("%lld triangles in %lld unique meshes, total %lld triangles in %lld instances"),
+		UE_LOG(LogApproximateActors, Log, TEXT("%lld triangles in %lld unique meshes, total %lld triangles in %lld instances"),
 			Stats.UniqueMeshTriangleCount, Stats.UniqueMeshCount, Stats.InstanceMeshTriangleCount, Stats.InstanceMeshCount);
 	}
 
diff --git a/Engine/Source/Developer/GeometryProcessingInterfaces/Public/GeometryProcessingInterfaces/ApproximateActors.h b/Engine/Source/Developer/GeometryProcessingInterfaces/Public/GeometryProcessingInterfaces/ApproximateActors.h
--- a/Engine/Source/Developer/GeometryProcessingInterfaces/Public/GeometryProcessingInterfaces/ApproximateActors.h
+++ b/Engine/Source/Developer/GeometryProcessingInterfaces/Public/GeometryProcessingInterfaces/ApproximateActors.h
@@ -75,7 +75,8 @@ public:
 	enum class EUVGenerationPolicy : uint8
 	{
 		PreferUVAtlas = 0,
-		PreferXAtlas = 1
+		PreferXAtlas = 1,
+		PreferPatchBuilder = 2
 	};
 
 	enum class ETextureSizePolicy : uint8
@@ -139,6 +140,10 @@ public:
 
 		EUVGenerationPolicy UVPolicy = EUVGenerationPolicy::PreferXAtlas;
 		double UVAtlasStretchTarget = 0.1;
+		int32 PatchBuilderInitialPatchCount = 250;
+		double PatchBuilderCurvatureAlignment = 1.0;
+		double PatchBuilderMergingThreshold = 1.5;
+		double PatchBuilderMaxNormalDeviationDeg = 45.0;
 
 		//
 		// Material approximation settings
diff --git a/Engine/Source/Runtime/Engine/Classes/Engine/MeshMerging.h b/Engine/Source/Runtime/Engine/Classes/Engine/MeshMerging.h
--- a/Engine/Source/Runtime/Engine/Classes/Engine/MeshMerging.h
+++ b/Engine/Source/Runtime/Engine/Classes/Engine/MeshMerging.h
@@ -720,7 +720,8 @@ UENUM()
 enum class EMeshApproximationUVGenerationPolicy : uint8
 {
 	PreferUVAtlas = 0,
-	PreferXAtlas = 1
+	PreferXAtlas = 1,
+	PreferPatchBuilder = 2
 };
 
 
@@ -838,6 +839,22 @@ struct FMeshApproximationSettings
 	EMeshApproximationUVGenerationPolicy UVGenerationMethod = EMeshApproximationUVGenerationPolicy::PreferXAtlas;
 
 
+	/** Number of initial patches mesh will be split into before computing island merging */
+	UPROPERTY(BlueprintReadWrite, EditAnywhere, Category = UVSettings, AdvancedDisplay, meta = (UIMin = "1", UIMax = "1000", ClampMin = "1", ClampMax = "99999999", EditCondition = "UVGenerationMethod == EMeshApproximationUVGenerationPolicy::PreferPatchBuilder"))
+	int InitialPatchCount = 250;
+
+	/** This parameter controls alignment of the initial patches to creases in the mesh */
+	UPROPERTY(BlueprintReadWrite, EditAnywhere, Category = UVSettings, AdvancedDisplay, meta = (UIMin = "0.1", UIMax = "2.0", ClampMin = "0.01", ClampMax = "100.0", EditCondition = "UVGenerationMethod == EMeshApproximationUVGenerationPolicy::PreferPatchBuilder"))
+	float CurvatureAlignment = 1.0f;
+
+	/** Distortion/Stretching Threshold for island merging - larger values increase the allowable UV stretching */
+	UPROPERTY(BlueprintReadWrite, EditAnywhere, Category = UVSettings, AdvancedDisplay, meta = (UIMin = "1.0", UIMax = "5.0", ClampMin = "1.0", EditCondition = "UVGenerationMethod == EMeshApproximationUVGenerationPolicy::PreferPatchBuilder"))
+	float MergingThreshold = 1.5f;
+
+	/** UV islands will not be merged if their average face normals deviate by larger than this amount */
+	UPROPERTY(BlueprintReadWrite, EditAnywhere, Category = UVSettings, AdvancedDisplay, meta = (UIMin = "0.0", UIMax = "90.0", ClampMin = "0.0", ClampMax = "180.0", EditCondition = "UVGenerationMethod == EMeshApproximationUVGenerationPolicy::PreferPatchBuilder"))
+	float MaxAngleDeviation = 45.0f;
+
 	//
 	// Output Static Mesh Settings
 	//
