Renamed ESubdivisionType to EDisplaceMeshToolSubdivisionType to follow the naming convention of the enums.
Renamed Loop tessellation option to Flat to match UE4 tessellation option and  not to be confused with Loop subdivision.
Changed the Channel option to be a drop down menu with R,G,B,A choices instead of being a numeric textfield (0-3).

#rnx

diff --git a/Engine/Plugins/Experimental/MeshModelingToolsetExp/Source/MeshModelingToolsExp/Private/DisplaceMeshTool.cpp b/Engine/Plugins/Experimental/MeshModelingToolsetExp/Source/MeshModelingToolsExp/Private/DisplaceMeshTool.cpp
--- a/Engine/Plugins/Experimental/MeshModelingToolsetExp/Source/MeshModelingToolsExp/Private/DisplaceMeshTool.cpp
+++ b/Engine/Plugins/Experimental/MeshModelingToolsetExp/Source/MeshModelingToolsExp/Private/DisplaceMeshTool.cpp
@@ -214,14 +214,14 @@ namespace DisplaceMeshToolLocals{
 	class FSubdivideMeshOp : public FDynamicMeshOperator
 	{
 	public:
-		FSubdivideMeshOp(const FDynamicMesh3& SourceMesh, ESubdivisionType SubdivisionTypeIn, int SubdivisionsCountIn, TSharedPtr<FIndexedWeightMap, ESPMode::ThreadSafe> WeightMap);
+		FSubdivideMeshOp(const FDynamicMesh3& SourceMesh, EDisplaceMeshToolSubdivisionType SubdivisionTypeIn, int SubdivisionsCountIn, TSharedPtr<FIndexedWeightMap, ESPMode::ThreadSafe> WeightMap);
 		void CalculateResult(FProgressCancel* Progress) final;
 	private:
-		ESubdivisionType SubdivisionType;
+		EDisplaceMeshToolSubdivisionType SubdivisionType;
 		int SubdivisionsCount;
 	};
 
-	FSubdivideMeshOp::FSubdivideMeshOp(const FDynamicMesh3& SourceMesh, ESubdivisionType SubdivisionTypeIn, int SubdivisionsCountIn, TSharedPtr<FIndexedWeightMap, ESPMode::ThreadSafe> WeightMap)
+	FSubdivideMeshOp::FSubdivideMeshOp(const FDynamicMesh3& SourceMesh, EDisplaceMeshToolSubdivisionType SubdivisionTypeIn, int SubdivisionsCountIn, TSharedPtr<FIndexedWeightMap, ESPMode::ThreadSafe> WeightMap)
 		: SubdivisionType(SubdivisionTypeIn), SubdivisionsCount(SubdivisionsCountIn)
 	{
 		ResultMesh->Copy(SourceMesh);
@@ -248,7 +248,7 @@ namespace DisplaceMeshToolLocals{
 
 	void FSubdivideMeshOp::CalculateResult(FProgressCancel* ProgressCancel)
 	{
-		if (SubdivisionType == ESubdivisionType::Loop) 
+		if (SubdivisionType == EDisplaceMeshToolSubdivisionType::Flat) 
 		{
 			FUniformTesselate Tesselator(ResultMesh.Get());
 			Tesselator.Progress = ProgressCancel;
@@ -259,7 +259,7 @@ namespace DisplaceMeshToolLocals{
 				Tesselator.Compute();
 			}
 		}
-		else if (SubdivisionType == ESubdivisionType::PNTriangles) 
+		else if (SubdivisionType == EDisplaceMeshToolSubdivisionType::PNTriangles) 
 		{
 			FPNTriangles PNTriangles(ResultMesh.Get());
 			PNTriangles.Progress = ProgressCancel;
@@ -281,15 +281,15 @@ namespace DisplaceMeshToolLocals{
 	{
 	public:
 		FSubdivideMeshOpFactory(FDynamicMesh3& SourceMeshIn,
-			ESubdivisionType SubdivisionTypeIn,
+			EDisplaceMeshToolSubdivisionType SubdivisionTypeIn,
 			int SubdivisionsCountIn,
 			TSharedPtr<FIndexedWeightMap, ESPMode::ThreadSafe> WeightMapIn)
 			: SourceMesh(SourceMeshIn), SubdivisionType(SubdivisionTypeIn), SubdivisionsCount(SubdivisionsCountIn), WeightMap(WeightMapIn)
 		{
 		}
 
-		void SetSubdivisionType(ESubdivisionType SubdivisionTypeIn);
-		ESubdivisionType GetSubdivisionType() const;
+		void SetSubdivisionType(EDisplaceMeshToolSubdivisionType SubdivisionTypeIn);
+		EDisplaceMeshToolSubdivisionType GetSubdivisionType() const;
 
 		void SetSubdivisionsCount(int SubdivisionsCountIn);
 		int  GetSubdivisionsCount() const;
@@ -302,17 +302,17 @@ namespace DisplaceMeshToolLocals{
 		}
 	private:
 		const FDynamicMesh3& SourceMesh;
-		ESubdivisionType SubdivisionType;
+		EDisplaceMeshToolSubdivisionType SubdivisionType;
 		int SubdivisionsCount;
 		TSharedPtr<FIndexedWeightMap, ESPMode::ThreadSafe> WeightMap;
 	};
 
-	void FSubdivideMeshOpFactory::SetSubdivisionType(ESubdivisionType SubdivisionTypeIn) 
+	void FSubdivideMeshOpFactory::SetSubdivisionType(EDisplaceMeshToolSubdivisionType SubdivisionTypeIn) 
 	{
 		SubdivisionType = SubdivisionTypeIn;
 	}
 
-	ESubdivisionType FSubdivideMeshOpFactory::GetSubdivisionType() const
+	EDisplaceMeshToolSubdivisionType FSubdivideMeshOpFactory::GetSubdivisionType() const
 	{
 		return SubdivisionType;
 	}
@@ -1026,11 +1026,11 @@ void UDisplaceMeshTool::OnPropertyModified(UObject* PropertySet, FProperty* Prop
 				GetToolManager()->DisplayMessage(FText::GetEmpty(), EToolMessageLevel::UserWarning);
 			}
 
-			DisplacerDownCast->SetDisplacementMap(TextureMapProperties->DisplacementMap, TextureMapProperties->Channel);
+			DisplacerDownCast->SetDisplacementMap(TextureMapProperties->DisplacementMap, static_cast<int32>(TextureMapProperties->Channel));
 		}
 		else if (PropName == GET_MEMBER_NAME_CHECKED(UDisplaceMeshTextureMapProperties, Channel))
 		{
-			DisplacerDownCast->SetDisplacementMap(TextureMapProperties->DisplacementMap, TextureMapProperties->Channel);
+			DisplacerDownCast->SetDisplacementMap(TextureMapProperties->DisplacementMap, static_cast<int32>(TextureMapProperties->Channel));
 		}
 		else if (PropName == GET_MEMBER_NAME_CHECKED(UDisplaceMeshTextureMapProperties, DisplacementMapBaseValue))
 		{
diff --git a/Engine/Plugins/Experimental/MeshModelingToolsetExp/Source/MeshModelingToolsExp/Public/DisplaceMeshTool.h b/Engine/Plugins/Experimental/MeshModelingToolsetExp/Source/MeshModelingToolsExp/Public/DisplaceMeshTool.h
--- a/Engine/Plugins/Experimental/MeshModelingToolsetExp/Source/MeshModelingToolsExp/Public/DisplaceMeshTool.h
+++ b/Engine/Plugins/Experimental/MeshModelingToolsetExp/Source/MeshModelingToolsExp/Public/DisplaceMeshTool.h
@@ -44,10 +44,10 @@ enum class EDisplaceMeshToolDisplaceType : uint8
 };
 
 UENUM() 
-enum class ESubdivisionType : uint8  
+enum class EDisplaceMeshToolSubdivisionType : uint8  
 {
 	/** Subdivide the mesh using loop-style subdivision. */
-	Loop UMETA(DisplayName = "Loop"),
+	Flat UMETA(DisplayName = "Flat"),
 	
 	/** Subdivide the mesh using PN triangles which replace each original flat triangle by a curved shape that is 
 		retriangulated into a number of small subtriangles. The geometry of a PN triangle is defined as one cubic Bezier 
@@ -56,6 +56,15 @@ enum class ESubdivisionType : uint8
 	PNTriangles UMETA(DisplayName = "PN Triangles"),
 };
 
+UENUM() 
+enum class EDisplaceMeshToolChannelType : uint8  
+{
+	Red,
+	Green,
+	Blue,
+	Alpha
+};
+
 /** The basic set of properties shared by (more or less) all DisplacementTypes. */
 UCLASS()
 class MESHMODELINGTOOLSEXP_API UDisplaceMeshCommonProperties : public UInteractiveToolPropertySet
@@ -79,7 +88,7 @@ public:
 
 	/** Type of the  mesh subdivision. */
 	UPROPERTY(EditAnywhere, Category = Options)
-	ESubdivisionType SubdivisionType = ESubdivisionType::Loop;
+	EDisplaceMeshToolSubdivisionType SubdivisionType = EDisplaceMeshToolSubdivisionType::PNTriangles;
 
 	/** Number of times to subdivide the mesh before displacing it. */
 	UPROPERTY(EditAnywhere, Category = Options,
@@ -119,10 +128,9 @@ public:
 	UPROPERTY(EditAnywhere, Category = Options)
 	TObjectPtr<UTexture2D> DisplacementMap = nullptr;
 
-	//~ TODO: Consider making this an enum
-	/** Channel in the displacement map to use. 0-3 coresponding to the four channels (RGBA) */
-	UPROPERTY(EditAnywhere, Category = Options, meta = (ClampMin = "0", ClampMax = "3"))
-	int32 Channel = 0;
+	/** Channel in the displacement map to use. */
+	UPROPERTY(EditAnywhere, Category = Options)
+	EDisplaceMeshToolChannelType Channel = EDisplaceMeshToolChannelType::Red;
 
 	/** The value in the texture map that corresponds to no displacement. For instance, if set to 0, then all
 	 displacement will be positive. If set to 0.5, displacement below 0.5 will be negative, and above will be
