ModelingTools: Add HideSourceMesh option to BakeTx and BakeVtx.

#rnx

diff --git a/Engine/Plugins/Experimental/MeshModelingToolsetExp/Source/MeshModelingToolsExp/Private/BakeMeshAttributeMapsTool.cpp b/Engine/Plugins/Experimental/MeshModelingToolsetExp/Source/MeshModelingToolsExp/Private/BakeMeshAttributeMapsTool.cpp
--- a/Engine/Plugins/Experimental/MeshModelingToolsetExp/Source/MeshModelingToolsExp/Private/BakeMeshAttributeMapsTool.cpp
+++ b/Engine/Plugins/Experimental/MeshModelingToolsetExp/Source/MeshModelingToolsExp/Private/BakeMeshAttributeMapsTool.cpp
@@ -26,7 +26,6 @@
 #include "TargetInterfaces/StaticMeshBackedTarget.h"
 #include "TargetInterfaces/SkeletalMeshBackedTarget.h"
 #include "ToolTargetManager.h"
-#include "ModelingToolTargetUtil.h"
 
 // required to pass UStaticMesh asset so we can save at same location
 #include "Engine/StaticMesh.h"
@@ -302,6 +301,7 @@ void UBakeMeshAttributeMapsTool::Setup()
 	InputMeshSettings->SourceDynamicMesh = !bIsBakeToSelf ? GetDynamicMeshTarget(DetailTarget) : nullptr;
 	InputMeshSettings->SourceNormalMap = nullptr;
 	UpdateUVLayerNames(InputMeshSettings->TargetUVLayer, InputMeshSettings->TargetUVLayerNamesList, TargetMesh);
+	InputMeshSettings->WatchProperty(InputMeshSettings->bHideSourceMesh, [this](bool bState) { SetSourceObjectVisible(!bState); });
 	InputMeshSettings->WatchProperty(InputMeshSettings->TargetUVLayer, [this](FString) { OpState |= EBakeOpState::Evaluate; });
 	InputMeshSettings->WatchProperty(InputMeshSettings->ProjectionDistance, [this](float) { OpState |= EBakeOpState::Evaluate; });
 	InputMeshSettings->WatchProperty(InputMeshSettings->bProjectionInWorldSpace, [this](bool) { OpState |= EBakeOpState::EvaluateDetailMesh; });
@@ -315,6 +315,7 @@ void UBakeMeshAttributeMapsTool::Setup()
 		}
 		OpState |= EBakeOpState::Evaluate;
 	});
+	SetSourceObjectVisible(!InputMeshSettings->bHideSourceMesh);
 	
 
 	ResultSettings = NewObject<UBakeMeshAttributeMapsResultToolProperties>(this);
@@ -469,6 +470,8 @@ void UBakeMeshAttributeMapsTool::Shutdown(EToolShutdownType ShutdownType)
 	TextureSettings->SaveProperties(this);
 	MultiTextureSettings->SaveProperties(this);
 
+	SetSourceObjectVisible(true);
+
 	if (Compute)
 	{
 		Compute->Shutdown();
diff --git a/Engine/Plugins/Experimental/MeshModelingToolsetExp/Source/MeshModelingToolsExp/Private/BakeMeshAttributeVertexTool.cpp b/Engine/Plugins/Experimental/MeshModelingToolsetExp/Source/MeshModelingToolsExp/Private/BakeMeshAttributeVertexTool.cpp
--- a/Engine/Plugins/Experimental/MeshModelingToolsetExp/Source/MeshModelingToolsExp/Private/BakeMeshAttributeVertexTool.cpp
+++ b/Engine/Plugins/Experimental/MeshModelingToolsetExp/Source/MeshModelingToolsExp/Private/BakeMeshAttributeVertexTool.cpp
@@ -22,7 +22,6 @@
 #include "TargetInterfaces/StaticMeshBackedTarget.h"
 #include "TargetInterfaces/SkeletalMeshBackedTarget.h"
 #include "ToolTargetManager.h"
-#include "ModelingToolTargetUtil.h"
 #include "AssetUtils/Texture2DUtil.h"
 
 #include "EngineAnalytics.h"
@@ -299,8 +298,10 @@ void UBakeMeshAttributeVertexTool::Setup()
 	InputMeshSettings->SourceSkeletalMesh = !bIsBakeToSelf ? GetSkeletalMeshTarget(DetailTarget) : nullptr;
 	InputMeshSettings->SourceDynamicMesh = !bIsBakeToSelf ? GetDynamicMeshTarget(DetailTarget) : nullptr;
 	InputMeshSettings->SourceNormalMap = nullptr;
+	InputMeshSettings->WatchProperty(InputMeshSettings->bHideSourceMesh, [this](bool bState) { SetSourceObjectVisible(!bState); });
 	InputMeshSettings->WatchProperty(InputMeshSettings->ProjectionDistance, [this](float) { OpState |= EBakeOpState::Evaluate; });
 	InputMeshSettings->WatchProperty(InputMeshSettings->bProjectionInWorldSpace, [this](bool) { OpState |= EBakeOpState::EvaluateDetailMesh; });
+	SetSourceObjectVisible(!InputMeshSettings->bHideSourceMesh);
 
 	OcclusionSettings = NewObject<UBakeOcclusionMapToolProperties>(this);
 	OcclusionSettings->RestoreProperties(this);
@@ -371,6 +372,7 @@ void UBakeMeshAttributeVertexTool::Shutdown(EToolShutdownType ShutdownType)
 	MultiTextureSettings->SaveProperties(this);
 
 	UE::ToolTarget::ShowSourceObject(Targets[0]);
+	SetSourceObjectVisible(true);
 
 	if (Compute)
 	{
diff --git a/Engine/Plugins/Experimental/MeshModelingToolsetExp/Source/MeshModelingToolsExp/Public/BakeMeshAttributeMapsTool.h b/Engine/Plugins/Experimental/MeshModelingToolsetExp/Source/MeshModelingToolsExp/Public/BakeMeshAttributeMapsTool.h
--- a/Engine/Plugins/Experimental/MeshModelingToolsetExp/Source/MeshModelingToolsExp/Public/BakeMeshAttributeMapsTool.h
+++ b/Engine/Plugins/Experimental/MeshModelingToolsetExp/Source/MeshModelingToolsExp/Public/BakeMeshAttributeMapsTool.h
@@ -3,7 +3,6 @@
 #pragma once
 
 #include "CoreMinimal.h"
-#include "InteractiveToolBuilder.h"
 #include "DynamicMesh/DynamicMesh3.h"
 #include "DynamicMesh/DynamicMeshAABBTree3.h"
 #include "Image/ImageDimensions.h"
@@ -11,6 +10,7 @@
 #include "Sampling/MeshMapBaker.h"
 #include "ModelingOperators.h"
 #include "PreviewMesh.h"
+#include "ModelingToolTargetUtil.h"
 #include "BakeMeshAttributeMapsToolBase.h"
 #include "BakeMeshAttributeMapsTool.generated.h"
 
@@ -149,5 +149,13 @@ protected:
 	FDetailMeshSettings CachedDetailMeshSettings;
 	TSharedPtr<UE::Geometry::TImageBuilder<FVector4f>, ESPMode::ThreadSafe> CachedDetailNormalMap;
 	EBakeOpState UpdateResult_DetailNormalMap();
+
+	void SetSourceObjectVisible(bool bState)
+	{
+		if (!bIsBakeToSelf)
+		{
+			UE::ToolTarget::SetSourceObjectVisible(Targets[1], bState);
+		}
+	}
 };
 
diff --git a/Engine/Plugins/Experimental/MeshModelingToolsetExp/Source/MeshModelingToolsExp/Public/BakeMeshAttributeToolCommon.h b/Engine/Plugins/Experimental/MeshModelingToolsetExp/Source/MeshModelingToolsExp/Public/BakeMeshAttributeToolCommon.h
--- a/Engine/Plugins/Experimental/MeshModelingToolsetExp/Source/MeshModelingToolsExp/Public/BakeMeshAttributeToolCommon.h
+++ b/Engine/Plugins/Experimental/MeshModelingToolsetExp/Source/MeshModelingToolsExp/Public/BakeMeshAttributeToolCommon.h
@@ -66,6 +66,11 @@ public:
 		EditCondition = "SourceDynamicMesh != nullptr", EditConditionHides))
 	TObjectPtr<AActor> SourceDynamicMesh = nullptr;
 
+	/** If true, hide the source mesh */
+	UPROPERTY(EditAnywhere, Category = BakeInput, meta = (EditConditionHides, HideEditConditionToggle,
+		EditCondition = "SourceStaticMesh != nullptr || SourceSkeletalMesh != nullptr || SourceDynamicMesh != nullptr"))
+	bool bHideSourceMesh = false;
+
 	/** Source mesh normal map; if empty, the geometric normals will be used */
 	UPROPERTY(EditAnywhere, Category = BakeInput, AdvancedDisplay, meta = (TransientToolProperty,
 		EditCondition = "bHasSourceNormalMap == true", EditConditionHides, HideEditConditionToggle))
diff --git a/Engine/Plugins/Experimental/MeshModelingToolsetExp/Source/MeshModelingToolsExp/Public/BakeMeshAttributeVertexTool.h b/Engine/Plugins/Experimental/MeshModelingToolsetExp/Source/MeshModelingToolsExp/Public/BakeMeshAttributeVertexTool.h
--- a/Engine/Plugins/Experimental/MeshModelingToolsetExp/Source/MeshModelingToolsExp/Public/BakeMeshAttributeVertexTool.h
+++ b/Engine/Plugins/Experimental/MeshModelingToolsetExp/Source/MeshModelingToolsExp/Public/BakeMeshAttributeVertexTool.h
@@ -10,6 +10,7 @@
 #include "ModelingOperators.h"
 #include "MeshOpPreviewHelpers.h"
 #include "PreviewMesh.h"
+#include "ModelingToolTargetUtil.h"
 #include "Sampling/MeshVertexBaker.h"
 #include "BakeMeshAttributeTool.h"
 #include "BakeMeshAttributeVertexTool.generated.h"
@@ -195,6 +196,14 @@ protected:
 	};
 	FBakeSettings CachedBakeSettings;
 
+	void SetSourceObjectVisible(bool bState)
+	{
+		if (!bIsBakeToSelf)
+		{
+			UE::ToolTarget::SetSourceObjectVisible(Targets[1], bState);
+		}
+	}
+
 
 	//
 	// Analytics
