ModelingTools: Add PreviewAsMaterial option to BakeTx and BakeAll.
BakeTx and BakeAll will now default to previewing textures directly. Users can enable PreviewAsMaterial to preview the texture as a material parameter input.

#rnx

diff --git a/Engine/Plugins/Experimental/MeshModelingToolsetExp/Source/MeshModelingToolsExp/Private/BakeMeshAttributeMapsToolBase.cpp b/Engine/Plugins/Experimental/MeshModelingToolsetExp/Source/MeshModelingToolsExp/Private/BakeMeshAttributeMapsToolBase.cpp
--- a/Engine/Plugins/Experimental/MeshModelingToolsetExp/Source/MeshModelingToolsExp/Private/BakeMeshAttributeMapsToolBase.cpp
+++ b/Engine/Plugins/Experimental/MeshModelingToolsetExp/Source/MeshModelingToolsExp/Private/BakeMeshAttributeMapsToolBase.cpp
@@ -73,6 +73,7 @@ void UBakeMeshAttributeMapsToolBase::PostSetup()
 	VisualizationProps = NewObject<UBakeVisualizationProperties>(this);
 	VisualizationProps->RestoreProperties(this);
 	AddToolPropertySource(VisualizationProps);
+	VisualizationProps->WatchProperty(VisualizationProps->bPreviewAsMaterial, [this](bool) { UpdateVisualization(); GetToolManager()->PostInvalidation(); });
 
 	// Initialize background compute
 	Compute = MakeUnique<TGenericDataBackgroundCompute<FMeshMapBaker>>();
@@ -199,47 +200,57 @@ void UBakeMeshAttributeMapsToolBase::UpdatePreview(const EBakeMapType PreviewMap
 	{
 		return;
 	}
-	
+
+	const bool bPreviewAsMaterial = VisualizationProps->bPreviewAsMaterial;
 	UTexture2D* PreviewMap = CachedMaps[PreviewMapType];
-	switch (PreviewMapType)
+	if (bPreviewAsMaterial)
 	{
-	default:
-		PreviewMaterial->SetTextureParameterValue(TEXT("NormalMap"), EmptyNormalMap);
-		PreviewMaterial->SetTextureParameterValue(TEXT("OcclusionMap"), EmptyColorMapWhite);
-		PreviewMaterial->SetTextureParameterValue(TEXT("ColorMap"), EmptyColorMapWhite);
-		break;
-	case EBakeMapType::TangentSpaceNormal:
-		PreviewMaterial->SetTextureParameterValue(TEXT("NormalMap"), PreviewMap);
-		PreviewMaterial->SetTextureParameterValue(TEXT("OcclusionMap"), EmptyColorMapWhite);
-		PreviewMaterial->SetTextureParameterValue(TEXT("ColorMap"), EmptyColorMapWhite);
-		break;
-	case EBakeMapType::AmbientOcclusion:
-		PreviewMaterial->SetTextureParameterValue(TEXT("NormalMap"), EmptyNormalMap);
-		PreviewMaterial->SetTextureParameterValue(TEXT("OcclusionMap"), PreviewMap);
-		PreviewMaterial->SetTextureParameterValue(TEXT("ColorMap"), EmptyColorMapWhite);
-		break;
-	case EBakeMapType::BentNormal:
+		switch (PreviewMapType)
+		{
+		default:
+			PreviewMaterial->SetTextureParameterValue(TEXT("NormalMap"), EmptyNormalMap);
+			PreviewMaterial->SetTextureParameterValue(TEXT("OcclusionMap"), EmptyColorMapWhite);
+			PreviewMaterial->SetTextureParameterValue(TEXT("ColorMap"), EmptyColorMapWhite);
+			break;
+		case EBakeMapType::TangentSpaceNormal:
+			PreviewMaterial->SetTextureParameterValue(TEXT("NormalMap"), PreviewMap);
+			PreviewMaterial->SetTextureParameterValue(TEXT("OcclusionMap"), EmptyColorMapWhite);
+			PreviewMaterial->SetTextureParameterValue(TEXT("ColorMap"), EmptyColorMapWhite);
+			break;
+		case EBakeMapType::AmbientOcclusion:
+			PreviewMaterial->SetTextureParameterValue(TEXT("NormalMap"), EmptyNormalMap);
+			PreviewMaterial->SetTextureParameterValue(TEXT("OcclusionMap"), PreviewMap);
+			PreviewMaterial->SetTextureParameterValue(TEXT("ColorMap"), EmptyColorMapWhite);
+			break;
+		case EBakeMapType::BentNormal:
+		{
+			UTexture2D* AOMap = CachedMaps.Contains(EBakeMapType::AmbientOcclusion) ? CachedMaps[EBakeMapType::AmbientOcclusion] : EmptyColorMapWhite;
+			BentNormalPreviewMaterial->SetTextureParameterValue(TEXT("NormalMap"), EmptyNormalMap);
+			BentNormalPreviewMaterial->SetTextureParameterValue(TEXT("OcclusionMap"), AOMap);
+			BentNormalPreviewMaterial->SetTextureParameterValue(TEXT("ColorMap"), EmptyColorMapWhite);
+			BentNormalPreviewMaterial->SetTextureParameterValue(TEXT("BentNormalMap"), PreviewMap);
+			PreviewMesh->SetOverrideRenderMaterial(BentNormalPreviewMaterial);
+			break;
+		}	
+		case EBakeMapType::Curvature:
+		case EBakeMapType::ObjectSpaceNormal:
+		case EBakeMapType::FaceNormal:
+		case EBakeMapType::Position:
+		case EBakeMapType::MaterialID:
+		case EBakeMapType::Texture:
+		case EBakeMapType::MultiTexture:
+		case EBakeMapType::VertexColor:
+			PreviewMaterial->SetTextureParameterValue(TEXT("NormalMap"), EmptyNormalMap);
+			PreviewMaterial->SetTextureParameterValue(TEXT("OcclusionMap"), EmptyColorMapWhite);
+			PreviewMaterial->SetTextureParameterValue(TEXT("ColorMap"), PreviewMap);
+			break;
+		}
+	}
+	else
 	{
-		UTexture2D* AOMap = CachedMaps.Contains(EBakeMapType::AmbientOcclusion) ? CachedMaps[EBakeMapType::AmbientOcclusion] : EmptyColorMapWhite;
-		BentNormalPreviewMaterial->SetTextureParameterValue(TEXT("NormalMap"), EmptyNormalMap);
-		BentNormalPreviewMaterial->SetTextureParameterValue(TEXT("OcclusionMap"), AOMap);
-		BentNormalPreviewMaterial->SetTextureParameterValue(TEXT("ColorMap"), EmptyColorMapWhite);
-		BentNormalPreviewMaterial->SetTextureParameterValue(TEXT("BentNormalMap"), PreviewMap);
-		PreviewMesh->SetOverrideRenderMaterial(BentNormalPreviewMaterial);
-		break;
-	}	
-	case EBakeMapType::Curvature:
-	case EBakeMapType::ObjectSpaceNormal:
-	case EBakeMapType::FaceNormal:
-	case EBakeMapType::Position:
-	case EBakeMapType::MaterialID:
-	case EBakeMapType::Texture:
-	case EBakeMapType::MultiTexture:
-	case EBakeMapType::VertexColor:
 		PreviewMaterial->SetTextureParameterValue(TEXT("NormalMap"), EmptyNormalMap);
 		PreviewMaterial->SetTextureParameterValue(TEXT("OcclusionMap"), EmptyColorMapWhite);
 		PreviewMaterial->SetTextureParameterValue(TEXT("ColorMap"), PreviewMap);
-		break;
 	}
 	PreviewMaterial->SetScalarParameterValue(TEXT("UVChannel"), CachedBakeSettings.TargetUVLayer);
 	BentNormalPreviewMaterial->SetScalarParameterValue(TEXT("UVChannel"), CachedBakeSettings.TargetUVLayer);
diff --git a/Engine/Plugins/Experimental/MeshModelingToolsetExp/Source/MeshModelingToolsExp/Public/BakeMeshAttributeToolCommon.h b/Engine/Plugins/Experimental/MeshModelingToolsetExp/Source/MeshModelingToolsExp/Public/BakeMeshAttributeToolCommon.h
--- a/Engine/Plugins/Experimental/MeshModelingToolsetExp/Source/MeshModelingToolsExp/Public/BakeMeshAttributeToolCommon.h
+++ b/Engine/Plugins/Experimental/MeshModelingToolsetExp/Source/MeshModelingToolsExp/Public/BakeMeshAttributeToolCommon.h
@@ -278,6 +278,10 @@ class MESHMODELINGTOOLSEXP_API UBakeVisualizationProperties : public UInteractiv
 {
 	GENERATED_BODY()
 public:
+	/** Preview the texture as an input to the respective material parameter (ex. Normal, Ambient Occlusion) */
+	UPROPERTY(EditAnywhere, Category = Preview)
+	bool bPreviewAsMaterial = false;
+	
 	/** Adjust the brightness of the preview material; does not affect results stored in textures */
 	UPROPERTY(EditAnywhere, Category = Preview, meta = (DisplayName = "Brightness", UIMin = "0.0", UIMax = "1.0"))
 	float Brightness = 1.0f;
