MeshModelingToolset: Adjust point size for visualizing degenerate tangents

diff --git a/Engine/Plugins/Experimental/MeshModelingToolsetExp/Source/MeshModelingToolsEditorOnly/Private/MeshTangentsTool.cpp b/Engine/Plugins/Experimental/MeshModelingToolsetExp/Source/MeshModelingToolsEditorOnly/Private/MeshTangentsTool.cpp
--- a/Engine/Plugins/Experimental/MeshModelingToolsetExp/Source/MeshModelingToolsEditorOnly/Private/MeshTangentsTool.cpp
+++ b/Engine/Plugins/Experimental/MeshModelingToolsetExp/Source/MeshModelingToolsEditorOnly/Private/MeshTangentsTool.cpp
@@ -171,9 +171,9 @@ void UMeshTangentsTool::Render(IToolsContextRenderAPI* RenderAPI)
 		{
 			if (ErrorPt.MaxAngleDeg > Settings->CompareWithMikktThreshold)
 			{
-				Visualizer.DrawPoint(ErrorPt.VertexPos, FLinearColor(0.95f, 0.05f, 0.05f), 6.0f * Settings->LineThickness, false);
-				Visualizer.DrawLine<FVector3f>(ErrorPt.VertexPos, ErrorPt.VertexPos + Settings->LineLength * ErrorPt.MikktTangent, FLinearColor(0.95f, 0.05f, 0.05f), 2 * Settings->LineThickness, false);
-				Visualizer.DrawLine<FVector3f>(ErrorPt.VertexPos, ErrorPt.VertexPos + Settings->LineLength * ErrorPt.MikktBitangent, FLinearColor(0.05f, 0.95f, 0.05f), 2 * Settings->LineThickness, false);
+				Visualizer.DrawPoint(ErrorPt.VertexPos, FLinearColor(0.95f, 0.05f, 0.05f), 4.0f * Settings->LineThickness, false);
+				Visualizer.DrawLine<FVector3f>(ErrorPt.VertexPos, ErrorPt.VertexPos + Settings->LineLength * ErrorPt.MikktTangent, FLinearColor(0.95f, 0.05f, 0.05f), 2.0f * Settings->LineThickness, false);
+				Visualizer.DrawLine<FVector3f>(ErrorPt.VertexPos, ErrorPt.VertexPos + Settings->LineLength * ErrorPt.MikktBitangent, FLinearColor(0.05f, 0.95f, 0.05f), 2.0f * Settings->LineThickness, false);
 
 				Visualizer.DrawLine<FVector3f>(ErrorPt.VertexPos, ErrorPt.VertexPos + (1.1f * Settings->LineLength) * ErrorPt.OtherTangent, FLinearColor(0.95f, 0.50f, 0.05f), Settings->LineThickness, false);
 				Visualizer.DrawLine<FVector3f>(ErrorPt.VertexPos, ErrorPt.VertexPos + (1.1f * Settings->LineLength) * ErrorPt.OtherBitangent, FLinearColor(0.05f, 0.95f, 0.95f), Settings->LineThickness, false);
