MeshModelingTools: Fix Boundary Revolve tool not working after CL 17630965 because PolygonSelectionMechanic was not adding its behaviors with the proper source. Also changed it to respect Shift modifier when clicking on a boundary the same way that the plane mechanic does.

#rnx

diff --git a/Engine/Plugins/Experimental/MeshModelingToolsetExp/Source/MeshModelingToolsExp/Private/RevolveBoundaryTool.cpp b/Engine/Plugins/Experimental/MeshModelingToolsetExp/Source/MeshModelingToolsExp/Private/RevolveBoundaryTool.cpp
--- a/Engine/Plugins/Experimental/MeshModelingToolsetExp/Source/MeshModelingToolsExp/Private/RevolveBoundaryTool.cpp
+++ b/Engine/Plugins/Experimental/MeshModelingToolsetExp/Source/MeshModelingToolsExp/Private/RevolveBoundaryTool.cpp
@@ -79,7 +79,8 @@ void URevolveBoundaryTool::Setup()
 
 	USingleClickInputBehavior* ClickBehavior = NewObject<USingleClickInputBehavior>();
 	ClickBehavior->Initialize(this);
-	ClickBehavior->Modifiers.RegisterModifier(AlignAxisModifier, FInputDeviceState::IsCtrlKeyDown);
+	ClickBehavior->Modifiers.RegisterModifier(CtrlModifier, FInputDeviceState::IsCtrlKeyDown);
+	ClickBehavior->Modifiers.RegisterModifier(ShiftModifier, FInputDeviceState::IsShiftKeyDown);
 	AddInputBehavior(ClickBehavior);
 	
 	OutputTypeProperties = NewObject<UCreateMeshObjectTypeProperties>(this);
@@ -117,7 +118,7 @@ void URevolveBoundaryTool::Setup()
 
 	SetToolDisplayName(LOCTEXT("ToolName", "Revolve Boundary"));
 	GetToolManager()->DisplayMessage(
-		LOCTEXT("OnStartRevolveBoundaryTool", "Revolve an open mesh boundary loop around an axis to create a new mesh."),
+		LOCTEXT("OnStartRevolveBoundaryTool", "Revolve an open mesh boundary loop around an axis to create a new mesh. Ctrl+Click to align the axis to a surface/edge."),
 		EToolMessageLevel::UserNotification);
 	if (Topology->Edges.Num() == 1)
 	{
@@ -142,9 +143,17 @@ void URevolveBoundaryTool::Setup()
 
 void URevolveBoundaryTool::OnUpdateModifierState(int ModifierId, bool bIsOn)
 {
-	if (ModifierId == AlignAxisModifier)
+	if (ModifierId == CtrlModifier)
 	{
-		bAlignAxisOnClick = bIsOn;
+		// Like with the plane mechanic, clicking an edge while holding ctrl should move the axis
+		// to that point and (by default) align it.
+		bMoveAxisOnClick = bIsOn;
+	}
+	else if (ModifierId == ShiftModifier)
+	{
+		// Like with the plane mechanic, holding ctrl + shift shifts the axis to the point without
+		// changing the orientation
+		bAlignAxisOnClick = !bIsOn;
 	}
 }
 
@@ -171,8 +180,8 @@ void URevolveBoundaryTool::OnClicked(const FInputDeviceRay& ClickPos)
 		// Clear the "multiple boundaries" warning, since we've selected one.
 		GetToolManager()->DisplayMessage(FText(), EToolMessageLevel::UserWarning);
 
-		// If Ctrl is pressed, we also want to align the revolution axis to the edge that we clicked
-		if (bAlignAxisOnClick)
+		// Act on Ctrl/Shift modifiers the same way that the plane mechanic does.
+		if (bMoveAxisOnClick)
 		{
 			const FGroupTopologySelection& Selection = SelectionMechanic->GetActiveSelection();
 			int32 ClickedEid = Topology->GetGroupEdgeEdges(Selection.GetASelectedEdgeID())[HitResult.Item];
@@ -183,9 +192,13 @@ void URevolveBoundaryTool::OnClicked(const FInputDeviceRay& ClickPos)
 			FLine3d EdgeLine = FLine3d::FromPoints((FVector3d)ToWorldTranform.TransformPosition((FVector)VertexA), 
 				(FVector3d)ToWorldTranform.TransformPosition((FVector)VertexB));
 			
-			FFrame3d RevolutionAxisFrame;
-			RevolutionAxisFrame.Origin = EdgeLine.NearestPoint((FVector3d)HitResult.ImpactPoint);
-			RevolutionAxisFrame.AlignAxis(0, EdgeLine.Direction);
+			// New frame starts as the old one with modified origin
+			FFrame3d RevolutionAxisFrame(EdgeLine.NearestPoint((FVector3d)HitResult.ImpactPoint), PlaneMechanic->Plane.Rotation);
+			
+			if (bAlignAxisOnClick)
+			{
+				RevolutionAxisFrame.AlignAxis(0, EdgeLine.Direction);
+			}
 
 			PlaneMechanic->SetPlaneWithoutBroadcast(RevolutionAxisFrame);
 
diff --git a/Engine/Plugins/Experimental/MeshModelingToolsetExp/Source/MeshModelingToolsExp/Public/RevolveBoundaryTool.h b/Engine/Plugins/Experimental/MeshModelingToolsetExp/Source/MeshModelingToolsExp/Public/RevolveBoundaryTool.h
--- a/Engine/Plugins/Experimental/MeshModelingToolsetExp/Source/MeshModelingToolsExp/Public/RevolveBoundaryTool.h
+++ b/Engine/Plugins/Experimental/MeshModelingToolsetExp/Source/MeshModelingToolsExp/Public/RevolveBoundaryTool.h
@@ -95,9 +95,11 @@ public:
 	virtual void OnUpdateModifierState(int ModifierID, bool bIsOn) override;
 protected:
 
-	// Support for Ctrl+Clicking a boundary to align the revolution axis to that segment
-	bool bAlignAxisOnClick = false;
-	int32 AlignAxisModifier = 2;
+	// Support for Ctrl+(Shift+)Clicking a boundary to align the revolution axis to that segment
+	bool bMoveAxisOnClick = false;
+	bool bAlignAxisOnClick = true;
+	int32 CtrlModifier = 2;
+	int32 ShiftModifier = 3;
 
 	/** Property set for type of output object (StaticMesh, Volume, etc) */
 	UPROPERTY()
diff --git a/Engine/Plugins/Runtime/MeshModelingToolset/Source/ModelingComponents/Private/Selection/PolygonSelectionMechanic.cpp b/Engine/Plugins/Runtime/MeshModelingToolset/Source/ModelingComponents/Private/Selection/PolygonSelectionMechanic.cpp
--- a/Engine/Plugins/Runtime/MeshModelingToolset/Source/ModelingComponents/Private/Selection/PolygonSelectionMechanic.cpp
+++ b/Engine/Plugins/Runtime/MeshModelingToolset/Source/ModelingComponents/Private/Selection/PolygonSelectionMechanic.cpp
@@ -46,7 +46,7 @@ void UPolygonSelectionMechanic::Setup(UInteractiveTool* ParentToolIn)
 	ClickOrDragBehavior->Modifiers.RegisterModifier(ShiftModifierID, FInputDeviceState::IsShiftKeyDown);
 	ClickOrDragBehavior->Modifiers.RegisterModifier(CtrlModifierID, FInputDeviceState::IsCtrlKeyDown);
 	ClickOrDragBehavior->SetDefaultPriority(BasePriority);
-	ParentTool->AddInputBehavior(ClickOrDragBehavior);
+	ParentTool->AddInputBehavior(ClickOrDragBehavior, this);
 
 	Properties = NewObject<UPolygonSelectionMechanicProperties>(this);
 	if (bAddSelectionFilterPropertiesToParentTool)
