MeshModelingToolset: Revolution axis properties now properly show reset-to-default icon after manipulating gizmo

diff --git a/Engine/Plugins/Experimental/MeshModelingToolsetExp/Source/MeshModelingToolsExp/Private/RevolveBoundaryTool.cpp b/Engine/Plugins/Experimental/MeshModelingToolsetExp/Source/MeshModelingToolsExp/Private/RevolveBoundaryTool.cpp
--- a/Engine/Plugins/Experimental/MeshModelingToolsetExp/Source/MeshModelingToolsExp/Private/RevolveBoundaryTool.cpp
+++ b/Engine/Plugins/Experimental/MeshModelingToolsetExp/Source/MeshModelingToolsExp/Private/RevolveBoundaryTool.cpp
@@ -111,6 +111,7 @@ void URevolveBoundaryTool::Setup()
 		FRotator AxisOrientation = ((FQuat)PlaneMechanic->Plane.Rotation).Rotator();
 		Settings->AxisOrientation.X = AxisOrientation.Pitch;
 		Settings->AxisOrientation.Y = AxisOrientation.Yaw;
+		NotifyOfPropertyChangeByTool(Settings);
 		UpdateRevolutionAxis();
 		});
 
