ModelingTools: Fix crash when baking AO in BakeTexture.

#rnx

diff --git a/Engine/Plugins/Experimental/MeshModelingToolsetExp/Source/MeshModelingToolsExp/Private/BakeMeshAttributeMapsTool.cpp b/Engine/Plugins/Experimental/MeshModelingToolsetExp/Source/MeshModelingToolsExp/Private/BakeMeshAttributeMapsTool.cpp
--- a/Engine/Plugins/Experimental/MeshModelingToolsetExp/Source/MeshModelingToolsExp/Private/BakeMeshAttributeMapsTool.cpp
+++ b/Engine/Plugins/Experimental/MeshModelingToolsetExp/Source/MeshModelingToolsExp/Private/BakeMeshAttributeMapsTool.cpp
@@ -171,12 +171,14 @@ public:
 			case EBakeMapType::AmbientOcclusion:
 			{
 				InitOcclusionEval(OcclusionEval);
-				OcclusionEval->OcclusionType |= EMeshOcclusionMapType::AmbientOcclusion; 
+				OcclusionEval->OcclusionType |= EMeshOcclusionMapType::AmbientOcclusion;
+				break;
 			}
 			case EBakeMapType::BentNormal:
 			{
 				InitOcclusionEval(OcclusionEval);
 				OcclusionEval->OcclusionType |= EMeshOcclusionMapType::BentNormal;
+				break;
 			}
 			case EBakeMapType::Curvature:
 			{
