Enable mesh tangent tool for dynamic mesh actors

#rnx

diff --git a/Engine/Plugins/Editor/ModelingToolsEditorMode/Source/ModelingToolsEditorMode/Private/ModelingToolsManagerActions.cpp b/Engine/Plugins/Editor/ModelingToolsEditorMode/Source/ModelingToolsEditorMode/Private/ModelingToolsManagerActions.cpp
--- a/Engine/Plugins/Editor/ModelingToolsEditorMode/Source/ModelingToolsEditorMode/Private/ModelingToolsManagerActions.cpp
+++ b/Engine/Plugins/Editor/ModelingToolsEditorMode/Source/ModelingToolsEditorMode/Private/ModelingToolsManagerActions.cpp
@@ -182,6 +182,8 @@ void FModelingToolsManagerCommands::RegisterCommands()
 	// Note that passing a chord into one of these calls hooks the key press to the respective action. 
 	UI_COMMAND(AcceptOrCompleteActiveTool, "Accept or Complete", "Accept or complete the active tool", EUserInterfaceActionType::Button, FInputChord(EKeys::Enter));
 	UI_COMMAND(CancelOrCompleteActiveTool, "Cancel or Complete", "Cancel or complete the active tool", EUserInterfaceActionType::Button, FInputChord(EKeys::Escape));
+	
+#undef REGISTER_MODELING_TOOL_COMMAND
 }
 
 
diff --git a/Engine/Plugins/Experimental/MeshModelingToolsetExp/Source/MeshModelingToolsEditorOnly/Private/MeshTangentsTool.cpp b/Engine/Plugins/Experimental/MeshModelingToolsetExp/Source/MeshModelingToolsEditorOnly/Private/MeshTangentsTool.cpp
--- a/Engine/Plugins/Experimental/MeshModelingToolsetExp/Source/MeshModelingToolsEditorOnly/Private/MeshTangentsTool.cpp
+++ b/Engine/Plugins/Experimental/MeshModelingToolsetExp/Source/MeshModelingToolsEditorOnly/Private/MeshTangentsTool.cpp
@@ -36,7 +36,6 @@ const FToolTargetTypeRequirements& UMeshTangentsToolBuilder::GetTargetRequiremen
 		UMeshDescriptionCommitter::StaticClass(),
 		UMeshDescriptionProvider::StaticClass(),
 		UPrimitiveComponentBackedTarget::StaticClass(),
-		UStaticMeshBackedTarget::StaticClass()
 		});
 	return TypeRequirements;
 }
@@ -127,14 +126,20 @@ void UMeshTangentsTool::Shutdown(EToolShutdownType ShutdownType)
 		GetToolManager()->BeginUndoTransaction(LOCTEXT("UpdateTangents", "Update Tangents"));
 
 		UStaticMeshComponent* StaticMeshComponent = Cast<UStaticMeshComponent>(UE::ToolTarget::GetTargetComponent(Target));
-		UStaticMesh* StaticMesh = StaticMeshComponent->GetStaticMesh();
-		StaticMesh->Modify();
-
-		// disable auto-generated normals and tangents build settings
-		UE::MeshDescription::FStaticMeshBuildSettingChange SettingsChange;
-		SettingsChange.AutoGeneratedNormals = UE::MeshDescription::EBuildSettingBoolChange::Disable;
-		SettingsChange.AutoGeneratedTangents = UE::MeshDescription::EBuildSettingBoolChange::Disable;
-		UE::MeshDescription::ConfigureBuildSettings(StaticMesh, 0, SettingsChange);
+		if (StaticMeshComponent != nullptr)
+		{
+			UStaticMesh* StaticMesh = StaticMeshComponent->GetStaticMesh();
+			if (ensure(StaticMesh != nullptr))
+			{
+				StaticMesh->Modify();
+
+				// disable auto-generated normals and tangents build settings
+				UE::MeshDescription::FStaticMeshBuildSettingChange SettingsChange;
+				SettingsChange.AutoGeneratedNormals = UE::MeshDescription::EBuildSettingBoolChange::Disable;
+				SettingsChange.AutoGeneratedTangents = UE::MeshDescription::EBuildSettingBoolChange::Disable;
+				UE::MeshDescription::ConfigureBuildSettings(StaticMesh, 0, SettingsChange);
+			}
+		}
 
 		Tangents->CopyToOverlays(*InputMesh);
 
