MeshModelingToolset: Fix NonUnity Compile UnrealEditor Mac in AttributeEditorTool.cpp

diff --git a/Engine/Plugins/Experimental/MeshModelingToolsetExp/Source/MeshModelingToolsEditorOnly/Private/AttributeEditorTool.cpp b/Engine/Plugins/Experimental/MeshModelingToolsetExp/Source/MeshModelingToolsEditorOnly/Private/AttributeEditorTool.cpp
--- a/Engine/Plugins/Experimental/MeshModelingToolsetExp/Source/MeshModelingToolsEditorOnly/Private/AttributeEditorTool.cpp
+++ b/Engine/Plugins/Experimental/MeshModelingToolsetExp/Source/MeshModelingToolsEditorOnly/Private/AttributeEditorTool.cpp
@@ -469,7 +469,7 @@ void UAttributeEditorTool::ClearUVs()
 		{
 			if (!FStaticMeshOperations::RemoveUVChannel(EditedMesh, LayerIndex))
 			{
-				for (const FVertexInstanceID& ElID : Instances.GetElementIDs())
+				for (const FVertexInstanceID ElID : Instances.GetElementIDs())
 				{
 					InstanceUVs.Set(ElID, LayerIndex, FVector2D::ZeroVector);
 				}
@@ -518,7 +518,7 @@ void UAttributeEditorTool::DeleteSelectedUVSet()
 		if (!FStaticMeshOperations::RemoveUVChannel(EditedMesh, DeleteIndex))
 		{
 			const FVertexInstanceArray& Instances = EditedMesh.VertexInstances();
-			for (const FVertexInstanceID& InstanceID : Instances.GetElementIDs())
+			for (const FVertexInstanceID InstanceID : Instances.GetElementIDs())
 			{
 				InstanceUVs.Set(InstanceID, DeleteIndex, FVector2D::ZeroVector);
 			}
@@ -595,7 +595,7 @@ void UAttributeEditorTool::DuplicateSelectedUVSet()
 		if (FStaticMeshOperations::AddUVChannel(EditedMesh))
 		{
 			const FVertexInstanceArray& Instances = EditedMesh.VertexInstances();
-			for (const FVertexInstanceID& InstanceID : Instances.GetElementIDs())
+			for (const FVertexInstanceID InstanceID : Instances.GetElementIDs())
 			{
 				FVector2D SourceUV = InstanceUVs.Get(InstanceID, SourceIndex);
 				InstanceUVs.Set(InstanceID, NewChannelIndex, SourceUV);
