MeshModelingTools: numerous Path Extrude fixes and improvements:
- Fixed snapping not disabling with Shift key.
- Fixed inability to end path if enabling snapping after starting the path.
- Fixed some undo items being undoable in certain circumstances.
- Fixed 0 height paths being treated differently than what PolyExt does (now they are unextruded ribbons instead of 0 thickness extrusions).
- Fixed interactive width selection not aligning with actual sides of path.
- Made behavior more intuitive when using fixed height/width/radius parameters and switching between fixed and interactive modes, including making it possible to set the fixed parameters after leaving that step of creation.
- Made radius setting behavior similar to width setting behavior.
- Disallowed 0 width paths.
- Made 0 radius same as not having rounded corners.
- Made the tool complete a current mesh if exited while one is ready.

#rnx

diff --git a/Engine/Plugins/Experimental/MeshModelingToolsetExp/Source/MeshModelingToolsExp/Private/DrawPolyPathTool.cpp b/Engine/Plugins/Experimental/MeshModelingToolsetExp/Source/MeshModelingToolsExp/Private/DrawPolyPathTool.cpp
--- a/Engine/Plugins/Experimental/MeshModelingToolsetExp/Source/MeshModelingToolsExp/Private/DrawPolyPathTool.cpp
+++ b/Engine/Plugins/Experimental/MeshModelingToolsetExp/Source/MeshModelingToolsExp/Private/DrawPolyPathTool.cpp
@@ -156,10 +156,12 @@ void UDrawPolyPathTool::Setup()
 	// register click behavior
 	USingleClickInputBehavior* ClickBehavior = NewObject<USingleClickInputBehavior>(this);
 	ClickBehavior->Initialize(this);
+	ClickBehavior->Modifiers.RegisterModifier(ShiftModifierID, FInputDeviceState::IsShiftKeyDown);
 	AddInputBehavior(ClickBehavior);
 
 	UMouseHoverBehavior* HoverBehavior = NewObject<UMouseHoverBehavior>(this);
 	HoverBehavior->Initialize(this);
+	HoverBehavior->Modifiers.RegisterModifier(ShiftModifierID, FInputDeviceState::IsShiftKeyDown);
 	AddInputBehavior(HoverBehavior);
 
 	DrawPlaneWorld = FFrame3d();
@@ -184,23 +186,38 @@ void UDrawPolyPathTool::Setup()
 	TransformProps = NewObject<UDrawPolyPathProperties>(this);
 	TransformProps->RestoreProperties(this);
 	AddToolPropertySource(TransformProps);
-	TransformProps->WatchProperty(TransformProps->bSnapToWorldGrid, [this](bool) 
+
+	TransformProps->WatchProperty(TransformProps->WidthMode, [this](EDrawPolyPathWidthMode Mode)
 	{
-		if (SurfacePathMechanic != nullptr)
+		if (State == EState::SettingWidth)
+		{
+			// Switch to the other mode of setting width.
+			BeginSettingWidth();
+		}
+		else if (State != EState::DrawingPath)
 		{
-			SurfacePathMechanic->bSnapToWorldGrid = TransformProps->bSnapToWorldGrid;
+			UpdatePathPreview();
 		}
 	});
-
-	TransformProps->WatchProperty(TransformProps->ExtrudeMode, [this](EDrawPolyPathExtrudeMode)
+	TransformProps->WatchProperty(TransformProps->RadiusMode, [this](EDrawPolyPathRadiusMode Mode)
 	{
-		if(ExtrudeHeightMechanic != nullptr)
+		if (State == EState::SettingRadius)
 		{
-			// regenerate the base path mesh
-			BeginInteractiveExtrudeHeight();
+			// Switch to the other mode of setting radius.
+			BeginSettingRadius();
+		}
+		else if (State != EState::DrawingPath)
+		{
+			UpdatePathPreview();
+		}
+	});
+	TransformProps->WatchProperty(TransformProps->ExtrudeMode, [this](EDrawPolyPathExtrudeMode Mode)
+	{
+		if (State == EState::SettingHeight)
+		{
+			BeginSettingHeight();
 		}
 	});
-
 
 	ExtrudeProperties = NewObject<UDrawPolyPathExtrudeProperties>();
 	ExtrudeProperties->RestoreProperties(this);
@@ -219,15 +236,28 @@ void UDrawPolyPathTool::Setup()
 	SetToolDisplayName(LOCTEXT("ToolName", "Path Extrude"));
 }
 
-
 void UDrawPolyPathTool::Shutdown(EToolShutdownType ShutdownType)
 {
+	// If we have a ready result, go ahead and create it, because especially with "fixed"
+	// height mode, it's easy to think that Enter should result in its creation.
+	// TODO: We could consider letting this tool accept nested accept/cancel commands, but
+	// it's a bit more work.
+	if (State == EState::SettingHeight && ShutdownType != EToolShutdownType::Cancel)
+	{
+		CurHeight = TransformProps->ExtrudeHeight;
+		OnCompleteExtrudeHeight();
+	}
+
 	if (bHasSavedWidth)
 	{
 		TransformProps->Width = SavedWidth;
 		bHasSavedWidth = false;
 	}
-
+	if (bHasSavedRadius)
+	{
+		TransformProps->CornerRadius = SavedRadius;
+		bHasSavedRadius = false;
+	}
 	if (bHasSavedExtrudeHeight)
 	{
 		TransformProps->ExtrudeHeight = SavedExtrudeHeight;
@@ -245,6 +275,46 @@ void UDrawPolyPathTool::Shutdown(EToolShutdownType ShutdownType)
 	ClearPreview();
 }
 
+void UDrawPolyPathTool::OnPropertyModified(UObject* PropertySet, FProperty* Property)
+{
+	if (!ensure(Property))
+	{
+		// Not sure whether this would ever happen, but just in case
+		return;
+	}
+
+	// We deal with these properties here instead of inside a watcher because we frequently update them
+	// ourselves as we set them interactively, and we don't want to bother with doing silent watcher updates.
+	// Instead we rely on the fact that user scrubbing of the properties calls OnPropertyModified whereas
+	// programmatic changes generally don't.
+	if (Property->GetFName() == GET_MEMBER_NAME_CHECKED(UDrawPolyPathProperties, Width)
+		|| Property->GetFName() == GET_MEMBER_NAME_CHECKED(UDrawPolyPathProperties, CornerRadius))
+	{
+		UpdatePathPreview();
+	}
+	else if (State == EState::SettingHeight
+		&& Property->GetFName() == GET_MEMBER_NAME_CHECKED(UDrawPolyPathProperties, ExtrudeHeight)
+		&& (TransformProps->ExtrudeMode == EDrawPolyPathExtrudeMode::Fixed 
+			|| TransformProps->ExtrudeMode == EDrawPolyPathExtrudeMode::RampFixed))
+	{
+		CurHeight = TransformProps->ExtrudeHeight;
+		UpdateExtrudePreview();
+	}
+	else if (Property->GetFName() == GET_MEMBER_NAME_CHECKED(UDrawPolyPathProperties, bRoundedCorners))
+	{
+		if (!TransformProps->bRoundedCorners && State == EState::SettingRadius)
+		{
+			GetToolManager()->EmitObjectChange(this,
+				MakeUnique<FDrawPolyPathStateChange>(CurrentCurveTimestamp, EState::SettingRadius),
+				LOCTEXT("DrawPolyPathCancelRadius", "Cancel setting corner radius"));
+			OnCompleteRadius();
+		}
+		else if(State != EState::DrawingPath)
+		{
+			UpdatePathPreview();
+		}
+	}
+}
 
 bool UDrawPolyPathTool::HitTest(const FRay& Ray, FHitResult& OutHit)
 {
@@ -276,7 +346,17 @@ bool UDrawPolyPathTool::HitTest(const FRay& Ray, FHitResult& OutHit)
 	return false;
 }
 
-
+void UDrawPolyPathTool::OnUpdateModifierState(int ModifierID, bool bIsOn)
+{
+	if (ModifierID == ShiftModifierID && bIgnoreSnappingToggle != bIsOn)
+	{
+		bIgnoreSnappingToggle = bIsOn;
+		if (SurfacePathMechanic)
+		{
+			SurfacePathMechanic->bSnapToWorldGrid = !bIgnoreSnappingToggle;
+		}
+	}
+}
 
 FInputRayHit UDrawPolyPathTool::IsHitByClick(const FInputDeviceRay& ClickPos)
 {
@@ -293,37 +373,61 @@ FInputRayHit UDrawPolyPathTool::IsHitByClick(const FInputDeviceRay& ClickPos)
 
 void UDrawPolyPathTool::OnClicked(const FInputDeviceRay& ClickPos)
 {
-	if (SurfacePathMechanic != nullptr)
+	switch(State)
 	{
-		if (SurfacePathMechanic->TryAddPointFromRay((FRay3d)ClickPos.WorldRay))
-		{
-			if (SurfacePathMechanic->IsDone())
+		case EState::DrawingPath:
+			if (SurfacePathMechanic != nullptr 
+				&& SurfacePathMechanic->TryAddPointFromRay((FRay3d)ClickPos.WorldRay))
+			{
+				if (SurfacePathMechanic->IsDone())
+				{
+					bPathIsClosed = SurfacePathMechanic->LoopWasClosed();
+					GetToolManager()->EmitObjectChange(this, 
+						MakeUnique<FDrawPolyPathStateChange>(CurrentCurveTimestamp, EState::DrawingPath), 
+						LOCTEXT("DrawPolyPathFinishPath", "Finish path"));
+					OnCompleteSurfacePath();
+				}
+				else
+				{
+					GetToolManager()->EmitObjectChange(this, 
+						MakeUnique<FDrawPolyPathStateChange>(CurrentCurveTimestamp, EState::DrawingPath),
+						LOCTEXT("DrawPolyPathAddToPath", "Add point to path"));
+				}
+			}
+			break;
+		case EState::SettingWidth:
+			if (TransformProps->Width == 0)
 			{
-				bPathIsClosed = SurfacePathMechanic->LoopWasClosed();
-				GetToolManager()->EmitObjectChange(this, MakeUnique<FDrawPolyPathStateChange>(CurrentCurveTimestamp), LOCTEXT("DrawPolyPathFinishPath", "Finish path"));
-				OnCompleteSurfacePath();
+				// This can happen accidentally when the user has snapping turned on and it snaps to 0. 
+				// We'll ignore that click and show an error message.
+				GetToolManager()->DisplayMessage(
+					LOCTEXT("ZeroWidthPath", "Cannot set path width to 0."),
+					EToolMessageLevel::UserError);
 			}
 			else
 			{
-				GetToolManager()->EmitObjectChange(this, MakeUnique<FDrawPolyPathStateChange>(CurrentCurveTimestamp), LOCTEXT("DrawPolyPathAddToPath", "Add point to path"));
+				GetToolManager()->DisplayMessage(FText(), EToolMessageLevel::UserError);
+
+				GetToolManager()->EmitObjectChange(this,
+					MakeUnique<FDrawPolyPathStateChange>(CurrentCurveTimestamp, EState::SettingWidth),
+					LOCTEXT("DrawPolyPathBeginWidth", "Set path width"));
+				OnCompleteWidth();
 			}
-		}
-	}
-	else if (CurveDistMechanic != nullptr && !bSpecifyingRadius)
-	{
-		GetToolManager()->EmitObjectChange(this, MakeUnique<FDrawPolyPathStateChange>(CurrentCurveTimestamp), LOCTEXT("DrawPolyPathBeginWidth", "Set path width"));
-		OnCompleteWidth();
-	}
-	else if (CurveDistMechanic != nullptr && bSpecifyingRadius)
-	{
-		GetToolManager()->EmitObjectChange(this, MakeUnique<FDrawPolyPathStateChange>(CurrentCurveTimestamp), LOCTEXT("DrawPolyPathBeginRadius", "Set corner radius"));
-		OnCompleteRadius();
-	}
-	else if (ExtrudeHeightMechanic != nullptr)
-	{
-		CurHeight = TransformProps->ExtrudeHeight;
-		GetToolManager()->EmitObjectChange(this, MakeUnique<FDrawPolyPathStateChange>(CurrentCurveTimestamp), LOCTEXT("DrawPolyPathBeginExtrude", "Set extrude height"));
-		OnCompleteExtrudeHeight();
+			break;
+		case EState::SettingRadius:
+			GetToolManager()->EmitObjectChange(this, 
+				MakeUnique<FDrawPolyPathStateChange>(CurrentCurveTimestamp, EState::SettingRadius), 
+				LOCTEXT("DrawPolyPathBeginRadius", "Set corner radius"));
+			OnCompleteRadius();
+			break;
+		case EState::SettingHeight:
+			if (TransformProps->ExtrudeMode == EDrawPolyPathExtrudeMode::Interactive
+				|| TransformProps->ExtrudeMode == EDrawPolyPathExtrudeMode::RampInteractive)
+			{
+				CurHeight = TransformProps->ExtrudeHeight;
+			}
+			OnCompleteExtrudeHeight();
+			break;
 	}
 }
 
@@ -345,40 +449,60 @@ FInputRayHit UDrawPolyPathTool::BeginHoverSequenceHitTest(const FInputDeviceRay&
 
 bool UDrawPolyPathTool::OnUpdateHover(const FInputDeviceRay& DevicePos)
 {
-	if (SurfacePathMechanic != nullptr)
+	switch (State)
 	{
-		SurfacePathMechanic->UpdatePreviewPoint((FRay3d)DevicePos.WorldRay);
-	} 
-	else if (CurveDistMechanic != nullptr)
-	{
-		CurveDistMechanic->UpdateCurrentDistance(DevicePos.WorldRay);
-
-		double CurveDistance = CurveDistMechanic->CurrentDistance;
-
-		if (TransformProps->bSnapToWorldGrid)
+	case EState::DrawingPath:
+		if (ensure(SurfacePathMechanic != nullptr))
 		{
-			CurveDistance = ToolSceneQueriesUtil::SnapDistanceToWorldGridSize(this, CurveDistance);
+			SurfacePathMechanic->UpdatePreviewPoint((FRay3d)DevicePos.WorldRay);
 		}
-
-		if (bSpecifyingRadius)
+		break;
+	case EState::SettingWidth:
+		if (CurveDistMechanic != nullptr)
+		{
+			CurveDistMechanic->UpdateCurrentDistance(DevicePos.WorldRay);
+			if (TransformProps->WidthMode == EDrawPolyPathWidthMode::Interactive)
+			{
+				double CurveDistance = CurveDistMechanic->CurrentDistance;
+				if (!bIgnoreSnappingToggle)
+				{
+					CurveDistance = ToolSceneQueriesUtil::SnapDistanceToWorldGridSize(this, CurveDistance);
+				}
+				TransformProps->Width = CurveDistance * 2;
+				UpdatePathPreview();
+			}
+		}
+		break;
+	case EState::SettingRadius:
+		if (CurveDistMechanic != nullptr)
 		{
-			TransformProps->CornerRadius = FMath::Clamp(CurveDistance / CurWidth, 0.0, 2.0);
-			CurRadius = TransformProps->CornerRadius;
+			CurveDistMechanic->UpdateCurrentDistance(DevicePos.WorldRay);
+			if (TransformProps->RadiusMode == EDrawPolyPathRadiusMode::Interactive)
+			{
+				double CurveDistance = CurveDistMechanic->CurrentDistance;
+				if (!bIgnoreSnappingToggle)
+				{
+					CurveDistance = ToolSceneQueriesUtil::SnapDistanceToWorldGridSize(this, CurveDistance);
+				}
+				TransformProps->CornerRadius = FMath::Clamp(CurveDistance / TransformProps->Width, 0.0, 2.0);
+				UpdatePathPreview();
+			}
 		}
-		else
+		break;
+	case EState::SettingHeight:
+		if ((TransformProps->ExtrudeMode == EDrawPolyPathExtrudeMode::Interactive
+			|| TransformProps->ExtrudeMode == EDrawPolyPathExtrudeMode::RampInteractive)
+			&& ensure(ExtrudeHeightMechanic != nullptr))
 		{
-			TransformProps->Width = CurveDistance;
-			CurWidth = 0.5 * CurveDistance;
+			ExtrudeHeightMechanic->UpdateCurrentDistance(DevicePos.WorldRay);
+			
+			CurHeight = ExtrudeHeightMechanic->CurrentHeight;
+			TransformProps->ExtrudeHeight = ExtrudeHeightMechanic->CurrentHeight;
+			UpdateExtrudePreview();
 		}
-		UpdatePathPreview();
-	}
-	else if (ExtrudeHeightMechanic != nullptr)
-	{
-		ExtrudeHeightMechanic->UpdateCurrentDistance(DevicePos.WorldRay);
-		CurHeight = ExtrudeHeightMechanic->CurrentHeight;
-		TransformProps->ExtrudeHeight = ExtrudeHeightMechanic->CurrentHeight;
-		UpdateExtrudePreview();
+		break;
 	}
+
 	return true;
 }
 
@@ -451,6 +575,11 @@ void UDrawPolyPathTool::Render(IToolsContextRenderAPI* RenderAPI)
 
 void UDrawPolyPathTool::InitializeNewSurfacePath()
 {
+	State = EState::DrawingPath;
+
+	CurveDistMechanic = nullptr;
+	CurveDistMechanic = nullptr;
+
 	SurfacePathMechanic = NewObject<UCollectSurfacePathMechanic>(this);
 	SurfacePathMechanic->Setup(this);
 	double SnapTol = ToolSceneQueriesUtil::GetDefaultVisualAngleSnapThreshD();
@@ -459,11 +588,7 @@ void UDrawPolyPathTool::InitializeNewSurfacePath()
 		return ToolSceneQueriesUtil::PointSnapQuery(this->CameraState, Position1, Position2, SnapTol);
 	};
 	SurfacePathMechanic->SetDoubleClickOrCloseLoopMode();
-
-	if (TransformProps)
-	{
-		SurfacePathMechanic->bSnapToWorldGrid = TransformProps->bSnapToWorldGrid;
-	}
+	SurfacePathMechanic->bSnapToWorldGrid = !bIgnoreSnappingToggle;
 
 	UpdateSurfacePathPlane();
 
@@ -473,7 +598,7 @@ void UDrawPolyPathTool::InitializeNewSurfacePath()
 
 bool UDrawPolyPathTool::CanUpdateDrawPlane() const
 {
- 	return (SurfacePathMechanic != nullptr && SurfacePathMechanic->HitPath.Num() == 0);
+ 	return (State == EState::DrawingPath && SurfacePathMechanic != nullptr && SurfacePathMechanic->HitPath.Num() == 0);
 }
 
 void UDrawPolyPathTool::UpdateSurfacePathPlane()
@@ -535,104 +660,113 @@ void UDrawPolyPathTool::OnCompleteSurfacePath()
 	}
 
 	SurfacePathMechanic = nullptr;
-	if (TransformProps->WidthMode == EDrawPolyPathWidthMode::Fixed)
-	{
-		BeginConstantWidth();
-	}
-	else
-	{
-		BeginInteractiveWidth();
-	}
+
+	InitializePreviewMesh();
+
+	// Progress to next state
+	BeginSettingWidth();
 }
 
 
-void UDrawPolyPathTool::BeginInteractiveWidth()
+void UDrawPolyPathTool::BeginSettingWidth()
 {
-	bSpecifyingRadius = false;
-	bHasSavedWidth = true;
-	SavedWidth = TransformProps->Width;
+	// Note that even when the width is constant, we still wait for a click from the user because we 
+	// want them to have a chance to edit it in the detail panel, or to switch to interactive width mode.
 	
-	// begin setting offset distance
+	State = EState::SettingWidth;
+
+	if (TransformProps->WidthMode == EDrawPolyPathWidthMode::Interactive)
+	{
+		bHasSavedWidth = true;
+		SavedWidth = TransformProps->Width;
+	}
+	else if(bHasSavedWidth)
+	{
+		TransformProps->Width = SavedWidth;
+		bHasSavedWidth = false;
+	}
+
 	CurveDistMechanic = NewObject<USpatialCurveDistanceMechanic>(this);
 	CurveDistMechanic->Setup(this);
 	CurveDistMechanic->InitializePolyCurve(CurPolyLine, UE::Geometry::FTransform3d::Identity());
 
-	InitializePreviewMesh();
-
-	ShowOffsetMessage();
-}
-
+	ExtrudeHeightMechanic = nullptr;
 
-void UDrawPolyPathTool::BeginConstantWidth()
-{
-	bSpecifyingRadius = false;
-	InitializePreviewMesh();
-	CurWidth = TransformProps->Width * 0.5;
 	UpdatePathPreview();
-	OnCompleteWidth();
+
+	if (TransformProps->WidthMode == EDrawPolyPathWidthMode::Interactive)
+	{
+		GetToolManager()->DisplayMessage(
+			LOCTEXT("OnStartOffset", "Set the width of the path by clicking on the drawing plane. Hold Shift to ignore snapping."),
+			EToolMessageLevel::UserNotification);
+	}
+	else
+	{
+		GetToolManager()->DisplayMessage(
+			LOCTEXT("OnStartOffset", "Click in viewport to accept fixed path width, or change it in details panel."),
+			EToolMessageLevel::UserNotification);
+	}
 }
 
 void UDrawPolyPathTool::OnCompleteWidth()
 {
 	if (TransformProps->bRoundedCorners)
 	{
-		if (TransformProps->WidthMode == EDrawPolyPathWidthMode::Fixed)
-		{
-			BeginConstantRadius();
-		}
-		else
-		{
-			BeginInteractiveRadius();
-		}
+		BeginSettingRadius();
 	}
 	else
 	{
-		BeginConstantRadius();
+		// Skip radius setting
+		OnCompleteRadius();
 	}
 }
 
 
-void UDrawPolyPathTool::BeginConstantRadius()
+void UDrawPolyPathTool::BeginSettingRadius()
 {
-	bSpecifyingRadius = true;
-	InitializePreviewMesh();
-	CurRadius = TransformProps->CornerRadius;
-	UpdatePathPreview();
-	OnCompleteRadius();
-}
+	// Note that even when the radius is constant, we still wait for a click from the user because we 
+	// want them to have a chance to edit it in the detail panel, or to switch to interactive radius mode.
 
-void UDrawPolyPathTool::BeginInteractiveRadius()
-{
-	bSpecifyingRadius = true;
+	State = EState::SettingRadius;
+
+	if (TransformProps->RadiusMode == EDrawPolyPathRadiusMode::Interactive)
+	{
+		bHasSavedRadius = true;
+		SavedRadius = TransformProps->CornerRadius;
+	}
+	else if (bHasSavedRadius)
+	{
+		TransformProps->CornerRadius = SavedRadius;
+		bHasSavedRadius = false;
+	}
 
-	// begin setting offset distance
 	CurveDistMechanic = NewObject<USpatialCurveDistanceMechanic>(this);
 	CurveDistMechanic->Setup(this);
 	CurveDistMechanic->InitializePolyCurve(CurPolyLine, UE::Geometry::FTransform3d::Identity());
 
-	InitializePreviewMesh();
-}
+	ExtrudeHeightMechanic = nullptr;
 
-void UDrawPolyPathTool::OnCompleteRadius()
-{
-	CurveDistMechanic = nullptr;
+	UpdatePathPreview();
 
-	if (TransformProps->ExtrudeMode == EDrawPolyPathExtrudeMode::Flat)
-	{
-		CurHeight = 0.0;
-		OnCompleteExtrudeHeight();
-	}
-	else if (TransformProps->ExtrudeMode == EDrawPolyPathExtrudeMode::Fixed || TransformProps->ExtrudeMode == EDrawPolyPathExtrudeMode::RampFixed)
+	if (TransformProps->RadiusMode == EDrawPolyPathRadiusMode::Interactive)
 	{
-		CurHeight = TransformProps->ExtrudeHeight;
-		OnCompleteExtrudeHeight();
+		GetToolManager()->DisplayMessage(
+			LOCTEXT("OnStartOffset", "Set the radius of the corners by clicking on the drawing plane. Hold Shift to ignore snapping."),
+			EToolMessageLevel::UserNotification);
 	}
 	else
 	{
-		BeginInteractiveExtrudeHeight();
+		GetToolManager()->DisplayMessage(
+			LOCTEXT("OnStartOffset", "Click in viewport to accept fixed corner radius, or change it in details panel."),
+			EToolMessageLevel::UserNotification);
 	}
 }
 
+void UDrawPolyPathTool::OnCompleteRadius()
+{
+	BeginSettingHeight();
+}
+
 
 void UDrawPolyPathTool::OnCompleteExtrudeHeight()
 {
@@ -649,13 +783,24 @@ void UDrawPolyPathTool::OnCompleteExtrudeHeight()
 
 void UDrawPolyPathTool::UpdatePathPreview()
 {
-	check(EditPreview != nullptr);
+	if (EditPreview == nullptr)
+	{
+		return;
+	}
 
 	FDynamicMesh3 PathMesh;
 	GeneratePathMesh(PathMesh);
-	EditPreview->ReplaceMesh( MoveTempIfPossible(PathMesh) );
-}
 
+	if (State == EState::SettingHeight)
+	{
+		EditPreview->InitializeExtrudeType(MoveTemp(PathMesh), DrawPlaneWorld.Z(), nullptr, false);
+		UpdateExtrudePreview();
+	}
+	else
+	{
+		EditPreview->ReplaceMesh(MoveTempIfPossible(PathMesh));
+	}
+}
 
 
 void UDrawPolyPathTool::GeneratePathMesh(FDynamicMesh3& Mesh)
@@ -668,12 +813,14 @@ void UDrawPolyPathTool::GeneratePathMesh(FDynamicMesh3& Mesh)
 	DrawPolyPathToolLocals::GeneratePathMesh(Mesh, 
 		CurPathPoints, 
 		OffsetScaleFactors, 
-		CurWidth, 
+		TransformProps->Width/2,
 		bPathIsClosed, 
 		bRampMode, 
 		TransformProps->bSinglePolyGroup, 
-		TransformProps->bRoundedCorners, 
-		CurRadius,
+		// Treat radius 0 corners as not rounded, rather than placing a bunch of vertices
+		// in the same place.
+		TransformProps->bRoundedCorners && TransformProps->CornerRadius > 0,
+		TransformProps->CornerRadius,
 		bLimitCornerRadius,
 		TransformProps->RadialSlices);
 
@@ -690,7 +837,20 @@ void UDrawPolyPathTool::GeneratePathMesh(FDynamicMesh3& Mesh)
 	}
 }
 
+void UDrawPolyPathTool::BeginSettingHeight()
+{
+	State = EState::SettingHeight;
 
+	if (TransformProps->ExtrudeMode == EDrawPolyPathExtrudeMode::Interactive
+		|| TransformProps->ExtrudeMode == EDrawPolyPathExtrudeMode::RampInteractive)
+	{
+		BeginInteractiveExtrudeHeight();
+	}
+	else
+	{
+		BeginConstantExtrudeHeight();
+	}
+}
 
 void UDrawPolyPathTool::BeginInteractiveExtrudeHeight()
 {
@@ -707,18 +867,17 @@ void UDrawPolyPathTool::BeginInteractiveExtrudeHeight()
 	};
 	ExtrudeHeightMechanic->WorldPointSnapFunc = [this](const FVector3d& WorldPos, FVector3d& SnapPos)
 	{
-		return TransformProps->bSnapToWorldGrid && ToolSceneQueriesUtil::FindWorldGridSnapPoint(this, WorldPos, SnapPos);
+		return !bIgnoreSnappingToggle && ToolSceneQueriesUtil::FindWorldGridSnapPoint(this, WorldPos, SnapPos);
 	};
 	ExtrudeHeightMechanic->CurrentHeight = 1.0f;  // initialize to something non-zero...prob should be based on polygon bounds maybe?
 
-	InitializePreviewMesh();
-
 	FDynamicMesh3 PathMesh;
 	GeneratePathMesh(PathMesh);
 	EditPreview->InitializeExtrudeType(MoveTemp(PathMesh), DrawPlaneWorld.Z(), nullptr, false);
 
 	FDynamicMesh3 TmpMesh;
 	EditPreview->MakeExtrudeTypeHitTargetMesh(TmpMesh, false);
+
 	FFrame3d UseFrame = DrawPlaneWorld; 
 	UseFrame.Origin = CurPathPoints.Last().Origin;
 	ExtrudeHeightMechanic->Initialize(MoveTemp(TmpMesh), UseFrame, true);
@@ -726,7 +885,42 @@ void UDrawPolyPathTool::BeginInteractiveExtrudeHeight()
 	ShowExtrudeMessage();
 }
 
+void UDrawPolyPathTool::BeginConstantExtrudeHeight()
+{
+	State = EState::SettingHeight;
+
+	if (bHasSavedExtrudeHeight)
+	{
+		TransformProps->ExtrudeHeight = SavedExtrudeHeight;
+		bHasSavedExtrudeHeight = false;
+	}
+
+	if (TransformProps->ExtrudeMode == EDrawPolyPathExtrudeMode::Flat)
+	{
+		CurHeight = 0.0;
+	}
+	else if (TransformProps->ExtrudeMode == EDrawPolyPathExtrudeMode::Fixed || TransformProps->ExtrudeMode == EDrawPolyPathExtrudeMode::RampFixed)
+	{
+		CurHeight = TransformProps->ExtrudeHeight;
+	}
+	else
+	{
+		ensure(false);
+	}
+
+	FDynamicMesh3 PathMesh;
+	GeneratePathMesh(PathMesh);
+	EditPreview->InitializeExtrudeType(MoveTemp(PathMesh), DrawPlaneWorld.Z(), nullptr, false);
+	UpdateExtrudePreview();
+
+	ExtrudeHeightMechanic = nullptr;
+
+	GetToolManager()->DisplayMessage(
+		LOCTEXT("OnStartOffset", "Click in viewport to accept fixed path height, or change it in details panel."),
+		EToolMessageLevel::UserNotification);
+}
 
+// This should only be used after doing EditPreview->InitializeExtrudeType, when setting height
 void UDrawPolyPathTool::UpdateExtrudePreview()
 {
 	EditPreview->UpdateExtrudeType([&](FDynamicMesh3& Mesh) { GenerateExtrudeMesh(Mesh); }, true);
@@ -768,6 +962,13 @@ void UDrawPolyPathTool::ClearPreview()
 
 void UDrawPolyPathTool::GenerateExtrudeMesh(FDynamicMesh3& PathMesh)
 {
+	if (CurHeight == 0)
+	{
+		// This behavior should match whatever we do in DrawPolygonTool. Currently a flat mesh should
+		// just be a one-sided strip, not a degenerately thin closed shape.
+		return;
+	}
+
 	FExtrudeMesh Extruder(&PathMesh);
 	const FVector3d ExtrudeDir = DrawPlaneWorld.Z();
 
@@ -834,7 +1035,11 @@ void UDrawPolyPathTool::EmitNewObject()
 		TransformProps->Width = SavedWidth;
 		bHasSavedWidth = false;
 	}
-
+	if (bHasSavedRadius)
+	{
+		TransformProps->CornerRadius = SavedRadius;
+		bHasSavedRadius = false;
+	}
 	if (bHasSavedExtrudeHeight)
 	{
 		TransformProps->ExtrudeHeight = SavedExtrudeHeight;
@@ -854,12 +1059,6 @@ void UDrawPolyPathTool::ShowStartupMessage()
 		EToolMessageLevel::UserNotification);
 }
 
-void UDrawPolyPathTool::ShowOffsetMessage()
-{
-	GetToolManager()->DisplayMessage(
-		LOCTEXT("OnStartOffset", "Set the width of the path by clicking on the drawing plane."),
-		EToolMessageLevel::UserNotification);
-}
 
 void UDrawPolyPathTool::ShowExtrudeMessage()
 {
@@ -870,61 +1069,53 @@ void UDrawPolyPathTool::ShowExtrudeMessage()
 
 
 
-void UDrawPolyPathTool::UndoCurrentOperation()
+void UDrawPolyPathTool::UndoCurrentOperation(EState DestinationState)
 {
-	if (SurfacePathMechanic != nullptr)
+	switch (State)
 	{
-		SurfacePathMechanic->PopLastPoint();
-		if (SurfacePathMechanic->HitPath.Num() == 0)
+	case EState::DrawingPath:
+		if (ensure(DestinationState == EState::DrawingPath))
 		{
-			CurrentCurveTimestamp++;
+			SurfacePathMechanic->PopLastPoint();
+			if (SurfacePathMechanic->HitPath.Num() == 0)
+			{
+				CurrentCurveTimestamp++;
+			}
 		}
-	}
-	else if (CurveDistMechanic != nullptr)
-	{
-		if (!bSpecifyingRadius)
+		break;
+	case EState::SettingWidth:
+		if (ensure(DestinationState == EState::DrawingPath))
 		{
 			CurveDistMechanic = nullptr;
 			ClearPreview();
 			InitializeNewSurfacePath();
 			SurfacePathMechanic->HitPath = CurPathPoints;
 		}
-		else
+		break;
+	case EState::SettingRadius:
+		if (ensure(DestinationState == EState::SettingWidth))
 		{
-			BeginInteractiveWidth();
+			BeginSettingWidth();
 		}
-	}
-	else if (ExtrudeHeightMechanic != nullptr)
-	{
-		ExtrudeHeightMechanic = nullptr;
-
-		if (TransformProps->bRoundedCorners && TransformProps->WidthMode == EDrawPolyPathWidthMode::Interactive)
+		break;
+	case EState::SettingHeight:
+		if (DestinationState == EState::SettingRadius)
 		{
-			BeginInteractiveRadius();
+			TransformProps->bRoundedCorners = true;
+			BeginSettingRadius();
 		}
-		else
+		else if (ensure(DestinationState == EState::SettingWidth))
 		{
-			if (TransformProps->WidthMode == EDrawPolyPathWidthMode::Fixed)
-			{
-				// Pop all the way back to initial path creation
-				CurveDistMechanic = nullptr;
-				ClearPreview();
-				InitializeNewSurfacePath();
-				SurfacePathMechanic->HitPath = CurPathPoints;
-			}
-			else
-			{
-				BeginInteractiveWidth();
-			}
+			BeginSettingWidth();
 		}
-	
+		break;
 	}
 }
 
 
 void FDrawPolyPathStateChange::Revert(UObject* Object)
 {
-	Cast<UDrawPolyPathTool>(Object)->UndoCurrentOperation();
+	Cast<UDrawPolyPathTool>(Object)->UndoCurrentOperation(PreviousState);
 	bHaveDoneUndo = true;
 }
 bool FDrawPolyPathStateChange::HasExpired(UObject* Object) const
diff --git a/Engine/Plugins/Experimental/MeshModelingToolsetExp/Source/MeshModelingToolsExp/Public/DrawPolyPathTool.h b/Engine/Plugins/Experimental/MeshModelingToolsetExp/Source/MeshModelingToolsExp/Public/DrawPolyPathTool.h
--- a/Engine/Plugins/Experimental/MeshModelingToolsetExp/Source/MeshModelingToolsExp/Public/DrawPolyPathTool.h
+++ b/Engine/Plugins/Experimental/MeshModelingToolsetExp/Source/MeshModelingToolsExp/Public/DrawPolyPathTool.h
@@ -41,6 +41,16 @@ enum class EDrawPolyPathWidthMode
 	Interactive
 };
 
+UENUM()
+enum class EDrawPolyPathRadiusMode
+{
+	/** Fixed radius determined by the CornerRadius property. */
+	Fixed,
+
+	/** Set the radius interactively by clicking in the viewport.  */
+	Interactive
+};
+
 UENUM()
 enum class EDrawPolyPathExtrudeMode
 {
@@ -72,16 +82,23 @@ public:
 	EDrawPolyPathWidthMode WidthMode = EDrawPolyPathWidthMode::Interactive;
 
 	/** Width of the drawn path when using Fixed width mode; also shows the width in Interactive width mode */
-	UPROPERTY(EditAnywhere, Category = Path, meta = (EditCondition = "WidthMode == EDrawPolyPathWidthMode::Fixed", UIMin = "0.0001", UIMax = "1000", ClampMin = "0", ClampMax = "999999"))
+	UPROPERTY(EditAnywhere, Category = Path, meta = (EditCondition = "WidthMode == EDrawPolyPathWidthMode::Fixed", 
+		UIMin = "0.0001", UIMax = "1000", ClampMin = "0.0001", ClampMax = "999999"))
 	float Width = 10.0f;
 
 	/** Use arc segments instead of straight lines in corners */
 	UPROPERTY(EditAnywhere, Category = Path)
 	bool bRoundedCorners = false;
 
-	/** Radius of the corner arcs, as a fraction of path width. This is only available if Rounded Corners is enabled */
-	UPROPERTY(EditAnywhere, Category = Path, meta = (EditCondition = "WidthMode == EDrawPolyPathWidthMode::Fixed && bRoundedCorners", UIMin = "0.01", UIMax = "2.0", ClampMin = "0.01", ClampMax = "999999"))
-	float CornerRadius = 0.01f;
+	/** How the rounded corner radius gets set */
+	UPROPERTY(EditAnywhere, Category = Path, meta = (EditCondition = "bRoundedCorners"))
+	EDrawPolyPathRadiusMode RadiusMode = EDrawPolyPathRadiusMode::Interactive;
+
+	/** Radius of the corner arcs, as a fraction of path width. This is only available if Rounded Corners is enabled. */
+	UPROPERTY(EditAnywhere, Category = Path, meta = (
+		EditCondition = "RadiusMode == EDrawPolyPathRadiusMode::Fixed && bRoundedCorners", 
+		UIMin = "0", UIMax = "2.0", ClampMin = "0", ClampMax = "999999"))
+	float CornerRadius = 0.5f;
 
 	/** Number of radial subdivisions for rounded corners */
 	UPROPERTY(EditAnywhere, NonTransactional, Category = Path, meta = (UIMin = "3", UIMax = "100", ClampMin = "3", ClampMax = "10000",
@@ -107,10 +124,6 @@ public:
 		EditCondition = "ExtrudeMode == EDrawPolyPathExtrudeMode::RampFixed || ExtrudeMode == EDrawPolyPathExtrudeMode::RampInteractive",
 		UIMin = "0.0", UIMax = "1.0", ClampMin = "0.0", ClampMax = "100.0"))
 	float RampStartRatio = 0.05f;
-
-	/** If true, allows snapping to world grid according to editor settings */
-	UPROPERTY(EditAnywhere, Category = Snapping)
-	bool bSnapToWorldGrid = true;
 };
 
 
@@ -156,6 +169,7 @@ public:
 	virtual void Setup() override;
 	virtual void Shutdown(EToolShutdownType ShutdownType) override;
 
+	virtual void OnPropertyModified(UObject* PropertySet, FProperty* Property) override;
 	virtual void OnTick(float DeltaTime) override;
 	virtual void Render(IToolsContextRenderAPI* RenderAPI) override;
 
@@ -175,6 +189,9 @@ public:
 	virtual bool OnUpdateHover(const FInputDeviceRay& DevicePos) override;
 	virtual void OnEndHover() override {}
 
+	// IModifierToggleBehaviorTarget
+	virtual void OnUpdateModifierState(int ModifierID, bool bIsOn) override;
+
 protected:
 	UWorld* TargetWorld;
 
@@ -192,6 +209,15 @@ protected:
 	TObjectPtr<UNewMeshMaterialProperties> MaterialProperties;
 
 protected:
+	enum class EState
+	{
+		DrawingPath,
+		SettingWidth,
+		SettingRadius,
+		SettingHeight
+	};
+	EState State = EState::DrawingPath;
+
 	// camera state at last render
 	UE::Geometry::FTransform3d WorldTransform;
 	FViewCameraState CameraState;
@@ -212,15 +238,18 @@ protected:
 	TArray<double> OffsetScaleFactors;
 	TArray<FVector3d> CurPolyLine;
 	double CurPathLength;
-	double CurWidth;
-	double CurRadius;
 	double CurHeight;
 	bool bHasSavedWidth = false;
 	float SavedWidth;
+	bool bHasSavedRadius = false;
+	float SavedRadius;
 	bool bHasSavedExtrudeHeight = false;
 	float SavedExtrudeHeight;
 	bool bPathIsClosed = false;		// If true, CurPathPoints are assumed to define a closed path
 
+	static const int ShiftModifierID = 1;
+	bool bIgnoreSnappingToggle = false;		// toggled by hotkey (shift)
+
 	TArray<FVector3d> CurPolyLoop;
 	TArray<FVector3d> SecondPolyLoop;
 
@@ -240,15 +269,15 @@ protected:
 	void UpdateSurfacePathPlane();
 	void OnCompleteSurfacePath();
 
-	void BeginInteractiveWidth();
-	void BeginConstantWidth();
+	void BeginSettingWidth();
 	void OnCompleteWidth();
 
-	void BeginInteractiveRadius();
-	void BeginConstantRadius();
+	void BeginSettingRadius();
 	void OnCompleteRadius();
 
+	void BeginSettingHeight();
 	void BeginInteractiveExtrudeHeight();
+	void BeginConstantExtrudeHeight();
 	void UpdateExtrudePreview();
 	void OnCompleteExtrudeHeight();
 
@@ -267,7 +296,7 @@ protected:
 
 	friend class FDrawPolyPathStateChange;
 	int32 CurrentCurveTimestamp = 1;
-	void UndoCurrentOperation();
+	void UndoCurrentOperation(EState DestinationState);
 	bool CheckInCurve(int32 Timestamp) const { return CurrentCurveTimestamp == Timestamp; }
 };
 
@@ -279,9 +308,11 @@ class MESHMODELINGTOOLSEXP_API FDrawPolyPathStateChange : public FToolCommandCha
 public:
 	bool bHaveDoneUndo = false;
 	int32 CurveTimestamp = 0;
-	FDrawPolyPathStateChange(int32 CurveTimestampIn)
+	UDrawPolyPathTool::EState PreviousState = UDrawPolyPathTool::EState::DrawingPath;
+	FDrawPolyPathStateChange(int32 CurveTimestampIn, UDrawPolyPathTool::EState PreviousStateIn)
+		: CurveTimestamp(CurveTimestampIn)
+		, PreviousState(PreviousStateIn)
 	{
-		CurveTimestamp = CurveTimestampIn;
 	}
 	virtual void Apply(UObject* Object) override {}
 	virtual void Revert(UObject* Object) override;
diff --git a/Engine/Plugins/Runtime/MeshModelingToolset/Source/ModelingComponents/Private/Mechanics/CollectSurfacePathMechanic.cpp b/Engine/Plugins/Runtime/MeshModelingToolset/Source/ModelingComponents/Private/Mechanics/CollectSurfacePathMechanic.cpp
--- a/Engine/Plugins/Runtime/MeshModelingToolset/Source/ModelingComponents/Private/Mechanics/CollectSurfacePathMechanic.cpp
+++ b/Engine/Plugins/Runtime/MeshModelingToolset/Source/ModelingComponents/Private/Mechanics/CollectSurfacePathMechanic.cpp
@@ -214,15 +214,13 @@ bool UCollectSurfacePathMechanic::RayToPathPoint(const FRay3d& Ray, FFrame3d& Po
 
 	PointOut = NearestHitFrame;
 
-	// try snapping to close if we are in loop mode
+	// try snapping to close/end
 	bool bHaveSnapped = false;
-	if (DoneMode == ECollectSurfacePathDoneMode::SnapCloseLoop && HitPath.Num() > 2)
+	bool bWouldCloseLoop = false;
+	if (CheckGeometricClosure(PointOut, &bWouldCloseLoop))
 	{
-		if (SpatialSnapPointsFunc(PointOut.Origin, HitPath.Last().Origin))
-		{
-			PointOut = HitPath.Last();
-			bHaveSnapped = true;
-		}
+		bHaveSnapped = true;
+		PointOut = bWouldCloseLoop ? HitPath[0] : HitPath.Last();
 	}
 
 	// try snapping to other things, if we haven't yet
