MeshModelingTools: Flip E/Q controls in CubeGrid tool to align with flight controls. Display message showing whether tool is creating a new asset or editing an existing one.

#rnx

diff --git a/Engine/Plugins/Experimental/MeshModelingToolsetExp/Source/MeshModelingToolsExp/Private/CubeGridTool.cpp b/Engine/Plugins/Experimental/MeshModelingToolsetExp/Source/MeshModelingToolsExp/Private/CubeGridTool.cpp
--- a/Engine/Plugins/Experimental/MeshModelingToolsetExp/Source/MeshModelingToolsExp/Private/CubeGridTool.cpp
+++ b/Engine/Plugins/Experimental/MeshModelingToolsetExp/Source/MeshModelingToolsExp/Private/CubeGridTool.cpp
@@ -835,6 +835,17 @@ void UCubeGridTool::Setup()
 	Settings->SilentUpdateWatched();
 
 	UpdateComputeInputs();
+
+	if (Target)
+	{
+		GetToolManager()->DisplayMessage(LOCTEXT("EditingExistingAssetLabel", "Editing existing asset."), 
+			EToolMessageLevel::UserWarning);
+	}
+	else
+	{
+		GetToolManager()->DisplayMessage(LOCTEXT("CreatingNewAssetLabel", "Creating new asset."),
+			EToolMessageLevel::UserWarning);
+	}
 }
 
 void UCubeGridTool::UpdateComputeInputs()
@@ -1840,25 +1851,25 @@ void UCubeGridTool::RegisterActions(FInteractiveToolActionSet& ActionSet)
 		TEXT("PullBlock"),
 		LOCTEXT("PullAction", "Pull Out Blocks"),
 		LOCTEXT("PullTooltip", ""),
-		EModifierKey::None, EKeys::Q,
+		EModifierKey::None, EKeys::E,
 		[this]() { RequestAction(ECubeGridToolAction::Pull); });
 	ActionSet.RegisterAction(this, ActionID++,
 		TEXT("PushBlock"),
 		LOCTEXT("PushAction", "Push In Holes"),
 		LOCTEXT("PushTooltip", ""),
-		EModifierKey::None, EKeys::E,
+		EModifierKey::None, EKeys::Q,
 		[this]() { RequestAction(ECubeGridToolAction::Push); });
 	ActionSet.RegisterAction(this, ActionID++,
 		TEXT("SlideBack"),
 		LOCTEXT("SlideBackAction", "Slide Selection Back"),
 		LOCTEXT("SlideBackTooltip", ""),
-		EModifierKey::Shift, EKeys::Q,
+		EModifierKey::Shift, EKeys::E,
 		[this]() { RequestAction(ECubeGridToolAction::SlideBack); });
 	ActionSet.RegisterAction(this, ActionID++,
 		TEXT("SlideForward"),
 		LOCTEXT("SlideForwardAction", "Slide Selection Forward"),
 		LOCTEXT("SlideForwardTooltip", ""),
-		EModifierKey::Shift, EKeys::E,
+		EModifierKey::Shift, EKeys::Q,
 		[this]() { RequestAction(ECubeGridToolAction::SlideForward); });
 
 	ActionSet.RegisterAction(this, ActionID++,
diff --git a/Engine/Plugins/Experimental/MeshModelingToolsetExp/Source/MeshModelingToolsExp/Public/CubeGridTool.h b/Engine/Plugins/Experimental/MeshModelingToolsetExp/Source/MeshModelingToolsExp/Public/CubeGridTool.h
--- a/Engine/Plugins/Experimental/MeshModelingToolsetExp/Source/MeshModelingToolsExp/Public/CubeGridTool.h
+++ b/Engine/Plugins/Experimental/MeshModelingToolsetExp/Source/MeshModelingToolsExp/Public/CubeGridTool.h
@@ -92,7 +92,7 @@ public:
 		UIMin = "0", ClampMin = "0"))
 	int32 BlocksPerStep = 1;
 
-	/** Determines cube grid scale. Can also be adjusted with Ctrl+E/Q. */
+	/** Determines cube grid scale. Can also be adjusted with Ctrl + E/Q. */
 	UPROPERTY(EditAnywhere, Category = Options, meta = (
 		EditCondition = "bAllowedToEditGrid", HideEditConditionToggle,
 		UIMin = "0", UIMax = "10", ClampMin = "0", ClampMax = "31"))
@@ -139,24 +139,24 @@ public:
 	FString ToggleCornerMode = TEXT("Z to start/complete corner mode.");
 
 	UPROPERTY(VisibleAnywhere, Category = ShortcutInfo)
-	FString PushPull = TEXT("E and Q push and pull, or use Ctrl+drag.");
+	FString PushPull = TEXT("E/Q to pull/push, or use Ctrl+drag.");
 
 	UPROPERTY(VisibleAnywhere, Category = ShortcutInfo)
-	FString ResizeGrid = TEXT("Ctrl + E/Q resize the grid in powers of two.");
+	FString ResizeGrid = TEXT("Ctrl + E/Q to increase/decrease grid in powers of two.");
 
 	UPROPERTY(VisibleAnywhere, Category = ShortcutInfo)
 	FString SlideSelection = TEXT("Middle mouse drag to slide selection in plane. "
-		"Shift + E/Q shift selection forward/back.");
+		"Shift + E/Q to shift selection back/forward.");
 
 	UPROPERTY(VisibleAnywhere, Category = ShortcutInfo)
-	FString FlipSelection = TEXT("T flips the selection.");
+	FString FlipSelection = TEXT("T to flip the selection.");
 
 	UPROPERTY(VisibleAnywhere, Category = ShortcutInfo)
-	FString GridGizmo = TEXT("R shows/hides grid gizmo. ");
+	FString GridGizmo = TEXT("R to show/hide grid gizmo.");
 
 	UPROPERTY(VisibleAnywhere, Category = ShortcutInfo)
-	FString QuickShiftGizmo = TEXT("Ctrl + middle click quick-repositions "
-		"the gizmo while keeping it on grid");
+	FString QuickShiftGizmo = TEXT("Ctrl + middle click to quick-reposition "
+		"the gizmo while keeping it on grid.");
 
 	UPROPERTY(VisibleAnywhere, Category = ShortcutInfo)
 	FString AlignGizmo = TEXT("While dragging gizmo handles, hold Ctrl to align "
@@ -199,31 +199,31 @@ public:
 
 	void PostAction(ECubeGridToolAction Action);
 
+	/** Can also be invoked with E. */
+	UFUNCTION(CallInEditor, Category = Actions, meta = (DisplayPriority = 1))
+	void Pull() { PostAction(ECubeGridToolAction::Pull); }
+
+	/** Can also be invoked with Q. */
+	UFUNCTION(CallInEditor, Category = Actions, meta = (DisplayPriority = 2))
+	void Push() { PostAction(ECubeGridToolAction::Push); }
+
+	/** Can also be invoked with Shift + E. */
+	UFUNCTION(CallInEditor, Category = Actions, meta = (DisplayPriority = 3))
+	void SlideBack() { PostAction(ECubeGridToolAction::SlideBack); }
+
+	/** Can also be invoked with Shift + Q. */
+	UFUNCTION(CallInEditor, Category = Actions, meta = (DisplayPriority = 4))
+	void SlideForward() { PostAction(ECubeGridToolAction::SlideForward); }
+
 	/** Engages a mode where specific corners can be selected to push/pull only
 	 those corners. Press Apply to commit the result afterward. Can also be toggled
 	 with Z. */
-	UFUNCTION(CallInEditor, Category = Actions)
+	UFUNCTION(CallInEditor, Category = Actions, meta = (DisplayPriority = 5))
 	void CornerMode() { PostAction(ECubeGridToolAction::CornerMode); }
-	
-	/** Can also be invoked with E */
-	UFUNCTION(CallInEditor, Category = Actions)
-	void Push() { PostAction(ECubeGridToolAction::Push); }
-
-	/** Can aslo be invoked with Q */
-	UFUNCTION(CallInEditor, Category = Actions)
-	void Pull() { PostAction(ECubeGridToolAction::Pull); }
 
 	/** Can also be invoked with T. */
-	UFUNCTION(CallInEditor, Category = Actions)
+	UFUNCTION(CallInEditor, Category = Actions, meta = (DisplayPriority = 6))
 	void Flip() { PostAction(ECubeGridToolAction::Flip); }
-
-	/** Can also be invoked with Shift + E. */
-	UFUNCTION(CallInEditor, Category = Actions)
-	void SlideForward() { PostAction(ECubeGridToolAction::SlideForward); }
-
-	/** Can also be invoked with Shift + Q. */
-	UFUNCTION(CallInEditor, Category = Actions)
-	void SlideBack() { PostAction(ECubeGridToolAction::SlideBack); }
 };
 
 UCLASS()
