ModelingTools: restore missing parameters in DynaSculpt Tool. No longer re-using the USculptBrushProperties from the base sculpt tool as it now has a DetailsCustomization that is not appropriate for DSculpt (which needs to be ported to base-sculpt-tool eventually)
#rnx

diff --git a/Engine/Plugins/Experimental/MeshModelingToolsetExp/Source/MeshModelingToolsExp/Private/DynamicMeshSculptTool.cpp b/Engine/Plugins/Experimental/MeshModelingToolsetExp/Source/MeshModelingToolsExp/Private/DynamicMeshSculptTool.cpp
--- a/Engine/Plugins/Experimental/MeshModelingToolsetExp/Source/MeshModelingToolsExp/Private/DynamicMeshSculptTool.cpp
+++ b/Engine/Plugins/Experimental/MeshModelingToolsetExp/Source/MeshModelingToolsExp/Private/DynamicMeshSculptTool.cpp
@@ -145,16 +145,13 @@ void UDynamicMeshSculptTool::Setup()
 	// initialize brush radius range interval, brush properties
 	double MaxDimension = DynamicMeshComponent->GetMesh()->GetBounds(true).MaxDim();
 	BrushRelativeSizeRange = FInterval1d(MaxDimension*0.01, MaxDimension);
-	BrushProperties = NewObject<USculptBrushProperties>(this);
-	BrushProperties->bShowFlowRate = false;
-	BrushProperties->bShowSpacing = false;
-	BrushProperties->bShowLazyness = false;
+	BrushProperties = NewObject<UDynamicMeshBrushProperties>(this);
 	BrushProperties->BrushSize.InitializeWorldSizeRange(
 		TInterval<float>((float)BrushRelativeSizeRange.Min, (float)BrushRelativeSizeRange.Max));
 	CalculateBrushRadius();
 
 	// initialize other properties
-	SculptProperties = NewObject<UBrushSculptProperties>(this);
+	SculptProperties = NewObject<UDynamicMeshBrushSculptProperties>(this);
 	KelvinBrushProperties = NewObject<UKelvinBrushProperties>(this);
 
 	RemeshProperties = NewObject<UBrushRemeshProperties>(this);
diff --git a/Engine/Plugins/Experimental/MeshModelingToolsetExp/Source/MeshModelingToolsExp/Public/DynamicMeshSculptTool.h b/Engine/Plugins/Experimental/MeshModelingToolsetExp/Source/MeshModelingToolsExp/Public/DynamicMeshSculptTool.h
--- a/Engine/Plugins/Experimental/MeshModelingToolsetExp/Source/MeshModelingToolsExp/Public/DynamicMeshSculptTool.h
+++ b/Engine/Plugins/Experimental/MeshModelingToolsetExp/Source/MeshModelingToolsExp/Public/DynamicMeshSculptTool.h
@@ -113,8 +113,34 @@ public:
 
 
 
+
+UCLASS()
+class MESHMODELINGTOOLSEXP_API UDynamicMeshBrushProperties : public UInteractiveToolPropertySet
+{
+	GENERATED_BODY()
+public:
+	UPROPERTY(EditAnywhere, Category = Brush, meta = (DisplayPriority = 1))
+	FBrushToolRadius BrushSize;
+
+	/** Amount of falloff to apply (0.0 - 1.0) */
+	UPROPERTY(EditAnywhere, Category = Brush, meta = (DisplayName = "Falloff", UIMin = "0.0", UIMax = "1.0", ClampMin = "0.0", ClampMax = "1.0", DisplayPriority = 3))
+	float BrushFalloffAmount = 0.5f;
+
+	/** Depth of Brush into surface along view ray or surface normal, depending on the Active Brush Type */
+	UPROPERTY(EditAnywhere, Category = Brush, meta = (UIMin = "-0.5", UIMax = "0.5", ClampMin = "-1.0", ClampMax = "1.0", DisplayPriority = 5))
+	float Depth = 0;
+
+	/** Allow the Brush to hit the back-side of the mesh */
+	UPROPERTY(EditAnywhere, Category = Brush, meta = (DisplayPriority = 6))
+	bool bHitBackFaces = true;
+};
+
+
+
+
+
 UCLASS()
-class MESHMODELINGTOOLSEXP_API UBrushSculptProperties : public UInteractiveToolPropertySet
+class MESHMODELINGTOOLSEXP_API UDynamicMeshBrushSculptProperties : public UInteractiveToolPropertySet
 {
 	GENERATED_BODY()
 
@@ -257,11 +283,11 @@ public:
 public:
 	/** Properties that control brush size/etc*/
 	UPROPERTY()
-	TObjectPtr<USculptBrushProperties> BrushProperties;
+	TObjectPtr<UDynamicMeshBrushProperties> BrushProperties;
 
 	/** Properties that control sculpting*/
 	UPROPERTY()
-	TObjectPtr<UBrushSculptProperties> SculptProperties;
+	TObjectPtr<UDynamicMeshBrushSculptProperties> SculptProperties;
 
 	UPROPERTY()
 	TObjectPtr<USculptMaxBrushProperties> SculptMaxBrushProperties;
