MeshModelingToolset: Fix TangentsTool not showing deviating tangents when comparison with MikkT gets enabled for the first time

diff --git a/Engine/Plugins/Experimental/MeshModelingToolsetExp/Source/MeshModelingToolsEditorOnly/Private/MeshTangentsTool.cpp b/Engine/Plugins/Experimental/MeshModelingToolsetExp/Source/MeshModelingToolsEditorOnly/Private/MeshTangentsTool.cpp
--- a/Engine/Plugins/Experimental/MeshModelingToolsetExp/Source/MeshModelingToolsEditorOnly/Private/MeshTangentsTool.cpp
+++ b/Engine/Plugins/Experimental/MeshModelingToolsetExp/Source/MeshModelingToolsEditorOnly/Private/MeshTangentsTool.cpp
@@ -91,8 +91,9 @@ void UMeshTangentsTool::Setup()
 	Settings->WatchProperty(Settings->CalculationMethod, [this](EMeshTangentsType) { Compute->InvalidateResult(); });
 	Settings->WatchProperty(Settings->LineLength, [this](float) { bLengthDirty = true; });
 	Settings->WatchProperty(Settings->LineThickness, [this](float) { bThicknessDirty = true; });
-	Settings->WatchProperty(Settings->bShowTangents, [this](float) { bVisibilityChanged = true; });
-	Settings->WatchProperty(Settings->bShowNormals, [this](float) { bVisibilityChanged = true; });
+	Settings->WatchProperty(Settings->bShowTangents, [this](bool) { bVisibilityChanged = true; });
+	Settings->WatchProperty(Settings->bShowNormals, [this](bool) { bVisibilityChanged = true; });
+	Settings->WatchProperty(Settings->bCompareWithMikkt, [this](bool) { ComputeMikkTDeviations(nullptr); });
 
 	PreviewGeometry = NewObject<UPreviewGeometry>(this);
 	PreviewGeometry->CreateInWorld(TargetActor->GetWorld(), PreviewMesh->GetTransform());
@@ -240,13 +241,7 @@ void UMeshTangentsTool::OnTangentsUpdated(const TUniquePtr<FMeshTangentsd>& NewR
 	const float LineLength = Settings->LineLength;
 	const float Thickness = Settings->LineThickness;
 
-	TSet<int32> DegenerateTris;
-	if (!Settings->bShowDegenerates || (Settings->bCompareWithMikkt && Settings->CalculationMethod != EMeshTangentsType::MikkTSpace) )
-	{
-		FMeshTangentsd DegenTangents(InputMesh.Get());
-		DegenTangents.ComputeTriangleTangents(InputMesh->Attributes()->GetUVLayer(0));
-		DegenerateTris = TSet<int32>(DegenTangents.GetDegenerateTris());
-	}
+	const TSet<int32> DegenerateTris = ComputeDegenerateTris();
 
 	// update Tangents rendering line set
 	PreviewGeometry->CreateOrUpdateLineSet(TEXT("Tangents"), InputMesh->MaxTriangleID(),
@@ -283,11 +278,44 @@ void UMeshTangentsTool::OnTangentsUpdated(const TUniquePtr<FMeshTangentsd>& NewR
 		}
 	}, 1);
 
+	ComputeMikkTDeviations(&DegenerateTris);
+
+	PreviewMesh->DeferredEditMesh([&](FDynamicMesh3& EditMesh)
+	{
+		NewResult.Get()->CopyToOverlays(EditMesh);
+	}, false);
+	PreviewMesh->NotifyDeferredEditCompleted(UPreviewMesh::ERenderUpdateMode::FastUpdate, EMeshRenderAttributeFlags::VertexNormals, false);
+
+	UpdateVisualization(false, false);
+}
+
+
+TSet<int32> UMeshTangentsTool::ComputeDegenerateTris() const
+{
+	TSet<int32> DegenerateTris;
+	if (!Settings->bShowDegenerates || (Settings->bCompareWithMikkt && Settings->CalculationMethod != EMeshTangentsType::MikkTSpace))
+	{
+		FMeshTangentsd DegenTangents(InputMesh.Get());
+		DegenTangents.ComputeTriangleTangents(InputMesh->Attributes()->GetUVLayer(0));
+		DegenerateTris = TSet<int32>(DegenTangents.GetDegenerateTris());
+	}
+	return DegenerateTris;
+}
+
 
+void UMeshTangentsTool::ComputeMikkTDeviations(const TSet<int32>* DegenerateTris)
+{
 	// calculate deviation between what we have and MikkT, if necessary
 	Deviations.Reset();
 	if (Settings->bCompareWithMikkt && Settings->CalculationMethod != EMeshTangentsType::MikkTSpace)
 	{
+		TSet<int32> TempDegenerateTris;
+		if (DegenerateTris == nullptr)
+		{
+			TempDegenerateTris = ComputeDegenerateTris();
+			DegenerateTris = &TempDegenerateTris;
+		}
+
 		FProgressCancel TmpCancel;
 		FCalculateTangentsOp MikktOp;
 		MikktOp.SourceMesh = InputMesh;
@@ -303,7 +331,7 @@ void UMeshTangentsTool::OnTangentsUpdated(const TUniquePtr<FMeshTangentsd>& NewR
 
 		for (int32 Index : InputMesh->TriangleIndicesItr())
 		{
-			if (DegenerateTris.Contains(Index) == false)
+			if (DegenerateTris->Contains(Index) == false)
 			{
 				FVector3d Verts[3];
 				InputMesh->GetTriVertices(Index, Verts[0], Verts[1], Verts[2]);
@@ -329,14 +357,6 @@ void UMeshTangentsTool::OnTangentsUpdated(const TUniquePtr<FMeshTangentsd>& NewR
 			}
 		}
 	}
-
-	PreviewMesh->DeferredEditMesh([&](FDynamicMesh3& EditMesh)
-	{
-		NewResult.Get()->CopyToOverlays(EditMesh);
-	}, false);
-	PreviewMesh->NotifyDeferredEditCompleted(UPreviewMesh::ERenderUpdateMode::FastUpdate, EMeshRenderAttributeFlags::VertexNormals, false);
-
-	UpdateVisualization(false, false);
 }
 
 
diff --git a/Engine/Plugins/Experimental/MeshModelingToolsetExp/Source/MeshModelingToolsEditorOnly/Public/MeshTangentsTool.h b/Engine/Plugins/Experimental/MeshModelingToolsetExp/Source/MeshModelingToolsEditorOnly/Public/MeshTangentsTool.h
--- a/Engine/Plugins/Experimental/MeshModelingToolsetExp/Source/MeshModelingToolsEditorOnly/Public/MeshTangentsTool.h
+++ b/Engine/Plugins/Experimental/MeshModelingToolsetExp/Source/MeshModelingToolsEditorOnly/Public/MeshTangentsTool.h
@@ -133,6 +133,8 @@ protected:
 
 	void OnTangentsUpdated(const TUniquePtr<UE::Geometry::FMeshTangentsd>& NewResult);
 	void UpdateVisualization(bool bThicknessChanged, bool bLengthChanged);
+	TSet<int32> ComputeDegenerateTris() const;
+	void ComputeMikkTDeviations(const TSet<int32>* DegenerateTris);
 
 	struct FMikktDeviation
 	{
