MeshModelingToolset: Add TextureResolution property for AutoUV tool in UVAtlas mode

diff --git a/Engine/Plugins/Experimental/MeshModelingToolsetExp/Source/MeshModelingToolsEditorOnly/Private/ParameterizeMeshTool.cpp b/Engine/Plugins/Experimental/MeshModelingToolsetExp/Source/MeshModelingToolsEditorOnly/Private/ParameterizeMeshTool.cpp
--- a/Engine/Plugins/Experimental/MeshModelingToolsetExp/Source/MeshModelingToolsEditorOnly/Private/ParameterizeMeshTool.cpp
+++ b/Engine/Plugins/Experimental/MeshModelingToolsetExp/Source/MeshModelingToolsEditorOnly/Private/ParameterizeMeshTool.cpp
@@ -217,6 +217,10 @@ TUniquePtr<FDynamicMeshOperator> UParameterizeMeshTool::MakeNewOperator()
 	// uvatlas options
 	ParameterizeMeshOp->Stretch   = UVAtlasProperties->IslandStretch;
 	ParameterizeMeshOp->NumCharts = UVAtlasProperties->NumIslands;
+	if (Settings->Method == EParameterizeMeshUVMethod::UVAtlas)
+	{
+		ParameterizeMeshOp->Width = ParameterizeMeshOp->Height = UVAtlasProperties->TextureResolution;
+	}
 	
 	// xatlas options
 	ParameterizeMeshOp->XAtlasMaxIterations = XAtlasProperties->MaxIterations;
@@ -234,8 +238,10 @@ TUniquePtr<FDynamicMeshOperator> UParameterizeMeshTool::MakeNewOperator()
 	ParameterizeMeshOp->ExpMapNormalSmoothingSteps = PatchBuilderProperties->SmoothingSteps;
 	ParameterizeMeshOp->ExpMapNormalSmoothingAlpha = PatchBuilderProperties->SmoothingAlpha;
 	ParameterizeMeshOp->bEnablePacking = PatchBuilderProperties->bRepack;
-	ParameterizeMeshOp->Width = ParameterizeMeshOp->Height = PatchBuilderProperties->TextureResolution;
-
+	if (Settings->Method == EParameterizeMeshUVMethod::PatchBuilder)
+	{
+		ParameterizeMeshOp->Width = ParameterizeMeshOp->Height = PatchBuilderProperties->TextureResolution;
+	}
 
 	ParameterizeMeshOp->SetTransform(UE::ToolTarget::GetLocalToWorldTransform(Target));
 
diff --git a/Engine/Plugins/Runtime/MeshModelingToolset/Source/ModelingOperatorsEditorOnly/Private/ParameterizationOps/ParameterizeMeshOp.cpp b/Engine/Plugins/Runtime/MeshModelingToolset/Source/ModelingOperatorsEditorOnly/Private/ParameterizationOps/ParameterizeMeshOp.cpp
--- a/Engine/Plugins/Runtime/MeshModelingToolset/Source/ModelingOperatorsEditorOnly/Private/ParameterizationOps/ParameterizeMeshOp.cpp
+++ b/Engine/Plugins/Runtime/MeshModelingToolset/Source/ModelingOperatorsEditorOnly/Private/ParameterizationOps/ParameterizeMeshOp.cpp
@@ -474,6 +474,8 @@ TUniquePtr<FDynamicMeshOperator> UParameterizeMeshOperatorFactory::MakeNewOperat
 		Op->Method = EParamOpBackend::UVAtlas;		
 		Op->Stretch = UVAtlasProperties->IslandStretch;
 		Op->NumCharts = UVAtlasProperties->NumIslands;
+		Op->Width = UVAtlasProperties->TextureResolution;
+		Op->Height = UVAtlasProperties->TextureResolution;
 		break;
 	case EParameterizeMeshUVMethod::XAtlas:
 		Op->Method = EParamOpBackend::XAtlas;
diff --git a/Engine/Plugins/Runtime/MeshModelingToolset/Source/ModelingOperatorsEditorOnly/Public/Properties/ParameterizeMeshProperties.h b/Engine/Plugins/Runtime/MeshModelingToolset/Source/ModelingOperatorsEditorOnly/Public/Properties/ParameterizeMeshProperties.h
--- a/Engine/Plugins/Runtime/MeshModelingToolset/Source/ModelingOperatorsEditorOnly/Public/Properties/ParameterizeMeshProperties.h
+++ b/Engine/Plugins/Runtime/MeshModelingToolset/Source/ModelingOperatorsEditorOnly/Public/Properties/ParameterizeMeshProperties.h
@@ -47,6 +47,10 @@ public:
 	/** Hint at number of UV islands. The default of 0 means it is determined automatically. */
 	UPROPERTY(EditAnywhere, Category = UVAtlas, meta = (UIMin = "0", UIMax = "100", ClampMin = "0", ClampMax = "10000"))
 	int NumIslands = 0;
+
+	/** Expected resolution of the output textures; this controls spacing left between UV islands to avoid interpolation artifacts. */
+	UPROPERTY(EditAnywhere, Category = UVAtlas, meta = (UIMin = "64", UIMax = "2048", ClampMin = "2", ClampMax = "4096"))
+	int TextureResolution = 1024;
 };
 
 
