ModelingTools: handle unset normals in Inspector Tool
#rnx

diff --git a/Engine/Plugins/Experimental/MeshModelingToolsetExp/Source/MeshModelingToolsExp/Private/MeshInspectorTool.cpp b/Engine/Plugins/Experimental/MeshModelingToolsetExp/Source/MeshModelingToolsExp/Private/MeshInspectorTool.cpp
--- a/Engine/Plugins/Experimental/MeshModelingToolsetExp/Source/MeshModelingToolsExp/Private/MeshInspectorTool.cpp
+++ b/Engine/Plugins/Experimental/MeshModelingToolsetExp/Source/MeshModelingToolsExp/Private/MeshInspectorTool.cpp
@@ -333,12 +333,15 @@ void UMeshInspectorTool::UpdateVisualization()
 		FVector3f TriN[3];
 		for (int tid : TargetMesh->TriangleIndicesItr())
 		{
-			TargetMesh->GetTriVertices(tid, TriV[0], TriV[1], TriV[2]);
-			NormalOverlay->GetTriElements(tid, TriN[0], TriN[1], TriN[2]);
-			for (int j = 0; j < 3; ++j)
+			if ( NormalOverlay->IsSetTriangle(tid) )
 			{
-				DrawnLineSet->AddLine((FVector)TriV[j], (FVector)((FVector3f)TriV[j] + NormalScaling * TriN[j]),
-					NormalColor, NormalThickness, NormalDepthBias);
+				TargetMesh->GetTriVertices(tid, TriV[0], TriV[1], TriV[2]);
+				NormalOverlay->GetTriElements(tid, TriN[0], TriN[1], TriN[2]);
+				for (int j = 0; j < 3; ++j)
+				{
+					DrawnLineSet->AddLine((FVector)TriV[j], (FVector)((FVector3f)TriV[j] + NormalScaling * TriN[j]),
+						NormalColor, NormalThickness, NormalDepthBias);
+				}
 			}
 		}
 	}
