MeshModelingTools: Make group paint tool initialize the current group to an unused group rather than potentially using settings from a different mesh. Don't keep allocating new groups if the current one is not used. Update revert-to-default arrow for consistency.

#rnx

diff --git a/Engine/Plugins/Experimental/MeshModelingToolsetExp/Source/MeshModelingToolsExp/Private/MeshGroupPaintTool.cpp b/Engine/Plugins/Experimental/MeshModelingToolsetExp/Source/MeshModelingToolsExp/Private/MeshGroupPaintTool.cpp
--- a/Engine/Plugins/Experimental/MeshModelingToolsetExp/Source/MeshModelingToolsExp/Private/MeshGroupPaintTool.cpp
+++ b/Engine/Plugins/Experimental/MeshModelingToolsetExp/Source/MeshModelingToolsExp/Private/MeshGroupPaintTool.cpp
@@ -163,6 +163,7 @@ void UMeshGroupPaintTool::Setup()
 	FilterProperties->RestoreProperties(this);
 	FilterProperties->BrushSize = UMeshSculptToolBase::BrushProperties->BrushSize.AdaptiveSize;
 	FilterProperties->bHitBackFaces = UMeshSculptToolBase::BrushProperties->bHitBackFaces;
+	FilterProperties->SetGroup = ActiveGroupSet->MaxGroupID;
 	AddToolPropertySource(FilterProperties);
 
 	InitializeIndicator();
@@ -1124,6 +1125,7 @@ void UMeshGroupPaintTool::OnTick(float DeltaTime)
 				if (bPendingPickGroup)
 				{
 					FilterProperties->SetGroup = HitGroupID;
+					NotifyOfPropertyChangeByTool(FilterProperties);
 				}
 				else if (bPendingToggleFreezeGroup)
 				{
@@ -1186,7 +1188,8 @@ void UMeshGroupPaintTool::OnTick(float DeltaTime)
 
 void UMeshGroupPaintTool::AllocateNewGroupAndSetAsCurrentAction()
 {
-	FilterProperties->SetGroup = ActiveGroupSet->AllocateNewGroupID();
+	FilterProperties->SetGroup = ActiveGroupSet->MaxGroupID;
+	NotifyOfPropertyChangeByTool(FilterProperties);
 }
 
 
