MeshModelingTools: Make details panel respond properly to hotkeys in Transform tool.

#rnx

diff --git a/Engine/Plugins/Experimental/MeshModelingToolsetExp/Source/MeshModelingToolsExp/Private/TransformMeshesTool.cpp b/Engine/Plugins/Experimental/MeshModelingToolsetExp/Source/MeshModelingToolsExp/Private/TransformMeshesTool.cpp
--- a/Engine/Plugins/Experimental/MeshModelingToolsetExp/Source/MeshModelingToolsExp/Private/TransformMeshesTool.cpp
+++ b/Engine/Plugins/Experimental/MeshModelingToolsetExp/Source/MeshModelingToolsExp/Private/TransformMeshesTool.cpp
@@ -135,21 +135,30 @@ void UTransformMeshesTool::RegisterActions(FInteractiveToolActionSet& ActionSet)
 		LOCTEXT("TransformToggleSetPivot", "Toggle Set Pivot"),
 		LOCTEXT("TransformToggleSetPivotTooltip", "Toggle Set Pivot on and off"),
 		EModifierKey::None, EKeys::S,
-		[this]() { this->TransformProps->bSetPivotMode = !this->TransformProps->bSetPivotMode; });
+		[this]() { 
+			TransformProps->bSetPivotMode = !TransformProps->bSetPivotMode; 
+			NotifyOfPropertyChangeByTool(TransformProps);
+		});
 
 	ActionSet.RegisterAction(this, (int32)EStandardToolActions::BaseClientDefinedActionID + 2,
 		TEXT("ToggleSnapDrag"),
 		LOCTEXT("TransformToggleSnapDrag", "Toggle SnapDrag"),
 		LOCTEXT("TransformToggleSnapDragTooltip", "Toggle SnapDrag on and off"),
 		EModifierKey::None, EKeys::D,
-		[this]() { this->TransformProps->bEnableSnapDragging = !this->TransformProps->bEnableSnapDragging; });
+		[this]() { 
+			TransformProps->bEnableSnapDragging = !TransformProps->bEnableSnapDragging;
+			NotifyOfPropertyChangeByTool(TransformProps);
+		});
 
 	ActionSet.RegisterAction(this, (int32)EStandardToolActions::BaseClientDefinedActionID + 3,
 		TEXT("CycleTransformMode"),
 		LOCTEXT("TransformCycleTransformMode", "Next Transform Mode"),
 		LOCTEXT("TransformCycleTransformModeTooltip", "Cycle through available Transform Modes"),
 		EModifierKey::None, EKeys::A,
-		[this]() { this->TransformProps->TransformMode = (ETransformMeshesTransformMode)(((uint8)TransformProps->TransformMode+1) % (uint8)ETransformMeshesTransformMode::LastValue); });
+		[this]() { 
+			TransformProps->TransformMode = (ETransformMeshesTransformMode)(((uint8)TransformProps->TransformMode+1) % (uint8)ETransformMeshesTransformMode::LastValue);
+			NotifyOfPropertyChangeByTool(TransformProps);
+		});
 
 
 	ActionSet.RegisterAction(this, (int32)EStandardToolActions::BaseClientDefinedActionID + 4,
@@ -157,14 +166,20 @@ void UTransformMeshesTool::RegisterActions(FInteractiveToolActionSet& ActionSet)
 		LOCTEXT("TransformCycleSourceMode", "Next SnapDrag Source Mode"),
 		LOCTEXT("TransformCycleSourceModeTooltip", "Cycle through available SnapDrag Source Modes"),
 		EModifierKey::None, EKeys::W,
-		[this]() { this->TransformProps->SnapDragSource = (ETransformMeshesSnapDragSource)(((uint8)TransformProps->SnapDragSource + 1) % (uint8)ETransformMeshesSnapDragSource::LastValue); });
+		[this]() { 
+			TransformProps->SnapDragSource = (ETransformMeshesSnapDragSource)(((uint8)TransformProps->SnapDragSource + 1) % (uint8)ETransformMeshesSnapDragSource::LastValue);
+			NotifyOfPropertyChangeByTool(TransformProps);
+		});
 
 	ActionSet.RegisterAction(this, (int32)EStandardToolActions::BaseClientDefinedActionID + 5,
 		TEXT("CycleRotationMode"),
 		LOCTEXT("TransformCycleRotationMode", "Next SnapDrag Rotation Mode"),
 		LOCTEXT("TransformCycleRotationModeTooltip", "Cycle through available SnapDrag Rotation Modes"),
 		EModifierKey::None, EKeys::E,
-		[this]() { this->TransformProps->RotationMode = (ETransformMeshesSnapDragRotationMode)(((uint8)TransformProps->RotationMode + 1) % (uint8)ETransformMeshesSnapDragRotationMode::LastValue); });
+		[this]() { 
+			TransformProps->RotationMode = (ETransformMeshesSnapDragRotationMode)(((uint8)TransformProps->RotationMode + 1) % (uint8)ETransformMeshesSnapDragRotationMode::LastValue); 
+			NotifyOfPropertyChangeByTool(TransformProps);
+		});
 
 
 }
