MeshModelingTools: Made BSP conversion tool be more permissive in "Convert then Combine" path, added better error detection in that path, clarified some messages.

#rnx

diff --git a/Engine/Plugins/Experimental/MeshModelingToolsetExp/Source/MeshModelingToolsEditorOnly/Private/BspConversionTool.cpp b/Engine/Plugins/Experimental/MeshModelingToolsetExp/Source/MeshModelingToolsEditorOnly/Private/BspConversionTool.cpp
--- a/Engine/Plugins/Experimental/MeshModelingToolsetExp/Source/MeshModelingToolsEditorOnly/Private/BspConversionTool.cpp
+++ b/Engine/Plugins/Experimental/MeshModelingToolsetExp/Source/MeshModelingToolsEditorOnly/Private/BspConversionTool.cpp
@@ -32,7 +32,7 @@ using namespace UE::Geometry;
 // Forward declarations of local functions
 void ConvertBrushesToDynamicMesh(TArray<ABrush*>& BrushesToConvert, FDynamicMesh3& OutputMesh, 
 	TArray<UMaterialInterface*>& OutputMaterials);
-void ApplyStaticMeshBooleanOperation(
+bool ApplyDynamicMeshBooleanOperation(
 	FDynamicMesh3& MeshA, const UE::Geometry::FTransform3d& TransformA, const TArray<UMaterialInterface*>& MaterialsA,
 	FDynamicMesh3& MeshB, const UE::Geometry::FTransform3d& TransformB, const TArray<UMaterialInterface*>& MaterialsB,
 	FDynamicMesh3& OutputMesh, UE::Geometry::FTransform3d& OutputTransform, TArray<UMaterialInterface*>& OutputMaterials,
@@ -107,7 +107,7 @@ void UBspConversionTool::Setup()
 	SetToolDisplayName(LOCTEXT("ToolName", "Convert BSP"));
 	// Give a description to put in the side panel
 	GetToolManager()->DisplayMessage(
-		LOCTEXT("OnStartTool", "Convert geometry brushes (also known as BSP brushes) into a single static mesh. \"Convert then Combine\" first converts the individual brushes and performs boolean operations on the resulting meshes (this requires individual brushes to have manifold geometry. \"Combine then Convert\" is the old functionality, combining brushes, then converting them."),
+		LOCTEXT("OnStartTool", "Convert geometry brushes (also known as BSP brushes) into a single static mesh."),
 		EToolMessageLevel::UserNotification);
 
 	// We write out an empty warning message to make the sidebar look unchanged if we write out a warning message and then
@@ -440,17 +440,6 @@ bool UBspConversionTool::ConvertThenCombine(FText *ErrorMessage)
 			}
 		}
 		*OutputTransform = UE::Geometry::FTransform3d::Identity(); // Always identity when first converted
-
-		// Check validity for the purposes of boolean operations. We'll do this after each brush conversion.
-		if (!OutputMesh->CheckValidity(FDynamicMesh3::FValidityOptions(), EValidityCheckFailMode::ReturnOnly))
-		{
-			PreviewMesh->ClearPreview();
-			if (ErrorMessage)
-			{
-				*ErrorMessage = GetBrushGeometryErrorMessage(NextBrush);
-			}
-			return false;
-		}
 	}
 
 	// Now convert the other meshes one at a time and apply them using boolean mesh operations
@@ -498,21 +487,11 @@ bool UBspConversionTool::ConvertThenCombine(FText *ErrorMessage)
 			}
 		}
 
-		// See if converted geometry is valid
-		if (!NextMesh.CheckValidity(FDynamicMesh3::FValidityOptions(), EValidityCheckFailMode::ReturnOnly)) 
-		{
-			PreviewMesh->ClearPreview();
-			if (ErrorMessage)
-			{
-				*ErrorMessage = GetBrushGeometryErrorMessage(NextBrush);
-			}
-			return false;
-		}
-
 		// Apply the boolean operation
+		bool bSuccess = true;
 		if (IsSubtractiveBrush)
 		{
-			ApplyStaticMeshBooleanOperation(
+			bSuccess = ApplyDynamicMeshBooleanOperation(
 				*InputMesh, *InputTransform, *InputMaterials,
 				NextMesh, NextTransform, NextMaterials,
 				*OutputMesh, *OutputTransform, *OutputMaterials, Operation);
@@ -525,11 +504,21 @@ bool UBspConversionTool::ConvertThenCombine(FText *ErrorMessage)
 			// For union, we actually swap the order of meshes in hopes of better lining up
 			// with BSP brush priority in coplanar places (brushes added later have priority,
 			// whereas for our boolean operations, first mesh has priority)
-			ApplyStaticMeshBooleanOperation(
+			bSuccess = ApplyDynamicMeshBooleanOperation(
 				NextMesh, NextTransform, NextMaterials,
 				*InputMesh, *InputTransform, *InputMaterials,
 				*OutputMesh, *OutputTransform, *OutputMaterials, Operation);
 		}
+
+		if (!bSuccess)
+		{
+			PreviewMesh->ClearPreview();
+			if (ErrorMessage)
+			{
+				*ErrorMessage = GetBrushGeometryErrorMessage(NextBrush);
+			}
+			return false;
+		}
 	}//end converting other brushes
 
 	// If there was only a single additive brush, we should keep its pivot in the original location.
@@ -550,8 +539,12 @@ bool UBspConversionTool::ConvertThenCombine(FText *ErrorMessage)
 
 FText GetBrushGeometryErrorMessage(ABrush* Brush)
 {
-	return FText::Format(LOCTEXT("ConvertThenCombineInvalidGeometryError", "The brush \"{0}\" has invalid geometry to be used with mesh boolean operations, so the \"Convert, then Combine\" mode cannot be used. Try using \"Combine, then Convert\" to convert, then use MeshInspector to see problematic areas."),
-		FText::FromString(Brush->GetName()));
+	return FText::Format(LOCTEXT("ConvertThenCombineInvalidGeometryError", 
+		"Failed attempting the \"Convert, then Combine\" path while trying to compose brush \"{0}\" with "
+		"previous results. Try using \"Combine, then Convert\" to convert, then use MeshInspector to look "
+		"for problematic areas around that brush."),
+		// Brush->GetActorLabel is an editor-only call
+		FText::FromString(Brush->GetActorLabel()));
 }
 
 /** Uses our existing conversion functions to convert brushes to a single DynamicMesh. */
@@ -588,7 +581,7 @@ void ConvertBrushesToDynamicMesh(TArray<ABrush*>& BrushesToConvert, FDynamicMesh
 	Converter.Convert(&MeshDescription, OutputMesh);
 }
 
-void ApplyStaticMeshBooleanOperation(
+bool ApplyDynamicMeshBooleanOperation(
 	FDynamicMesh3& MeshA, const UE::Geometry::FTransform3d& TransformA, const TArray<UMaterialInterface*>& MaterialsA,
 	FDynamicMesh3& MeshB, const UE::Geometry::FTransform3d& TransformB, const TArray<UMaterialInterface*>& MaterialsB,
 	FDynamicMesh3& OutputMesh, UE::Geometry::FTransform3d& OutputTransform, TArray<UMaterialInterface*>& OutputMaterials,
@@ -642,8 +635,10 @@ void ApplyStaticMeshBooleanOperation(
 	// Perform the actual boolean operation.
 	FMeshBoolean BooleanOperation(&MeshA, TransformA, &MeshB, TransformB, &OutputMesh, Operation);
 	BooleanOperation.bSimplifyAlongNewEdges = true;
-	BooleanOperation.Compute();
+	bool bSuccess = BooleanOperation.Compute();
 	OutputTransform = BooleanOperation.ResultTransform;
+
+	return bSuccess;
 }
 
 
diff --git a/Engine/Plugins/Experimental/MeshModelingToolsetExp/Source/MeshModelingToolsEditorOnly/Public/BspConversionTool.h b/Engine/Plugins/Experimental/MeshModelingToolsetExp/Source/MeshModelingToolsEditorOnly/Public/BspConversionTool.h
--- a/Engine/Plugins/Experimental/MeshModelingToolsetExp/Source/MeshModelingToolsEditorOnly/Public/BspConversionTool.h
+++ b/Engine/Plugins/Experimental/MeshModelingToolsetExp/Source/MeshModelingToolsEditorOnly/Public/BspConversionTool.h
@@ -31,10 +31,10 @@ public:
 UENUM()
 enum class EBspConversionMode : uint8
 {
-	/** First converts the brushes to static meshes, then performs static mesh boolean operations. */
+	/** First converts the brushes to static meshes, then performs mesh boolean operations. */
 	ConvertFirst = 0 UMETA(DisplayName = "Convert, then Combine"),
 
-	/** First combines brushes, then converts result to static mesh. */
+	/** First combines brushes using brush CSG operations, then converts result to static mesh (legacy path). */
 	CombineFirst = 1 UMETA(DisplayName = "Combine, then Convert"),
 };
 
@@ -71,12 +71,12 @@ public:
 
 	/** Caches individual brush conversions in "convert then combine" mode during a single invocation of 
 	 the tool. Only useful if changing selections or properties after starting the tool. Cleared on tool shutdown. */
-	UPROPERTY(EditAnywhere, NonTransactional, Category = Options, meta = (EditCondition = "ConversionMode == EBspConversionMode::ConvertFirst"))
-		bool bCacheBrushes = true;
+	UPROPERTY(EditAnywhere, NonTransactional, Category = Options, AdvancedDisplay, meta = (EditCondition = "ConversionMode == EBspConversionMode::ConvertFirst"))
+	bool bCacheBrushes = true;
 
 	/** Determines whether a dynamic preview is shown. Note that this introduces non-background computations 
 	at each event that changes the result, rather than only performing a computation on Accept. */
-	UPROPERTY(EditAnywhere, NonTransactional, Category = PreviewOptions)
+	UPROPERTY(EditAnywhere, NonTransactional, Category = PreviewOptions, AdvancedDisplay)
 	bool bShowPreview = true;
 };
 
@@ -109,7 +109,7 @@ public:
 	UFUNCTION(CallInEditor, Category = SelectionOperations, Meta = (DisplayName = "Deselect Volumes", DisplayPriority = 2))
 	void DeselectVolumes() { PostAction(EBspConversionToolAction::DeselectVolumes); }
 
-	/** Deselect any currently selected brushes that would not be converted given current settings. */
+	/** Deselect any currently selected brushes that are not valid targets given current settings. */
 	UFUNCTION(CallInEditor, Category = SelectionOperations, Meta = (DisplayName = "Deselect Non-Valid", DisplayPriority = 3))
 	void DeselectNonValid() { PostAction(EBspConversionToolAction::DeselectNonValid); }
 };
