ModelingTools: Update ActiveMaterial by index rather than by name in MatEd to avoid switching out the ActiveMaterial when the material slot changes.

#rnx

diff --git a/Engine/Plugins/Experimental/MeshModelingToolsetExp/Source/MeshModelingToolsEditorOnly/Private/EditMeshMaterialsTool.cpp b/Engine/Plugins/Experimental/MeshModelingToolsetExp/Source/MeshModelingToolsEditorOnly/Private/EditMeshMaterialsTool.cpp
--- a/Engine/Plugins/Experimental/MeshModelingToolsetExp/Source/MeshModelingToolsEditorOnly/Private/EditMeshMaterialsTool.cpp
+++ b/Engine/Plugins/Experimental/MeshModelingToolsetExp/Source/MeshModelingToolsEditorOnly/Private/EditMeshMaterialsTool.cpp
@@ -32,8 +32,9 @@ void UEditMeshMaterialsEditActions::PostMaterialAction(EEditMeshMaterialsToolAct
 
 void UEditMeshMaterialsToolProperties::UpdateFromMaterialsList()
 {
+	const int32 ActiveMaterialIdx = GetSelectedMaterialIndex();
+	
 	MaterialNamesList.Reset();
-
 	for ( int32 k = 0; k < Materials.Num(); ++k)
 	{
 		UMaterialInterface* Mat = Materials[k];
@@ -48,15 +49,8 @@ void UEditMeshMaterialsToolProperties::UpdateFromMaterialsList()
 		return;
 	}
 
-	// udpate active material if it no longer exists
-	for (int32 k = 0; k < MaterialNamesList.Num(); ++k)
-	{
-		if (MaterialNamesList[k] == ActiveMaterial)
-		{
-			return;
-		}
-	}
-	ActiveMaterial = MaterialNamesList[0];
+	// update active material by index
+	ActiveMaterial = MaterialNamesList[ActiveMaterialIdx < MaterialNamesList.Num() ? ActiveMaterialIdx : 0];
 }
 
 int32 UEditMeshMaterialsToolProperties::GetSelectedMaterialIndex() const
