MeshModelingToolset: In Tangents tool enable "Compare with MikkT" only if FastMikkT is selected

diff --git a/Engine/Plugins/Experimental/MeshModelingToolsetExp/Source/MeshModelingToolsEditorOnly/Private/MeshTangentsTool.cpp b/Engine/Plugins/Experimental/MeshModelingToolsetExp/Source/MeshModelingToolsEditorOnly/Private/MeshTangentsTool.cpp
--- a/Engine/Plugins/Experimental/MeshModelingToolsetExp/Source/MeshModelingToolsEditorOnly/Private/MeshTangentsTool.cpp
+++ b/Engine/Plugins/Experimental/MeshModelingToolsetExp/Source/MeshModelingToolsEditorOnly/Private/MeshTangentsTool.cpp
@@ -307,7 +307,7 @@ void UMeshTangentsTool::ComputeMikkTDeviations(const TSet<int32>* DegenerateTris
 {
 	// calculate deviation between what we have and MikkT, if necessary
 	Deviations.Reset();
-	if (Settings->bCompareWithMikkt && Settings->CalculationMethod != EMeshTangentsType::MikkTSpace)
+	if (Settings->bCompareWithMikkt && Settings->CalculationMethod == EMeshTangentsType::FastMikkTSpace)
 	{
 		TSet<int32> TempDegenerateTris;
 		if (DegenerateTris == nullptr)
diff --git a/Engine/Plugins/Experimental/MeshModelingToolsetExp/Source/MeshModelingToolsEditorOnly/Public/MeshTangentsTool.h b/Engine/Plugins/Experimental/MeshModelingToolsetExp/Source/MeshModelingToolsEditorOnly/Public/MeshTangentsTool.h
--- a/Engine/Plugins/Experimental/MeshModelingToolsetExp/Source/MeshModelingToolsEditorOnly/Public/MeshTangentsTool.h
+++ b/Engine/Plugins/Experimental/MeshModelingToolsetExp/Source/MeshModelingToolsEditorOnly/Public/MeshTangentsTool.h
@@ -67,16 +67,16 @@ public:
 	UPROPERTY(EditAnywhere, Category = Display, AdvancedDisplay)
 	bool bShowDegenerates = false;
 
-	/** Display difference between the current tangents and tangents calculated with standard MikkTSpace calculation.
-	 * This is only available if a Calculation Method other than MikkTSpace is selected. */
+	/** Display difference between FastMikkTSpace tangents and MikkTSpace tangents.
+	 * This is only available if the FastMikkTSpace Calculation Method is selected. */
 	UPROPERTY(EditAnywhere, Category = Display, AdvancedDisplay, meta = (DisplayName = "Compare with MikkT",
-		EditCondition = "CalculationMethod != EMeshTangentsType::MikkTSpace"))
+		EditCondition = "CalculationMethod == EMeshTangentsType::FastMikkTSpace"))
 	bool bCompareWithMikkt = false;
 
-	/** Minimum angle difference in degrees for a tangent to be considered different to the standard MikkTSpace calculation
-	 * This is only available if a Compare with MikkT is enabled and a Calculation Method other than MikkTSpace is selected. */
+	/** Minimum angle difference in degrees for a FastMikkTSpace tangent to be considered different to a MikkTSpace tangent.
+	 * This is only available if a Compare with MikkT is enabled and the FastMikkTSpace Calculation Method is selected. */
 	UPROPERTY(EditAnywhere, Category = Display, AdvancedDisplay, meta = (DisplayName = "Compare Threshold", UIMin = "0.5", UIMax = "90.0",
-		EditCondition = "CalculationMethod != EMeshTangentsType::MikkTSpace && bCompareWithMikkt"))
+		EditCondition = "CalculationMethod == EMeshTangentsType::FastMikkTSpace && bCompareWithMikkt"))
 	float CompareWithMikktThreshold = 5.0f;
 
 };
