Approximate Actors - Corrected handling of the EMeshApproximationBaseCappingType::ConvexPolygon base capping option.

Using EMeshApproximationBaseCappingType::ConvexPolygon would always generate a solid rather than a polygon.

diff --git a/Engine/Plugins/Experimental/MeshModelingToolsetExp/Source/GeometryProcessingAdapters/Private/ApproximateActorsImpl.cpp b/Engine/Plugins/Experimental/MeshModelingToolsetExp/Source/GeometryProcessingAdapters/Private/ApproximateActorsImpl.cpp
--- a/Engine/Plugins/Experimental/MeshModelingToolsetExp/Source/GeometryProcessingAdapters/Private/ApproximateActorsImpl.cpp
+++ b/Engine/Plugins/Experimental/MeshModelingToolsetExp/Source/GeometryProcessingAdapters/Private/ApproximateActorsImpl.cpp
@@ -1015,8 +1015,14 @@ void FApproximateActorsImpl::GenerateApproximationForActorSet(const TArray<AActo
 	if (Options.BaseCappingPolicy != EBaseCappingPolicy::NoBaseCapping)
 	{
 		TRACE_CPUPROFILER_EVENT_SCOPE(ApproximateActorsImpl_Generate_Capping);
-		double UseThickness = (Options.BaseThicknessOverrideMeters != 0) ? (Options.BaseThicknessOverrideMeters * 100.0) :
-			(Options.bAutoThickenThinParts ? SceneBuildOptions.DesiredMinThickness : 1.25 * ApproxAccuracy);
+
+		double UseThickness = 0.0;
+		if (Options.BaseCappingPolicy == EBaseCappingPolicy::ConvexSolid)
+		{
+			UseThickness = (Options.BaseThicknessOverrideMeters != 0) ? (Options.BaseThicknessOverrideMeters * 100.0) :
+						   (Options.bAutoThickenThinParts ? SceneBuildOptions.DesiredMinThickness : 1.25 * ApproxAccuracy);
+		}
+
 		double UseHeight = (Options.BaseHeightOverrideMeters != 0) ? (Options.BaseHeightOverrideMeters * 100.0) : (2.0 * ApproxAccuracy);
 		Scene.GenerateBaseClosingMesh(UseHeight, UseThickness);
 	}
