ModelingTools: Fix crash in Jacketing tool when the op would result in an empty mesh for one of the source meshes.

#rnx

diff --git a/Engine/Plugins/Experimental/MeshModelingToolsetExp/Source/MeshModelingToolsExp/Private/RemoveOccludedTrianglesTool.cpp b/Engine/Plugins/Experimental/MeshModelingToolsetExp/Source/MeshModelingToolsExp/Private/RemoveOccludedTrianglesTool.cpp
--- a/Engine/Plugins/Experimental/MeshModelingToolsetExp/Source/MeshModelingToolsExp/Private/RemoveOccludedTrianglesTool.cpp
+++ b/Engine/Plugins/Experimental/MeshModelingToolsetExp/Source/MeshModelingToolsExp/Private/RemoveOccludedTrianglesTool.cpp
@@ -484,7 +484,8 @@ void URemoveOccludedTrianglesTool::GenerateAsset(const TArray<FDynamicMeshOpResu
 			bWantDestroy = false;
 		}
 	}
-	
+
+	FSelectedOjectsChangeList NewSelection;
 	for (int32 PreviewIdx = 0; PreviewIdx < Previews.Num(); PreviewIdx++)
 	{
 		check(Results[PreviewIdx].Mesh.Get() != nullptr);
@@ -506,6 +507,16 @@ void URemoveOccludedTrianglesTool::GenerateAsset(const TArray<FDynamicMeshOpResu
 		}
 
 		UE::ToolTarget::CommitMeshDescriptionUpdateViaDynamicMesh(Targets[ComponentIdx], *Results[PreviewIdx].Mesh, true);
+
+		NewSelection.Actors.Add(UE::ToolTarget::GetTargetActor(Targets[ComponentIdx]));
+	}
+
+	// If we destroyed component(s) for empty mesh(es), ensure we update the selection
+	// to avoid the details panels from crashing when updating on next engine tick.
+	if (bWantDestroy)
+	{
+		NewSelection.ModificationType = NewSelection.Actors.Num() > 0 ? ESelectedObjectsModificationType::Replace : ESelectedObjectsModificationType::Clear;
+		GetToolManager()->RequestSelectionChange(NewSelection);
 	}
 
 	GetToolManager()->EndUndoTransaction();
