Lattice deformer: set a max lattice resolution for each dimension

#rnx

diff --git a/Engine/Plugins/Experimental/MeshModelingToolsetExp/Source/MeshModelingToolsExp/Public/LatticeDeformerTool.h b/Engine/Plugins/Experimental/MeshModelingToolsetExp/Source/MeshModelingToolsExp/Public/LatticeDeformerTool.h
--- a/Engine/Plugins/Experimental/MeshModelingToolsetExp/Source/MeshModelingToolsExp/Public/LatticeDeformerTool.h
+++ b/Engine/Plugins/Experimental/MeshModelingToolsetExp/Source/MeshModelingToolsExp/Public/LatticeDeformerTool.h
@@ -55,15 +55,15 @@ public:
 	void PostAction(ELatticeDeformerToolAction Action);
 
 	/** Number of lattice vertices along the X axis */
-	UPROPERTY(EditAnywhere, Category = Resolution, meta = (UIMin = "2", ClampMin = "2", EditCondition = "bCanChangeResolution", HideEditConditionToggle))
+	UPROPERTY(EditAnywhere, Category = Resolution, meta = (UIMin = "2", ClampMin = "2", UIMax = "25", ClampMax = "40", EditCondition = "bCanChangeResolution", HideEditConditionToggle))
 	int XAxisResolution = 5;
 
 	/** Number of lattice vertices along the Y axis */
-	UPROPERTY(EditAnywhere, Category = Resolution, meta = (UIMin = "2", ClampMin = "2", EditCondition = "bCanChangeResolution", HideEditConditionToggle))
+	UPROPERTY(EditAnywhere, Category = Resolution, meta = (UIMin = "2", ClampMin = "2", UIMax = "25", ClampMax = "40", EditCondition = "bCanChangeResolution", HideEditConditionToggle))
 	int YAxisResolution = 5;
 
 	/** Number of lattice vertices along the Z axis */
-	UPROPERTY(EditAnywhere, Category = Resolution, meta = (UIMin = "2", ClampMin = "2", EditCondition = "bCanChangeResolution", HideEditConditionToggle))
+	UPROPERTY(EditAnywhere, Category = Resolution, meta = (UIMin = "2", ClampMin = "2", UIMax = "25", ClampMax = "40", EditCondition = "bCanChangeResolution", HideEditConditionToggle))
 	int ZAxisResolution = 5;
 
 	/** Relative distance the lattice extends from the mesh */
