ModelingTools: Fix crash when reloading levels during an AutoLOD compute

#rnx

diff --git a/Engine/Plugins/Experimental/HairModelingToolset/Source/HairModelingToolset/Private/GenerateLODMeshesTool.cpp b/Engine/Plugins/Experimental/HairModelingToolset/Source/HairModelingToolset/Private/GenerateLODMeshesTool.cpp
--- a/Engine/Plugins/Experimental/HairModelingToolset/Source/HairModelingToolset/Private/GenerateLODMeshesTool.cpp
+++ b/Engine/Plugins/Experimental/HairModelingToolset/Source/HairModelingToolset/Private/GenerateLODMeshesTool.cpp
@@ -194,7 +194,7 @@ void UGenerateLODMeshesTool::UpdateNumPreviews()
 		Factory->ParentTool = this;
 
 		UMeshOpPreviewWithBackgroundCompute* NewPreview = NewObject<UMeshOpPreviewWithBackgroundCompute>(this);
-		NewPreview->Setup(this->TargetWorld, Factory.Get());
+		NewPreview->Setup(GetTargetWorld(), Factory.Get());
 
 		FComponentMaterialSet MaterialSet;
 		Cast<IMaterialProvider>(Target)->GetMaterialSet(MaterialSet);
@@ -394,7 +394,7 @@ void UGenerateLODMeshesTool::GenerateAssets()
 		FString Name = FString::Printf( TEXT("%s_LOD%d"), *BaseName, (SimplifyProperties->NameIndexBase+k) );
 
 		FCreateMeshObjectParams NewMeshObjectParams;
-		NewMeshObjectParams.TargetWorld = TargetWorld;
+		NewMeshObjectParams.TargetWorld = GetTargetWorld();
 		NewMeshObjectParams.Transform = Transform;
 		NewMeshObjectParams.BaseName = Name;
 		NewMeshObjectParams.Materials = MaterialSet.Materials;
diff --git a/Engine/Plugins/Experimental/MeshLODToolset/Source/MeshLODToolset/Private/Tools/GenerateStaticMeshLODAssetTool.cpp b/Engine/Plugins/Experimental/MeshLODToolset/Source/MeshLODToolset/Private/Tools/GenerateStaticMeshLODAssetTool.cpp
--- a/Engine/Plugins/Experimental/MeshLODToolset/Source/MeshLODToolset/Private/Tools/GenerateStaticMeshLODAssetTool.cpp
+++ b/Engine/Plugins/Experimental/MeshLODToolset/Source/MeshLODToolset/Private/Tools/GenerateStaticMeshLODAssetTool.cpp
@@ -373,7 +373,7 @@ void UGenerateStaticMeshLODAssetTool::Setup()
 
 	this->OpFactory = MakeUnique<FGenerateStaticMeshLODAssetOperatorFactory>(this, PreviewTransform);
 	PreviewWithBackgroundCompute = NewObject<UMeshOpPreviewWithBackgroundCompute>(this, "Preview");
-	PreviewWithBackgroundCompute->Setup(this->TargetWorld, this->OpFactory.Get());
+	PreviewWithBackgroundCompute->Setup(GetTargetWorld(), this->OpFactory.Get());
 	ToolSetupUtil::ApplyRenderingConfigurationToPreview(PreviewWithBackgroundCompute->PreviewMesh, nullptr);
 	PreviewWithBackgroundCompute->PreviewMesh->SetTangentsMode(EDynamicMeshComponentTangentsMode::ExternallyProvided);
 
@@ -433,7 +433,7 @@ void UGenerateStaticMeshLODAssetTool::Setup()
 	CollisionVizSettings->WatchProperty(CollisionVizSettings->bShowHidden, [this](bool bNewValue) { bCollisionVisualizationDirty = true; });
 
 	CollisionPreview = NewObject<UPreviewGeometry>(this);
-	CollisionPreview->CreateInWorld(TargetWorld, (FTransform)PreviewTransform);
+	CollisionPreview->CreateInWorld(GetTargetWorld(), (FTransform)PreviewTransform);
 
 	// read Preset if we started Tool with one already selected
 	OnPresetSelectionChanged();		
@@ -495,6 +495,7 @@ void UGenerateStaticMeshLODAssetTool::Shutdown(EToolShutdownType ShutdownType)
 	if (PreviewWithBackgroundCompute)
 	{
 		PreviewWithBackgroundCompute->Shutdown();
+		PreviewWithBackgroundCompute = nullptr;
 	}
 
 	if (GenerateProcess)
diff --git a/Engine/Plugins/Experimental/MeshLODToolset/Source/MeshLODToolset/Private/Tools/LODManagerTool.cpp b/Engine/Plugins/Experimental/MeshLODToolset/Source/MeshLODToolset/Private/Tools/LODManagerTool.cpp
--- a/Engine/Plugins/Experimental/MeshLODToolset/Source/MeshLODToolset/Private/Tools/LODManagerTool.cpp
+++ b/Engine/Plugins/Experimental/MeshLODToolset/Source/MeshLODToolset/Private/Tools/LODManagerTool.cpp
@@ -92,12 +92,12 @@ void ULODManagerTool::Setup()
 	LODPreviewProperties->WatchProperty(LODPreviewProperties->bShowSeams, [this](bool bNewValue) { bPreviewLODValid = false; });
 
 	LODPreview = NewObject<UPreviewMesh>(this);
-	LODPreview->CreateInWorld(TargetWorld, (FTransform)UE::ToolTarget::GetLocalToWorldTransform(Targets[0]));
+	LODPreview->CreateInWorld(GetTargetWorld(), (FTransform)UE::ToolTarget::GetLocalToWorldTransform(Targets[0]));
 	LODPreview->SetTangentsMode(EDynamicMeshComponentTangentsMode::ExternallyProvided);
 	LODPreview->SetVisible(false);
 
 	LODPreviewLines = NewObject<UPreviewGeometry>(this);
-	LODPreviewLines->CreateInWorld(TargetWorld, (FTransform)UE::ToolTarget::GetLocalToWorldTransform(Targets[0]));
+	LODPreviewLines->CreateInWorld(GetTargetWorld(), (FTransform)UE::ToolTarget::GetLocalToWorldTransform(Targets[0]));
 
 	FComponentMaterialSet MaterialSet = UE::ToolTarget::GetMaterialSet(Targets[0]);
 	LODPreview->SetMaterials(MaterialSet.Materials);
diff --git a/Engine/Plugins/Experimental/MeshModelingToolsetExp/Source/MeshModelingToolsEditorOnly/Private/MergeMeshesTool.cpp b/Engine/Plugins/Experimental/MeshModelingToolsetExp/Source/MeshModelingToolsEditorOnly/Private/MergeMeshesTool.cpp
--- a/Engine/Plugins/Experimental/MeshModelingToolsetExp/Source/MeshModelingToolsEditorOnly/Private/MergeMeshesTool.cpp
+++ b/Engine/Plugins/Experimental/MeshModelingToolsetExp/Source/MeshModelingToolsEditorOnly/Private/MergeMeshesTool.cpp
@@ -85,7 +85,7 @@ void UMergeMeshesTool::Setup()
 
 	// initialize the PreviewMesh+BackgroundCompute object
 	Preview = NewObject<UMeshOpPreviewWithBackgroundCompute>(this, "Preview");
-	Preview->Setup(this->TargetWorld, this);
+	Preview->Setup(GetTargetWorld(), this);
 	Preview->OnMeshUpdated.AddLambda([this](UMeshOpPreviewWithBackgroundCompute* Compute) {
 		MeshStatisticsProperties->Update(*Compute->PreviewMesh->GetPreviewDynamicMesh());
 		UpdateAcceptWarnings(Compute->HaveEmptyResult() ? EAcceptWarning::EmptyForbidden : EAcceptWarning::NoWarning);
@@ -211,7 +211,7 @@ void UMergeMeshesTool::GenerateAsset(const FDynamicMeshOpResult& OpResult)
 	check(OpResult.Mesh.Get() != nullptr);
 	
 	FCreateMeshObjectParams NewMeshObjectParams;
-	NewMeshObjectParams.TargetWorld = TargetWorld;
+	NewMeshObjectParams.TargetWorld = GetTargetWorld();
 	NewMeshObjectParams.Transform = (FTransform)OpResult.Transform;
 	NewMeshObjectParams.BaseName = TEXT("MergedMesh");
 	NewMeshObjectParams.Materials.Add(ToolSetupUtil::GetDefaultMaterial());
diff --git a/Engine/Plugins/Experimental/MeshModelingToolsetExp/Source/MeshModelingToolsEditorOnly/Private/ParameterizeMeshTool.cpp b/Engine/Plugins/Experimental/MeshModelingToolsetExp/Source/MeshModelingToolsEditorOnly/Private/ParameterizeMeshTool.cpp
--- a/Engine/Plugins/Experimental/MeshModelingToolsetExp/Source/MeshModelingToolsEditorOnly/Private/ParameterizeMeshTool.cpp
+++ b/Engine/Plugins/Experimental/MeshModelingToolsetExp/Source/MeshModelingToolsEditorOnly/Private/ParameterizeMeshTool.cpp
@@ -38,7 +38,7 @@ void UParameterizeMeshTool::Setup()
 	FComponentMaterialSet MaterialSet = UE::ToolTarget::GetMaterialSet(Target);
 
 	Preview = NewObject<UMeshOpPreviewWithBackgroundCompute>(this);
-	Preview->Setup(this->TargetWorld, this);
+	Preview->Setup(GetTargetWorld(), this);
 	ToolSetupUtil::ApplyRenderingConfigurationToPreview(Preview->PreviewMesh, nullptr);
 	Preview->PreviewMesh->SetTangentsMode(EDynamicMeshComponentTangentsMode::AutoCalculated);
 	Preview->PreviewMesh->ReplaceMesh(*InputMesh);
@@ -104,7 +104,7 @@ void UParameterizeMeshTool::Setup()
 	if (bCreateUVLayoutViewOnSetup)
 	{
 		UVLayoutView = NewObject<UUVLayoutPreview>(this);
-		UVLayoutView->CreateInWorld(TargetWorld);
+		UVLayoutView->CreateInWorld(GetTargetWorld());
 		UVLayoutView->SetSourceMaterials(MaterialSet);
 		UVLayoutView->SetSourceWorldPosition(InputTransform, UE::ToolTarget::GetTargetActor(Target)->GetComponentsBoundingBox());
 		UVLayoutView->Settings->bEnabled = false;
diff --git a/Engine/Plugins/Experimental/MeshModelingToolsetExp/Source/MeshModelingToolsEditorOnly/Private/SimplifyMeshTool.cpp b/Engine/Plugins/Experimental/MeshModelingToolsetExp/Source/MeshModelingToolsEditorOnly/Private/SimplifyMeshTool.cpp
--- a/Engine/Plugins/Experimental/MeshModelingToolsetExp/Source/MeshModelingToolsEditorOnly/Private/SimplifyMeshTool.cpp
+++ b/Engine/Plugins/Experimental/MeshModelingToolsetExp/Source/MeshModelingToolsEditorOnly/Private/SimplifyMeshTool.cpp
@@ -67,7 +67,7 @@ void USimplifyMeshTool::Setup()
 	UE::ToolTarget::HideSourceObject(Target);
 
 	Preview = NewObject<UMeshOpPreviewWithBackgroundCompute>(this);
-	Preview->Setup(this->TargetWorld, this);
+	Preview->Setup(GetTargetWorld(), this);
 	ToolSetupUtil::ApplyRenderingConfigurationToPreview(Preview->PreviewMesh, nullptr);
 	FComponentMaterialSet MaterialSet = UE::ToolTarget::GetMaterialSet(Target);
 	Preview->ConfigureMaterials( MaterialSet.Materials,
diff --git a/Engine/Plugins/Experimental/MeshModelingToolsetExp/Source/MeshModelingToolsEditorOnly/Private/SubdividePolyTool.cpp b/Engine/Plugins/Experimental/MeshModelingToolsetExp/Source/MeshModelingToolsEditorOnly/Private/SubdividePolyTool.cpp
--- a/Engine/Plugins/Experimental/MeshModelingToolsetExp/Source/MeshModelingToolsEditorOnly/Private/SubdividePolyTool.cpp
+++ b/Engine/Plugins/Experimental/MeshModelingToolsetExp/Source/MeshModelingToolsEditorOnly/Private/SubdividePolyTool.cpp
@@ -199,7 +199,7 @@ void USubdividePolyTool::Setup()
 	{
 		return;
 	}
-	PreviewMesh->CreateInWorld(TargetWorld, FTransform::Identity);
+	PreviewMesh->CreateInWorld(GetTargetWorld(), FTransform::Identity);
 	ToolSetupUtil::ApplyRenderingConfigurationToPreview(PreviewMesh, nullptr);
 
 	PreviewMesh->SetTransform(TargetComponent->GetWorldTransform());
diff --git a/Engine/Plugins/Experimental/MeshModelingToolsetExp/Source/MeshModelingToolsEditorOnly/Private/VoxelCSGMeshesTool.cpp b/Engine/Plugins/Experimental/MeshModelingToolsetExp/Source/MeshModelingToolsEditorOnly/Private/VoxelCSGMeshesTool.cpp
--- a/Engine/Plugins/Experimental/MeshModelingToolsetExp/Source/MeshModelingToolsEditorOnly/Private/VoxelCSGMeshesTool.cpp
+++ b/Engine/Plugins/Experimental/MeshModelingToolsetExp/Source/MeshModelingToolsEditorOnly/Private/VoxelCSGMeshesTool.cpp
@@ -71,7 +71,7 @@ void UVoxelCSGMeshesTool::Setup()
 
 	// initialize the PreviewMesh+BackgroundCompute object
 	Preview = NewObject<UMeshOpPreviewWithBackgroundCompute>(this, "Preview");
-	Preview->Setup(this->TargetWorld, this);
+	Preview->Setup(GetTargetWorld(), this);
 	ToolSetupUtil::ApplyRenderingConfigurationToPreview(Preview->PreviewMesh, nullptr);
 	Preview->OnMeshUpdated.AddLambda([this](UMeshOpPreviewWithBackgroundCompute* Compute) {
 		MeshStatisticsProperties->Update(*Compute->PreviewMesh->GetPreviewDynamicMesh());
@@ -196,7 +196,7 @@ void UVoxelCSGMeshesTool::GenerateAsset(const FDynamicMeshOpResult& OpResult)
 	if (ensure(OpResult.Mesh.Get() != nullptr))
 	{
 		FCreateMeshObjectParams NewMeshObjectParams;
-		NewMeshObjectParams.TargetWorld = TargetWorld;
+		NewMeshObjectParams.TargetWorld = GetTargetWorld();
 		NewMeshObjectParams.Transform = (FTransform)OpResult.Transform;
 		NewMeshObjectParams.BaseName = TEXT("CSGMesh");
 		NewMeshObjectParams.Materials.Add(ToolSetupUtil::GetDefaultMaterial());
diff --git a/Engine/Plugins/Experimental/MeshModelingToolsetExp/Source/MeshModelingToolsExp/Private/BakeMeshAttributeMapsToolBase.cpp b/Engine/Plugins/Experimental/MeshModelingToolsetExp/Source/MeshModelingToolsExp/Private/BakeMeshAttributeMapsToolBase.cpp
--- a/Engine/Plugins/Experimental/MeshModelingToolsetExp/Source/MeshModelingToolsExp/Private/BakeMeshAttributeMapsToolBase.cpp
+++ b/Engine/Plugins/Experimental/MeshModelingToolsetExp/Source/MeshModelingToolsExp/Private/BakeMeshAttributeMapsToolBase.cpp
@@ -58,7 +58,7 @@ void UBakeMeshAttributeMapsToolBase::Setup()
 	const FDynamicMesh3 InputMesh = UE::ToolTarget::GetDynamicMeshCopy(Targets[0], true);
 	const UE::Geometry::FTransform3d BaseToWorld = UE::ToolTarget::GetLocalToWorldTransform(Targets[0]);
 	PreviewMesh = NewObject<UPreviewMesh>(this);
-	PreviewMesh->CreateInWorld(TargetWorld, FTransform::Identity);
+	PreviewMesh->CreateInWorld(GetTargetWorld(), FTransform::Identity);
 	ToolSetupUtil::ApplyRenderingConfigurationToPreview(PreviewMesh, nullptr);
 	PreviewMesh->SetTransform(static_cast<FTransform>(BaseToWorld));
 	PreviewMesh->SetTangentsMode(EDynamicMeshComponentTangentsMode::ExternallyProvided);
diff --git a/Engine/Plugins/Experimental/MeshModelingToolsetExp/Source/MeshModelingToolsExp/Private/BakeMeshAttributeVertexTool.cpp b/Engine/Plugins/Experimental/MeshModelingToolsetExp/Source/MeshModelingToolsExp/Private/BakeMeshAttributeVertexTool.cpp
--- a/Engine/Plugins/Experimental/MeshModelingToolsetExp/Source/MeshModelingToolsExp/Private/BakeMeshAttributeVertexTool.cpp
+++ b/Engine/Plugins/Experimental/MeshModelingToolsetExp/Source/MeshModelingToolsExp/Private/BakeMeshAttributeVertexTool.cpp
@@ -249,7 +249,7 @@ void UBakeMeshAttributeVertexTool::Setup()
 
 	const FDynamicMesh3 InputMeshWithTangents = UE::ToolTarget::GetDynamicMeshCopy(Targets[0], true);
 	PreviewMesh = NewObject<UPreviewMesh>(this);
-	PreviewMesh->CreateInWorld(TargetWorld, FTransform::Identity);
+	PreviewMesh->CreateInWorld(GetTargetWorld(), FTransform::Identity);
 	ToolSetupUtil::ApplyRenderingConfigurationToPreview(PreviewMesh, nullptr);
 	PreviewMesh->SetTransform(static_cast<FTransform>(UE::ToolTarget::GetLocalToWorldTransform(Targets[0])));
 	PreviewMesh->SetTangentsMode(EDynamicMeshComponentTangentsMode::ExternallyProvided);
diff --git a/Engine/Plugins/Experimental/MeshModelingToolsetExp/Source/MeshModelingToolsExp/Private/ConvertMeshesTool.cpp b/Engine/Plugins/Experimental/MeshModelingToolsetExp/Source/MeshModelingToolsExp/Private/ConvertMeshesTool.cpp
--- a/Engine/Plugins/Experimental/MeshModelingToolsetExp/Source/MeshModelingToolsExp/Private/ConvertMeshesTool.cpp
+++ b/Engine/Plugins/Experimental/MeshModelingToolsetExp/Source/MeshModelingToolsExp/Private/ConvertMeshesTool.cpp
@@ -96,7 +96,7 @@ void UConvertMeshesTool::Shutdown(EToolShutdownType ShutdownType)
 			const FComponentMaterialSet* TransferMaterials = (BasicProperties->bTransferMaterials) ? &Materials : nullptr;
 
 			FCreateMeshObjectParams NewMeshObjectParams;
-			NewMeshObjectParams.TargetWorld = TargetWorld;
+			NewMeshObjectParams.TargetWorld = GetTargetWorld();
 			NewMeshObjectParams.Transform = (FTransform)SourceTransform;
 			NewMeshObjectParams.BaseName = AssetName;
 			if (TransferMaterials != nullptr)
diff --git a/Engine/Plugins/Experimental/MeshModelingToolsetExp/Source/MeshModelingToolsExp/Private/ConvertToPolygonsTool.cpp b/Engine/Plugins/Experimental/MeshModelingToolsetExp/Source/MeshModelingToolsExp/Private/ConvertToPolygonsTool.cpp
--- a/Engine/Plugins/Experimental/MeshModelingToolsetExp/Source/MeshModelingToolsExp/Private/ConvertToPolygonsTool.cpp
+++ b/Engine/Plugins/Experimental/MeshModelingToolsetExp/Source/MeshModelingToolsExp/Private/ConvertToPolygonsTool.cpp
@@ -192,7 +192,7 @@ void UConvertToPolygonsTool::Setup()
 		ConvertToPolygonsOperatorFactory->ConvertToPolygonsTool = this; // set the back pointer
 
 		PreviewCompute = NewObject<UMeshOpPreviewWithBackgroundCompute>(ConvertToPolygonsOperatorFactory);
-		PreviewCompute->Setup(this->TargetWorld, ConvertToPolygonsOperatorFactory);
+		PreviewCompute->Setup(GetTargetWorld(), ConvertToPolygonsOperatorFactory);
 		ToolSetupUtil::ApplyRenderingConfigurationToPreview(PreviewCompute->PreviewMesh, Target);
 		PreviewCompute->SetIsMeshTopologyConstant(true, EMeshRenderAttributeFlags::Positions | EMeshRenderAttributeFlags::VertexNormals);
 
@@ -221,7 +221,7 @@ void UConvertToPolygonsTool::Setup()
 	}
 	
 	PreviewGeometry = NewObject<UPreviewGeometry>(this);
-	PreviewGeometry->CreateInWorld(TargetWorld, MeshTransform);
+	PreviewGeometry->CreateInWorld(GetTargetWorld(), MeshTransform);
 
 	// updates the triangle color visualization
 	UpdateVisualization();
diff --git a/Engine/Plugins/Experimental/MeshModelingToolsetExp/Source/MeshModelingToolsExp/Private/EditNormalsTool.cpp b/Engine/Plugins/Experimental/MeshModelingToolsetExp/Source/MeshModelingToolsExp/Private/EditNormalsTool.cpp
--- a/Engine/Plugins/Experimental/MeshModelingToolsetExp/Source/MeshModelingToolsExp/Private/EditNormalsTool.cpp
+++ b/Engine/Plugins/Experimental/MeshModelingToolsetExp/Source/MeshModelingToolsExp/Private/EditNormalsTool.cpp
@@ -126,7 +126,7 @@ void UEditNormalsTool::UpdateNumPreviews()
 			Converter.Convert(UE::ToolTarget::GetMeshDescription(Targets[PreviewIdx]), *OriginalDynamicMeshes[PreviewIdx]);
 
 			UMeshOpPreviewWithBackgroundCompute* Preview = Previews.Add_GetRef(NewObject<UMeshOpPreviewWithBackgroundCompute>(OpFactory, "Preview"));
-			Preview->Setup(this->TargetWorld, OpFactory);
+			Preview->Setup(GetTargetWorld(), OpFactory);
 			ToolSetupUtil::ApplyRenderingConfigurationToPreview(Preview->PreviewMesh, Targets[PreviewIdx]);
 			Preview->PreviewMesh->SetTangentsMode(EDynamicMeshComponentTangentsMode::AutoCalculated);
 
diff --git a/Engine/Plugins/Experimental/MeshModelingToolsetExp/Source/MeshModelingToolsExp/Private/HoleFillTool.cpp b/Engine/Plugins/Experimental/MeshModelingToolsetExp/Source/MeshModelingToolsExp/Private/HoleFillTool.cpp
--- a/Engine/Plugins/Experimental/MeshModelingToolsetExp/Source/MeshModelingToolsExp/Private/HoleFillTool.cpp
+++ b/Engine/Plugins/Experimental/MeshModelingToolsetExp/Source/MeshModelingToolsExp/Private/HoleFillTool.cpp
@@ -172,7 +172,7 @@ void UHoleFillTool::Setup()
 	SelectionMechanic->Properties->bSelectVertices = false;
 	SelectionMechanic->Initialize(OriginalMesh.Get(),
 		TargetComponent->GetWorldTransform(),
-		TargetWorld,
+		GetTargetWorld(),
 		Topology.Get(),
 		[this]() { return &MeshSpatial; }
 	);
@@ -307,7 +307,7 @@ void UHoleFillTool::SetupPreview()
 	OpFactory->FillTool = this;
 
 	Preview = NewObject<UMeshOpPreviewWithBackgroundCompute>(OpFactory, "Preview");
-	Preview->Setup(this->TargetWorld, OpFactory);
+	Preview->Setup(GetTargetWorld(), OpFactory);
 	ToolSetupUtil::ApplyRenderingConfigurationToPreview(Preview->PreviewMesh, Target);
 
 	FComponentMaterialSet MaterialSet;
diff --git a/Engine/Plugins/Experimental/MeshModelingToolsetExp/Source/MeshModelingToolsExp/Private/LatticeDeformerTool.cpp b/Engine/Plugins/Experimental/MeshModelingToolsetExp/Source/MeshModelingToolsExp/Private/LatticeDeformerTool.cpp
--- a/Engine/Plugins/Experimental/MeshModelingToolsetExp/Source/MeshModelingToolsExp/Private/LatticeDeformerTool.cpp
+++ b/Engine/Plugins/Experimental/MeshModelingToolsetExp/Source/MeshModelingToolsExp/Private/LatticeDeformerTool.cpp
@@ -215,7 +215,7 @@ void ULatticeDeformerTool::Setup()
 	// Set up control points mechanic
 	ControlPointsMechanic = NewObject<ULatticeControlPointsMechanic>(this);
 	ControlPointsMechanic->Setup(this);
-	ControlPointsMechanic->SetWorld(TargetWorld);
+	ControlPointsMechanic->SetWorld(GetTargetWorld());
 	UE::Geometry::FTransform3d LocalToWorld(Cast<IPrimitiveComponentBackedTarget>(Target)->GetWorldTransform());
 	ControlPointsMechanic->Initialize(LatticePoints, LatticeEdges, LocalToWorld);
 
@@ -374,7 +374,7 @@ void ULatticeDeformerTool::StartPreview()
 	LatticeDeformOpCreator->LatticeDeformerTool = this;
 
 	Preview = NewObject<UMeshOpPreviewWithBackgroundCompute>(LatticeDeformOpCreator);
-	Preview->Setup(TargetWorld, LatticeDeformOpCreator);
+	Preview->Setup(GetTargetWorld(), LatticeDeformOpCreator);
 	ToolSetupUtil::ApplyRenderingConfigurationToPreview(Preview->PreviewMesh, Target);
 
 	Preview->SetIsMeshTopologyConstant(true, EMeshRenderAttributeFlags::Positions | EMeshRenderAttributeFlags::VertexNormals);
diff --git a/Engine/Plugins/Experimental/MeshModelingToolsetExp/Source/MeshModelingToolsExp/Private/MeshBoundaryToolBase.cpp b/Engine/Plugins/Experimental/MeshModelingToolsetExp/Source/MeshModelingToolsExp/Private/MeshBoundaryToolBase.cpp
--- a/Engine/Plugins/Experimental/MeshModelingToolsetExp/Source/MeshModelingToolsExp/Private/MeshBoundaryToolBase.cpp
+++ b/Engine/Plugins/Experimental/MeshModelingToolsetExp/Source/MeshModelingToolsExp/Private/MeshBoundaryToolBase.cpp
@@ -46,7 +46,7 @@ void UMeshBoundaryToolBase::Setup()
 	SelectionMechanic->Properties->bSelectVertices = false;
 	SelectionMechanic->Initialize(OriginalMesh.Get(),
 		Cast<IPrimitiveComponentBackedTarget>(Target)->GetWorldTransform(),
-		TargetWorld,
+		GetTargetWorld(),
 		Topology.Get(),
 		[this]() { return &MeshSpatial; }
 	);
diff --git a/Engine/Plugins/Experimental/MeshModelingToolsetExp/Source/MeshModelingToolsExp/Private/MeshInspectorTool.cpp b/Engine/Plugins/Experimental/MeshModelingToolsetExp/Source/MeshModelingToolsExp/Private/MeshInspectorTool.cpp
--- a/Engine/Plugins/Experimental/MeshModelingToolsetExp/Source/MeshModelingToolsExp/Private/MeshInspectorTool.cpp
+++ b/Engine/Plugins/Experimental/MeshModelingToolsetExp/Source/MeshModelingToolsExp/Private/MeshInspectorTool.cpp
@@ -50,7 +50,7 @@ void UMeshInspectorTool::Setup()
 
 	PreviewMesh = NewObject<UPreviewMesh>(this);
 	PreviewMesh->bBuildSpatialDataStructure = false;
-	PreviewMesh->CreateInWorld(TargetWorld, FTransform::Identity);
+	PreviewMesh->CreateInWorld(GetTargetWorld(), FTransform::Identity);
 	ToolSetupUtil::ApplyRenderingConfigurationToPreview(PreviewMesh, Target);
 	PreviewMesh->SetTransform((FTransform)UE::ToolTarget::GetLocalToWorldTransform(Target));
 
diff --git a/Engine/Plugins/Experimental/MeshModelingToolsetExp/Source/MeshModelingToolsExp/Private/MeshSpaceDeformerTool.cpp b/Engine/Plugins/Experimental/MeshModelingToolsetExp/Source/MeshModelingToolsExp/Private/MeshSpaceDeformerTool.cpp
--- a/Engine/Plugins/Experimental/MeshModelingToolsetExp/Source/MeshModelingToolsExp/Private/MeshSpaceDeformerTool.cpp
+++ b/Engine/Plugins/Experimental/MeshModelingToolsetExp/Source/MeshModelingToolsExp/Private/MeshSpaceDeformerTool.cpp
@@ -150,7 +150,7 @@ void UMeshSpaceDeformerTool::Setup()
 	TargetComponent->SetOwnerVisibility(false);
 
 	OriginalMeshPreview = NewObject<UPreviewMesh>();
-	OriginalMeshPreview->CreateInWorld(TargetWorld, MeshTransform);
+	OriginalMeshPreview->CreateInWorld(GetTargetWorld(), MeshTransform);
 	ToolSetupUtil::ApplyRenderingConfigurationToPreview(OriginalMeshPreview, Target);
 	OriginalMeshPreview->UpdatePreview(OriginalDynamicMesh.Get());
 	OriginalMeshPreview->SetMaterial(0, ToolSetupUtil::GetCustomDepthOffsetMaterial(GetToolManager(), FLinearColor::White,
@@ -276,7 +276,7 @@ void UMeshSpaceDeformerTool::Setup()
 
 
 		Preview = NewObject<UMeshOpPreviewWithBackgroundCompute>(DeformerOperatorFactory, "Preview");
-		Preview->Setup(this->TargetWorld, DeformerOperatorFactory);
+		Preview->Setup(GetTargetWorld(), DeformerOperatorFactory);
 		ToolSetupUtil::ApplyRenderingConfigurationToPreview(Preview->PreviewMesh, Target);
 		Preview->PreviewMesh->SetTangentsMode(EDynamicMeshComponentTangentsMode::AutoCalculated);
 
diff --git a/Engine/Plugins/Experimental/MeshModelingToolsetExp/Source/MeshModelingToolsExp/Private/MirrorTool.cpp b/Engine/Plugins/Experimental/MeshModelingToolsetExp/Source/MeshModelingToolsExp/Private/MirrorTool.cpp
--- a/Engine/Plugins/Experimental/MeshModelingToolsetExp/Source/MeshModelingToolsExp/Private/MirrorTool.cpp
+++ b/Engine/Plugins/Experimental/MeshModelingToolsetExp/Source/MeshModelingToolsExp/Private/MirrorTool.cpp
@@ -216,7 +216,7 @@ void UMirrorTool::Setup()
 	// Set up the mirror plane mechanic, which manages the gizmo
 	PlaneMechanic = NewObject<UConstructionPlaneMechanic>(this);
 	PlaneMechanic->Setup(this);
-	PlaneMechanic->Initialize(TargetWorld, FFrame3d(MirrorPlaneOrigin, MirrorPlaneNormal));
+	PlaneMechanic->Initialize(GetTargetWorld(), FFrame3d(MirrorPlaneOrigin, MirrorPlaneNormal));
 
 	// Have the plane mechanic update things properly
 	PlaneMechanic->OnPlaneChanged.AddLambda([this]() {
@@ -260,7 +260,7 @@ void UMirrorTool::SetupPreviews()
 
 		UMeshOpPreviewWithBackgroundCompute* Preview = Previews.Add_GetRef(
 			NewObject<UMeshOpPreviewWithBackgroundCompute>(MirrorOpCreator, "Preview"));
-		Preview->Setup(TargetWorld, MirrorOpCreator);
+		Preview->Setup(GetTargetWorld(), MirrorOpCreator);
 		ToolSetupUtil::ApplyRenderingConfigurationToPreview(Preview->PreviewMesh, nullptr);
 		Preview->PreviewMesh->SetTangentsMode(EDynamicMeshComponentTangentsMode::AutoCalculated);
 
@@ -429,7 +429,7 @@ void UMirrorTool::GenerateAsset(const TArray<FDynamicMeshOpResult>& Results)
 			}
 
 			FCreateMeshObjectParams NewMeshObjectParams;
-			NewMeshObjectParams.TargetWorld = TargetWorld;
+			NewMeshObjectParams.TargetWorld = GetTargetWorld();
 			NewMeshObjectParams.Transform = (FTransform)Results[OrigMeshIdx].Transform;
 			NewMeshObjectParams.BaseName = TEXT("Mirror");
 			NewMeshObjectParams.Materials = Materials;
diff --git a/Engine/Plugins/Experimental/MeshModelingToolsetExp/Source/MeshModelingToolsExp/Private/Physics/ExtractCollisionGeometryTool.cpp b/Engine/Plugins/Experimental/MeshModelingToolsetExp/Source/MeshModelingToolsExp/Private/Physics/ExtractCollisionGeometryTool.cpp
--- a/Engine/Plugins/Experimental/MeshModelingToolsetExp/Source/MeshModelingToolsExp/Private/Physics/ExtractCollisionGeometryTool.cpp
+++ b/Engine/Plugins/Experimental/MeshModelingToolsetExp/Source/MeshModelingToolsExp/Private/Physics/ExtractCollisionGeometryTool.cpp
@@ -57,7 +57,7 @@ void UExtractCollisionGeometryTool::Setup()
 	// create preview mesh
 	PreviewMesh = NewObject<UPreviewMesh>(this);
 	PreviewMesh->bBuildSpatialDataStructure = false;
-	PreviewMesh->CreateInWorld(TargetWorld, FTransform::Identity);
+	PreviewMesh->CreateInWorld(GetTargetWorld(), FTransform::Identity);
 	PreviewMesh->SetTransform((FTransform)UE::ToolTarget::GetLocalToWorldTransform(Target));
 	PreviewMesh->SetMaterial(ToolSetupUtil::GetDefaultSculptMaterial(GetToolManager()));
 	PreviewMesh->SetOverrideRenderMaterial(ToolSetupUtil::GetSelectionMaterial(GetToolManager()));
@@ -146,7 +146,7 @@ void UExtractCollisionGeometryTool::Shutdown(EToolShutdownType ShutdownType)
 		auto EmitNewMesh = [&](FDynamicMesh3&& Mesh, FTransform3d UseTransform, FString UseName)
 		{
 			FCreateMeshObjectParams NewMeshObjectParams;
-			NewMeshObjectParams.TargetWorld = TargetWorld;
+			NewMeshObjectParams.TargetWorld = GetTargetWorld();
 			NewMeshObjectParams.Transform = (FTransform)UseTransform;
 			NewMeshObjectParams.BaseName = UseName;
 			NewMeshObjectParams.Materials.Add(UseMaterial);
diff --git a/Engine/Plugins/Experimental/MeshModelingToolsetExp/Source/MeshModelingToolsExp/Private/PlaneCutTool.cpp b/Engine/Plugins/Experimental/MeshModelingToolsetExp/Source/MeshModelingToolsExp/Private/PlaneCutTool.cpp
--- a/Engine/Plugins/Experimental/MeshModelingToolsetExp/Source/MeshModelingToolsExp/Private/PlaneCutTool.cpp
+++ b/Engine/Plugins/Experimental/MeshModelingToolsetExp/Source/MeshModelingToolsExp/Private/PlaneCutTool.cpp
@@ -126,7 +126,7 @@ void UPlaneCutTool::Setup()
 	CutPlaneWorld.Origin = (FVector3d)CombinedBounds.GetCenter();
 	PlaneMechanic = NewObject<UConstructionPlaneMechanic>(this);
 	PlaneMechanic->Setup(this);
-	PlaneMechanic->Initialize(TargetWorld, CutPlaneWorld);
+	PlaneMechanic->Initialize(GetTargetWorld(), CutPlaneWorld);
 	PlaneMechanic->OnPlaneChanged.AddLambda([this]() {
 		CutPlaneWorld = PlaneMechanic->Plane;
 		InvalidatePreviews();
@@ -182,7 +182,7 @@ void UPlaneCutTool::SetupPreviews()
 		CutSide->CutTool = this;
 		CutSide->ComponentIndex = PreviewIdx;
 		UMeshOpPreviewWithBackgroundCompute* Preview = Previews.Add_GetRef(NewObject<UMeshOpPreviewWithBackgroundCompute>(CutSide, "Preview"));
-		Preview->Setup(this->TargetWorld, CutSide);
+		Preview->Setup(GetTargetWorld(), CutSide);
 		ToolSetupUtil::ApplyRenderingConfigurationToPreview(Preview->PreviewMesh, nullptr);
 		Preview->PreviewMesh->SetTangentsMode(EDynamicMeshComponentTangentsMode::AutoCalculated);
 
@@ -495,7 +495,7 @@ void UPlaneCutTool::GenerateAsset(const TArray<FDynamicMeshOpResult>& Results)
 			for (int AddMeshIdx = 1; AddMeshIdx < SplitMeshes.Num(); AddMeshIdx++)
 			{
 				FCreateMeshObjectParams NewMeshObjectParams;
-				NewMeshObjectParams.TargetWorld = TargetWorld;
+				NewMeshObjectParams.TargetWorld = GetTargetWorld();
 				NewMeshObjectParams.Transform = (FTransform)Results[OrigMeshIdx].Transform;
 				NewMeshObjectParams.BaseName = TEXT("PlaneCutOtherPart");
 				NewMeshObjectParams.Materials = ComponentMaterialSet.Materials;
diff --git a/Engine/Plugins/Experimental/MeshModelingToolsetExp/Source/MeshModelingToolsExp/Private/RemeshMeshTool.cpp b/Engine/Plugins/Experimental/MeshModelingToolsetExp/Source/MeshModelingToolsExp/Private/RemeshMeshTool.cpp
--- a/Engine/Plugins/Experimental/MeshModelingToolsetExp/Source/MeshModelingToolsExp/Private/RemeshMeshTool.cpp
+++ b/Engine/Plugins/Experimental/MeshModelingToolsetExp/Source/MeshModelingToolsExp/Private/RemeshMeshTool.cpp
@@ -90,7 +90,7 @@ void URemeshMeshTool::Setup()
 	// hide component and create + show preview
 	UE::ToolTarget::HideSourceObject(Targets[0]);
 	Preview = NewObject<UMeshOpPreviewWithBackgroundCompute>(this);
-	Preview->Setup(this->TargetWorld, this);
+	Preview->Setup(GetTargetWorld(), this);
 	ToolSetupUtil::ApplyRenderingConfigurationToPreview(Preview->PreviewMesh, Targets[0]);
 
 	const FComponentMaterialSet MaterialSet = UE::ToolTarget::GetMaterialSet(Targets[0]);
diff --git a/Engine/Plugins/Experimental/MeshModelingToolsetExp/Source/MeshModelingToolsExp/Private/RemoveOccludedTrianglesTool.cpp b/Engine/Plugins/Experimental/MeshModelingToolsetExp/Source/MeshModelingToolsExp/Private/RemoveOccludedTrianglesTool.cpp
--- a/Engine/Plugins/Experimental/MeshModelingToolsetExp/Source/MeshModelingToolsExp/Private/RemoveOccludedTrianglesTool.cpp
+++ b/Engine/Plugins/Experimental/MeshModelingToolsetExp/Source/MeshModelingToolsExp/Private/RemoveOccludedTrianglesTool.cpp
@@ -228,7 +228,7 @@ void URemoveOccludedTrianglesTool::SetupPreviews()
 			Converter.Convert(UE::ToolTarget::GetMeshDescription(Targets[TargetIdx]), *OriginalDynamicMeshes[PreviewIdx]);
 
 			UMeshOpPreviewWithBackgroundCompute* Preview = Previews.Add_GetRef(NewObject<UMeshOpPreviewWithBackgroundCompute>(OpFactory, "Preview"));
-			Preview->Setup(this->TargetWorld, OpFactory);
+			Preview->Setup(GetTargetWorld(), OpFactory);
 			ToolSetupUtil::ApplyRenderingConfigurationToPreview(Preview->PreviewMesh, nullptr);
 			Preview->PreviewMesh->SetTangentsMode(EDynamicMeshComponentTangentsMode::AutoCalculated);
 
@@ -263,7 +263,7 @@ void URemoveOccludedTrianglesTool::SetupPreviews()
 			// already did the conversion for a full UMeshOpPreviewWithBackgroundCompute -- just make a light version of that and hook it up to copy the other's work
 			int CopyIdx = PreviewCopies.Num();
 			UPreviewMesh* PreviewMesh = PreviewCopies.Add_GetRef(NewObject<UPreviewMesh>(this));
-			PreviewMesh->CreateInWorld(this->TargetWorld, (FTransform) UE::ToolTarget::GetLocalToWorldTransform(Targets[TargetIdx]));
+			PreviewMesh->CreateInWorld(GetTargetWorld(), (FTransform) UE::ToolTarget::GetLocalToWorldTransform(Targets[TargetIdx]));
 
 			PreviewToCopyIdx[PreviewIdx].Add(CopyIdx);
 
diff --git a/Engine/Plugins/Experimental/MeshModelingToolsetExp/Source/MeshModelingToolsExp/Private/RevolveBoundaryTool.cpp b/Engine/Plugins/Experimental/MeshModelingToolsetExp/Source/MeshModelingToolsExp/Private/RevolveBoundaryTool.cpp
--- a/Engine/Plugins/Experimental/MeshModelingToolsetExp/Source/MeshModelingToolsExp/Private/RevolveBoundaryTool.cpp
+++ b/Engine/Plugins/Experimental/MeshModelingToolsetExp/Source/MeshModelingToolsExp/Private/RevolveBoundaryTool.cpp
@@ -102,7 +102,7 @@ void URevolveBoundaryTool::Setup()
 	// The plane mechanic is used for the revolution axis
 	PlaneMechanic = NewObject<UConstructionPlaneMechanic>(this);
 	PlaneMechanic->Setup(this);
-	PlaneMechanic->Initialize(TargetWorld, FFrame3d(Settings->AxisOrigin, 
+	PlaneMechanic->Initialize(GetTargetWorld(), FFrame3d(Settings->AxisOrigin, 
 		FRotator(Settings->AxisOrientation.X, Settings->AxisOrientation.Y, 0).Quaternion()));
 	PlaneMechanic->UpdateClickPriority(ClickBehavior->GetPriority().MakeLower());
 	PlaneMechanic->bShowGrid = false;
@@ -246,7 +246,7 @@ void URevolveBoundaryTool::StartPreview()
 	RevolveBoundaryOpCreator->RevolveBoundaryTool = this;
 
 	Preview = NewObject<UMeshOpPreviewWithBackgroundCompute>(RevolveBoundaryOpCreator);
-	Preview->Setup(TargetWorld, RevolveBoundaryOpCreator);
+	Preview->Setup(GetTargetWorld(), RevolveBoundaryOpCreator);
 	ToolSetupUtil::ApplyRenderingConfigurationToPreview(Preview->PreviewMesh, nullptr);
 	Preview->PreviewMesh->SetTangentsMode(EDynamicMeshComponentTangentsMode::AutoCalculated);
 
@@ -301,7 +301,7 @@ void URevolveBoundaryTool::GenerateAsset(const FDynamicMeshOpResult& OpResult)
 	GetToolManager()->BeginUndoTransaction(LOCTEXT("RevolveBoundaryToolTransactionName", "Boundary Revolve Tool"));
 
 	FCreateMeshObjectParams NewMeshObjectParams;
-	NewMeshObjectParams.TargetWorld = TargetWorld;
+	NewMeshObjectParams.TargetWorld = GetTargetWorld();
 	NewMeshObjectParams.Transform = (FTransform)OpResult.Transform;
 	NewMeshObjectParams.BaseName = TEXT("Revolve");
 	NewMeshObjectParams.Materials.Add(MaterialProperties->Material.Get());
diff --git a/Engine/Plugins/Experimental/MeshModelingToolsetExp/Source/MeshModelingToolsExp/Private/SkinWeightsBindingTool.cpp b/Engine/Plugins/Experimental/MeshModelingToolsetExp/Source/MeshModelingToolsExp/Private/SkinWeightsBindingTool.cpp
--- a/Engine/Plugins/Experimental/MeshModelingToolsetExp/Source/MeshModelingToolsExp/Private/SkinWeightsBindingTool.cpp
+++ b/Engine/Plugins/Experimental/MeshModelingToolsetExp/Source/MeshModelingToolsExp/Private/SkinWeightsBindingTool.cpp
@@ -807,7 +807,7 @@ void USkinWeightsBindingTool::Setup()
 	
 	UE::ToolTarget::HideSourceObject(Targets[0]);
 	Preview = NewObject<UMeshOpPreviewWithBackgroundCompute>(this, "Preview");
-	Preview->Setup(this->TargetWorld, this);
+	Preview->Setup(GetTargetWorld(), this);
 	Preview->SetIsMeshTopologyConstant(true, EMeshRenderAttributeFlags::VertexColors);
 	Preview->OnMeshUpdated.AddLambda([this](UMeshOpPreviewWithBackgroundCompute* Compute)
 	{
diff --git a/Engine/Plugins/Experimental/MeshModelingToolsetExp/Source/MeshModelingToolsExp/Private/SplitMeshesTool.cpp b/Engine/Plugins/Experimental/MeshModelingToolsetExp/Source/MeshModelingToolsExp/Private/SplitMeshesTool.cpp
--- a/Engine/Plugins/Experimental/MeshModelingToolsetExp/Source/MeshModelingToolsExp/Private/SplitMeshesTool.cpp
+++ b/Engine/Plugins/Experimental/MeshModelingToolsetExp/Source/MeshModelingToolsExp/Private/SplitMeshesTool.cpp
@@ -108,7 +108,7 @@ void USplitMeshesTool::Shutdown(EToolShutdownType ShutdownType)
 			FString AssetName = TargetActor->GetActorNameOrLabel();
 
 			FCreateMeshObjectParams BaseMeshObjectParams;
-			BaseMeshObjectParams.TargetWorld = TargetWorld;
+			BaseMeshObjectParams.TargetWorld = GetTargetWorld();
 
 			if (OutputTypeProperties->OutputType == UCreateMeshObjectTypeProperties::AutoIdentifier)
 			{
diff --git a/Engine/Plugins/Experimental/MeshModelingToolsetExp/Source/MeshModelingToolsExp/Private/WeldMeshEdgesTool.cpp b/Engine/Plugins/Experimental/MeshModelingToolsetExp/Source/MeshModelingToolsExp/Private/WeldMeshEdgesTool.cpp
--- a/Engine/Plugins/Experimental/MeshModelingToolsetExp/Source/MeshModelingToolsExp/Private/WeldMeshEdgesTool.cpp
+++ b/Engine/Plugins/Experimental/MeshModelingToolsetExp/Source/MeshModelingToolsExp/Private/WeldMeshEdgesTool.cpp
@@ -112,7 +112,7 @@ void UWeldMeshEdgesTool::Setup()
 	OperatorFactory->WeldMeshEdgesTool = this;
 
 	PreviewCompute = NewObject<UMeshOpPreviewWithBackgroundCompute>(OperatorFactory);
-	PreviewCompute->Setup(this->TargetWorld, OperatorFactory);
+	PreviewCompute->Setup(GetTargetWorld(), OperatorFactory);
 	ToolSetupUtil::ApplyRenderingConfigurationToPreview(PreviewCompute->PreviewMesh, Target); 
 
 	// Give the preview something to display
diff --git a/Engine/Plugins/Runtime/MeshModelingToolset/Source/MeshModelingTools/Private/CSGMeshesTool.cpp b/Engine/Plugins/Runtime/MeshModelingToolset/Source/MeshModelingTools/Private/CSGMeshesTool.cpp
--- a/Engine/Plugins/Runtime/MeshModelingToolset/Source/MeshModelingTools/Private/CSGMeshesTool.cpp
+++ b/Engine/Plugins/Runtime/MeshModelingToolset/Source/MeshModelingTools/Private/CSGMeshesTool.cpp
@@ -232,7 +232,7 @@ void UCSGMeshesTool::ConvertInputsAndSetPreviewMaterials(bool bSetPreviewMesh)
 		if (bSetPreviewMesh)
 		{
 			UPreviewMesh* OriginalMeshPreview = OriginalMeshPreviews.Add_GetRef(NewObject<UPreviewMesh>());
-			OriginalMeshPreview->CreateInWorld(TargetWorld, (FTransform) UE::ToolTarget::GetLocalToWorldTransform(Targets[ComponentIdx]));
+			OriginalMeshPreview->CreateInWorld(GetTargetWorld(), (FTransform) UE::ToolTarget::GetLocalToWorldTransform(Targets[ComponentIdx]));
 			OriginalMeshPreview->UpdatePreview(OriginalDynamicMeshes[ComponentIdx].Get());
 
 			OriginalMeshPreview->SetMaterial(0, PreviewsGhostMaterial);
diff --git a/Engine/Plugins/Runtime/MeshModelingToolset/Source/MeshModelingTools/Private/CombineMeshesTool.cpp b/Engine/Plugins/Runtime/MeshModelingToolset/Source/MeshModelingTools/Private/CombineMeshesTool.cpp
--- a/Engine/Plugins/Runtime/MeshModelingToolset/Source/MeshModelingTools/Private/CombineMeshesTool.cpp
+++ b/Engine/Plugins/Runtime/MeshModelingToolset/Source/MeshModelingTools/Private/CombineMeshesTool.cpp
@@ -269,7 +269,7 @@ void UCombineMeshesTool::CreateNewAsset()
 		}
 
 		FCreateMeshObjectParams NewMeshObjectParams;
-		NewMeshObjectParams.TargetWorld = TargetWorld;
+		NewMeshObjectParams.TargetWorld = GetTargetWorld();
 		NewMeshObjectParams.Transform = AccumToWorld;
 		NewMeshObjectParams.BaseName = UseBaseName;
 		NewMeshObjectParams.Materials = AllMaterials;
diff --git a/Engine/Plugins/Runtime/MeshModelingToolset/Source/MeshModelingTools/Private/CutMeshWithMeshTool.cpp b/Engine/Plugins/Runtime/MeshModelingToolset/Source/MeshModelingTools/Private/CutMeshWithMeshTool.cpp
--- a/Engine/Plugins/Runtime/MeshModelingToolset/Source/MeshModelingTools/Private/CutMeshWithMeshTool.cpp
+++ b/Engine/Plugins/Runtime/MeshModelingToolset/Source/MeshModelingTools/Private/CutMeshWithMeshTool.cpp
@@ -50,7 +50,7 @@ void UCutMeshWithMeshTool::SetupProperties()
 
 	// create intersection preview mesh object
 	IntersectPreviewMesh = NewObject<UPreviewMesh>(this);
-	IntersectPreviewMesh->CreateInWorld(TargetWorld, FTransform::Identity);
+	IntersectPreviewMesh->CreateInWorld(GetTargetWorld(), FTransform::Identity);
 	ToolSetupUtil::ApplyRenderingConfigurationToPreview(IntersectPreviewMesh, nullptr); 
 	IntersectPreviewMesh->SetVisible(true);
 	IntersectPreviewMesh->SetMaterial(ToolSetupUtil::GetDefaultBrushVolumeMaterial(GetToolManager()));
@@ -372,7 +372,7 @@ void UCutMeshWithMeshTool::Shutdown(EToolShutdownType ShutdownType)
 			FString UseBaseName = FString::Printf(TEXT("%s_%s"), *CurName, TEXT("CutPart") );
 
 			FCreateMeshObjectParams NewMeshObjectParams;
-			NewMeshObjectParams.TargetWorld = TargetWorld;
+			NewMeshObjectParams.TargetWorld = GetTargetWorld();
 			NewMeshObjectParams.Transform = (FTransform)NewTransform;
 			NewMeshObjectParams.BaseName = UseBaseName;
 			NewMeshObjectParams.Materials = GetOutputMaterials();
diff --git a/Engine/Plugins/Runtime/MeshModelingToolset/Source/MeshModelingTools/Private/RecomputeUVsTool.cpp b/Engine/Plugins/Runtime/MeshModelingToolset/Source/MeshModelingTools/Private/RecomputeUVsTool.cpp
--- a/Engine/Plugins/Runtime/MeshModelingToolset/Source/MeshModelingTools/Private/RecomputeUVsTool.cpp
+++ b/Engine/Plugins/Runtime/MeshModelingToolset/Source/MeshModelingTools/Private/RecomputeUVsTool.cpp
@@ -81,7 +81,7 @@ void URecomputeUVsTool::Setup()
 	RecomputeUVsOpFactory->GetSelectedUVChannel = [this]() { return GetSelectedUVChannel(); };
 
 	Preview = NewObject<UMeshOpPreviewWithBackgroundCompute>(this);
-	Preview->Setup(this->TargetWorld, RecomputeUVsOpFactory);
+	Preview->Setup(GetTargetWorld(), RecomputeUVsOpFactory);
 	ToolSetupUtil::ApplyRenderingConfigurationToPreview(Preview->PreviewMesh, Target);
 	Preview->PreviewMesh->SetTangentsMode(EDynamicMeshComponentTangentsMode::AutoCalculated);
 	Preview->PreviewMesh->ReplaceMesh(*InputMesh);
@@ -98,7 +98,7 @@ void URecomputeUVsTool::Setup()
 	if (bCreateUVLayoutViewOnSetup)
 	{
 		UVLayoutView = NewObject<UUVLayoutPreview>(this);
-		UVLayoutView->CreateInWorld(TargetWorld);
+		UVLayoutView->CreateInWorld(GetTargetWorld());
 		UVLayoutView->SetSourceMaterials(MaterialSet);
 		UVLayoutView->SetSourceWorldPosition(TargetTransform, UE::ToolTarget::GetTargetActor(Target)->GetComponentsBoundingBox());
 		UVLayoutView->Settings->bEnabled = false;
diff --git a/Engine/Plugins/Runtime/MeshModelingToolset/Source/MeshModelingTools/Private/UVLayoutTool.cpp b/Engine/Plugins/Runtime/MeshModelingToolset/Source/MeshModelingTools/Private/UVLayoutTool.cpp
--- a/Engine/Plugins/Runtime/MeshModelingToolset/Source/MeshModelingTools/Private/UVLayoutTool.cpp
+++ b/Engine/Plugins/Runtime/MeshModelingToolset/Source/MeshModelingTools/Private/UVLayoutTool.cpp
@@ -80,7 +80,7 @@ void UUVLayoutTool::Setup()
 	if (Targets.Num() == 1)
 	{
 		UVLayoutView = NewObject<UUVLayoutPreview>(this);
-		UVLayoutView->CreateInWorld(TargetWorld);
+		UVLayoutView->CreateInWorld(GetTargetWorld());
 
 		const FComponentMaterialSet MaterialSet = UE::ToolTarget::GetMaterialSet(Targets[0]);
 		UVLayoutView->SetSourceMaterials(MaterialSet);
@@ -133,7 +133,7 @@ void UUVLayoutTool::UpdateNumPreviews()
 			Factories[PreviewIdx]->GetSelectedUVChannel = [this]() { return GetSelectedUVChannel(); };
 
 			UMeshOpPreviewWithBackgroundCompute* Preview = Previews.Add_GetRef(NewObject<UMeshOpPreviewWithBackgroundCompute>(Factories[PreviewIdx], "Preview"));
-			Preview->Setup(this->TargetWorld, Factories[PreviewIdx]);
+			Preview->Setup(GetTargetWorld(), Factories[PreviewIdx]);
 			ToolSetupUtil::ApplyRenderingConfigurationToPreview(Preview->PreviewMesh, Targets[PreviewIdx]); 
 
 			const FComponentMaterialSet MaterialSet = UE::ToolTarget::GetMaterialSet(Targets[PreviewIdx]);
diff --git a/Engine/Plugins/Runtime/MeshModelingToolset/Source/MeshModelingTools/Private/UVProjectionTool.cpp b/Engine/Plugins/Runtime/MeshModelingToolset/Source/MeshModelingTools/Private/UVProjectionTool.cpp
--- a/Engine/Plugins/Runtime/MeshModelingToolset/Source/MeshModelingTools/Private/UVProjectionTool.cpp
+++ b/Engine/Plugins/Runtime/MeshModelingToolset/Source/MeshModelingTools/Private/UVProjectionTool.cpp
@@ -154,7 +154,7 @@ void UUVProjectionTool::UpdateNumPreviews()
 	OperatorFactory->Tool = this;
 
 	Preview = NewObject<UMeshOpPreviewWithBackgroundCompute>(OperatorFactory);
-	Preview->Setup(this->TargetWorld, OperatorFactory);
+	Preview->Setup(GetTargetWorld(), OperatorFactory);
 	ToolSetupUtil::ApplyRenderingConfigurationToPreview(Preview->PreviewMesh, Target); 
 	Preview->OnMeshUpdated.AddUObject(this, &UUVProjectionTool::OnMeshUpdated);
 	Preview->PreviewMesh->SetTangentsMode(EDynamicMeshComponentTangentsMode::AutoCalculated);
@@ -169,7 +169,7 @@ void UUVProjectionTool::UpdateNumPreviews()
 	Preview->SetVisibility(true);
 
 	EdgeRenderer = NewObject<UPreviewGeometry>(this);
-	EdgeRenderer->CreateInWorld(this->TargetWorld, (FTransform)WorldTransform);
+	EdgeRenderer->CreateInWorld(GetTargetWorld(), (FTransform)WorldTransform);
 
 	// if we have an ROI, show its borders
 	if (TriangleROI->Num() > 0)
diff --git a/Engine/Plugins/Runtime/MeshModelingToolset/Source/ModelingComponents/Private/BaseTools/BaseCreateFromSelectedTool.cpp b/Engine/Plugins/Runtime/MeshModelingToolset/Source/ModelingComponents/Private/BaseTools/BaseCreateFromSelectedTool.cpp
--- a/Engine/Plugins/Runtime/MeshModelingToolset/Source/ModelingComponents/Private/BaseTools/BaseCreateFromSelectedTool.cpp
+++ b/Engine/Plugins/Runtime/MeshModelingToolset/Source/ModelingComponents/Private/BaseTools/BaseCreateFromSelectedTool.cpp
@@ -74,7 +74,7 @@ void UBaseCreateFromSelectedTool::Setup()
 	AddToolPropertySource(HandleSourcesProperties);
 
 	Preview = NewObject<UMeshOpPreviewWithBackgroundCompute>(this);
-	Preview->Setup(this->TargetWorld, this);
+	Preview->Setup(GetTargetWorld(), this);
 	ToolSetupUtil::ApplyRenderingConfigurationToPreview(Preview->PreviewMesh, nullptr);
 
 	SetPreviewCallbacks();
@@ -211,7 +211,7 @@ void UBaseCreateFromSelectedTool::GenerateAsset(const FDynamicMeshOpResult& OpRe
 	}
 
 	FCreateMeshObjectParams NewMeshObjectParams;
-	NewMeshObjectParams.TargetWorld = TargetWorld;
+	NewMeshObjectParams.TargetWorld = GetTargetWorld();
 	NewMeshObjectParams.Transform = (FTransform)NewTransform;
 	NewMeshObjectParams.BaseName = UseBaseName;
 	NewMeshObjectParams.Materials = GetOutputMaterials();
diff --git a/Engine/Plugins/Runtime/MeshModelingToolset/Source/ModelingComponents/Private/MeshOpPreviewHelpers.cpp b/Engine/Plugins/Runtime/MeshModelingToolset/Source/ModelingComponents/Private/MeshOpPreviewHelpers.cpp
--- a/Engine/Plugins/Runtime/MeshModelingToolset/Source/ModelingComponents/Private/MeshOpPreviewHelpers.cpp
+++ b/Engine/Plugins/Runtime/MeshModelingToolset/Source/ModelingComponents/Private/MeshOpPreviewHelpers.cpp
@@ -48,6 +48,8 @@ FDynamicMeshOpResult UMeshOpPreviewWithBackgroundCompute::Shutdown()
 	PreviewMesh->Disconnect();
 	PreviewMesh = nullptr;
 
+	PreviewWorld = nullptr;
+
 	return Result;
 }
 
diff --git a/Engine/Plugins/Runtime/MeshModelingToolset/Source/ModelingComponents/Public/BaseTools/MultiSelectionMeshEditingTool.h b/Engine/Plugins/Runtime/MeshModelingToolset/Source/ModelingComponents/Public/BaseTools/MultiSelectionMeshEditingTool.h
--- a/Engine/Plugins/Runtime/MeshModelingToolset/Source/ModelingComponents/Public/BaseTools/MultiSelectionMeshEditingTool.h
+++ b/Engine/Plugins/Runtime/MeshModelingToolset/Source/ModelingComponents/Public/BaseTools/MultiSelectionMeshEditingTool.h
@@ -45,7 +45,7 @@ class MODELINGCOMPONENTS_API UMultiSelectionMeshEditingTool : public UMultiSelec
 	GENERATED_BODY()
 public:
 	virtual void SetWorld(UWorld* World) { TargetWorld = World; }
-	virtual UWorld* GetTargetWorld() { return TargetWorld; }
+	virtual UWorld* GetTargetWorld() { return TargetWorld.Get(); }
 
 	/**
 	 * Helper to find which targets share source data.
@@ -57,5 +57,5 @@ public:
 
 protected:
 	UPROPERTY()
-	TObjectPtr<UWorld> TargetWorld = nullptr;
+	TWeakObjectPtr<UWorld> TargetWorld = nullptr;
 };
diff --git a/Engine/Plugins/Runtime/MeshModelingToolset/Source/ModelingComponents/Public/BaseTools/SingleSelectionMeshEditingTool.h b/Engine/Plugins/Runtime/MeshModelingToolset/Source/ModelingComponents/Public/BaseTools/SingleSelectionMeshEditingTool.h
--- a/Engine/Plugins/Runtime/MeshModelingToolset/Source/ModelingComponents/Public/BaseTools/SingleSelectionMeshEditingTool.h
+++ b/Engine/Plugins/Runtime/MeshModelingToolset/Source/ModelingComponents/Public/BaseTools/SingleSelectionMeshEditingTool.h
@@ -51,11 +51,11 @@ class MODELINGCOMPONENTS_API USingleSelectionMeshEditingTool : public USingleSel
 public:
 	virtual void SetWorld(UWorld* World) { TargetWorld = World; }
 
-	virtual UWorld* GetTargetWorld() { return TargetWorld; }
+	virtual UWorld* GetTargetWorld() { return TargetWorld.Get(); }
 
 protected:
 	UPROPERTY()
-	TObjectPtr<UWorld> TargetWorld = nullptr;
+	TWeakObjectPtr<UWorld> TargetWorld = nullptr;
 
 	//
 	// Mesh Selection support
diff --git a/Engine/Plugins/Runtime/MeshModelingToolset/Source/ModelingComponents/Public/MeshOpPreviewHelpers.h b/Engine/Plugins/Runtime/MeshModelingToolset/Source/ModelingComponents/Public/MeshOpPreviewHelpers.h
--- a/Engine/Plugins/Runtime/MeshModelingToolset/Source/ModelingComponents/Public/MeshOpPreviewHelpers.h
+++ b/Engine/Plugins/Runtime/MeshModelingToolset/Source/ModelingComponents/Public/MeshOpPreviewHelpers.h
@@ -147,7 +147,7 @@ public:
 	bool ProcessCurrentMesh(TFunctionRef<void(const UE::Geometry::FDynamicMesh3&)> ProcessFunc, bool bOnlyIfValid = false);
 
 	/** @return UWorld that the created PreviewMesh exist in */
-	virtual UWorld* GetWorld() const override { return PreviewWorld; }
+	virtual UWorld* GetWorld() const override { return PreviewWorld.Get(); }
 
 	//
 	// Optional configuration
@@ -224,7 +224,8 @@ public:
 	UPROPERTY()
 	TObjectPtr<UMaterialInterface> SecondaryMaterial = nullptr;
 
-	UWorld* PreviewWorld;
+	UPROPERTY()
+	TWeakObjectPtr<UWorld> PreviewWorld = nullptr;
 
 	/**
 	 * When true, the preview mesh is allowed to be temporarily updated using results that we know
