MeshModelingTools: Clarify that Shift only ignores snapping while drawing in PolyExt and PathExt.

#rnx

diff --git a/Engine/Plugins/Experimental/MeshModelingToolsetExp/Source/MeshModelingToolsExp/Private/DrawPolyPathTool.cpp b/Engine/Plugins/Experimental/MeshModelingToolsetExp/Source/MeshModelingToolsExp/Private/DrawPolyPathTool.cpp
--- a/Engine/Plugins/Experimental/MeshModelingToolsetExp/Source/MeshModelingToolsExp/Private/DrawPolyPathTool.cpp
+++ b/Engine/Plugins/Experimental/MeshModelingToolsetExp/Source/MeshModelingToolsExp/Private/DrawPolyPathTool.cpp
@@ -1055,7 +1055,7 @@ void UDrawPolyPathTool::EmitNewObject()
 void UDrawPolyPathTool::ShowStartupMessage()
 {
 	GetToolManager()->DisplayMessage(
-		LOCTEXT("OnStartDraw", "Draw a path on the drawing plane, set its width, and extrude it. Left-click to place path vertices, and click on the last or first vertex to complete the path. Hold Shift to ignore snapping."),
+		LOCTEXT("OnStartDraw", "Draw a path on the drawing plane, set its width, and extrude it. Left-click to place path vertices, and click on the last or first vertex to complete the path. Hold Shift to ignore snapping while drawing."),
 		EToolMessageLevel::UserNotification);
 }
 
diff --git a/Engine/Plugins/Runtime/MeshModelingToolset/Source/MeshModelingTools/Private/DrawPolygonTool.cpp b/Engine/Plugins/Runtime/MeshModelingToolset/Source/MeshModelingTools/Private/DrawPolygonTool.cpp
--- a/Engine/Plugins/Runtime/MeshModelingToolset/Source/MeshModelingTools/Private/DrawPolygonTool.cpp
+++ b/Engine/Plugins/Runtime/MeshModelingToolset/Source/MeshModelingTools/Private/DrawPolygonTool.cpp
@@ -1353,7 +1353,7 @@ bool UDrawPolygonTool::GeneratePolygonMesh(const TArray<FVector3d>& Polygon, con
 void UDrawPolygonTool::ShowStartupMessage()
 {
 	GetToolManager()->DisplayMessage(
-		LOCTEXT("OnStartDraw", "Draw a polygon on the drawing plane, and extrude it. Left-click to place polygon vertices. Hold Shift to ignore snapping."),
+		LOCTEXT("OnStartDraw", "Draw a polygon on the drawing plane, and extrude it. Left-click to place polygon vertices. Hold Shift to ignore snapping while drawing."),
 		EToolMessageLevel::UserNotification);
 }
 
