MeshModelingTools: In cube grid tool, allow Ctrl+drag extrude setting to work for corner mode the same way it does in normal mode.

#rnx

diff --git a/Engine/Plugins/Experimental/MeshModelingToolsetExp/Source/MeshModelingToolsExp/Private/CubeGridTool.cpp b/Engine/Plugins/Experimental/MeshModelingToolsetExp/Source/MeshModelingToolsExp/Private/CubeGridTool.cpp
--- a/Engine/Plugins/Experimental/MeshModelingToolsetExp/Source/MeshModelingToolsExp/Private/CubeGridTool.cpp
+++ b/Engine/Plugins/Experimental/MeshModelingToolsetExp/Source/MeshModelingToolsExp/Private/CubeGridTool.cpp
@@ -1303,18 +1303,8 @@ void UCubeGridTool::OnClickPress(const FInputDeviceRay& PressPos)
 	using namespace CubeGridToolLocals;
 
 	UpdateHoverLineSet(false, HoveredSelectionBox); // clear hover
-	
-	if (Mode == EMode::Corner)
-	{
-		MouseState = EMouseState::DraggingCornerSelection;
-		for (int i = 0; i < 4; ++i)
-		{
-			PreDragCornerSelectedFlags[i] = CornerSelectedFlags[i];
-		}
-		AttemptToSelectCorner((FRay3d)PressPos.WorldRay);
-		return;
-	}
 
+	// Ctrl+drag extrude setting works in both corner mode and regular mode
 	if (bMouseDragShouldPushPull)
 	{
 		MouseState = EMouseState::DraggingExtrudeDistance;
@@ -1326,7 +1316,20 @@ void UCubeGridTool::OnClickPress(const FInputDeviceRay& PressPos)
 			FDistLine3Ray3d DistanceCalculator(UE::Geometry::FLine3d(DragProjectionAxis.Origin, DragProjectionAxis.Direction), (FRay3d)PressPos.WorldRay);
 			DistanceCalculator.ComputeResult();
 			DragProjectedStartParam = DistanceCalculator.LineParameter;
+			DragStartExtrudeAmount = CurrentExtrudeAmount;
+		}
+		return;
+	}
+
+	// Deal with corner selection if in corner mode
+	if (Mode == EMode::Corner)
+	{
+		MouseState = EMouseState::DraggingCornerSelection;
+		for (int i = 0; i < 4; ++i)
+		{
+			PreDragCornerSelectedFlags[i] = CornerSelectedFlags[i];
 		}
+		AttemptToSelectCorner((FRay3d)PressPos.WorldRay);
 		return;
 	}
 
@@ -1383,7 +1386,8 @@ void UCubeGridTool::OnClickDrag(const FInputDeviceRay& DragPos)
 
 		double ParamDelta = DistanceCalculator.LineParameter - DragProjectedStartParam;
 		double CubeSize = CubeGrid->GetCurrentGridCellSize();
-		int32 NewExtrudeAmount = FMath::RoundToInt(ParamDelta / (CubeSize * Settings->BlocksPerStep)) * Settings->BlocksPerStep;
+		int32 NewExtrudeDelta = FMath::RoundToInt(ParamDelta / (CubeSize * Settings->BlocksPerStep)) * Settings->BlocksPerStep;
+		int32 NewExtrudeAmount = DragStartExtrudeAmount + NewExtrudeDelta;
 		if (NewExtrudeAmount != CurrentExtrudeAmount)
 		{
 			CurrentExtrudeAmount = NewExtrudeAmount;
@@ -1409,11 +1413,18 @@ void UCubeGridTool::OnClickDrag(const FInputDeviceRay& DragPos)
 
 void UCubeGridTool::OnClickRelease(const FInputDeviceRay& ReleasePos)
 {
-	if (MouseState == EMouseState::DraggingExtrudeDistance && CurrentExtrudeAmount != 0)
+	if (MouseState == EMouseState::DraggingExtrudeDistance)
 	{
-		bWaitingToApplyPreview = true;
-		bBlockUntilPreviewUpdate = false;
-		bAdjustSelectionOnPreviewUpdate = true;
+		// Only apply result if we're not in corner mode, because in corner mode
+		// we apply when exiting corner mode (that behavior is particularly important 
+		// when using E/Q to set extrude distance, to allow different slopes to be
+		// set).
+		if (Mode != EMode::Corner && CurrentExtrudeAmount != 0)
+		{
+			bWaitingToApplyPreview = true;
+			bBlockUntilPreviewUpdate = false;
+			bAdjustSelectionOnPreviewUpdate = true;
+		}
 	}
 	else if (MouseState == EMouseState::DraggingRegularSelection)
 	{
@@ -1425,11 +1436,15 @@ void UCubeGridTool::OnClickRelease(const FInputDeviceRay& ReleasePos)
 
 void UCubeGridTool::OnTerminateDragSequence()
 {
-	if (MouseState == EMouseState::DraggingExtrudeDistance && CurrentExtrudeAmount != 0)
+	if (MouseState == EMouseState::DraggingExtrudeDistance)
 	{
-		bWaitingToApplyPreview = true;
-		bBlockUntilPreviewUpdate = false;
-		bAdjustSelectionOnPreviewUpdate = true;
+		// Only apply result if we're not in corner mode
+		if (Mode != EMode::Corner && CurrentExtrudeAmount != 0)
+		{
+			bWaitingToApplyPreview = true;
+			bBlockUntilPreviewUpdate = false;
+			bAdjustSelectionOnPreviewUpdate = true;
+		}
 	}
 
 	MouseState = EMouseState::NotDragging;
diff --git a/Engine/Plugins/Experimental/MeshModelingToolsetExp/Source/MeshModelingToolsExp/Public/CubeGridTool.h b/Engine/Plugins/Experimental/MeshModelingToolsetExp/Source/MeshModelingToolsExp/Public/CubeGridTool.h
--- a/Engine/Plugins/Experimental/MeshModelingToolsetExp/Source/MeshModelingToolsExp/Public/CubeGridTool.h
+++ b/Engine/Plugins/Experimental/MeshModelingToolsetExp/Source/MeshModelingToolsExp/Public/CubeGridTool.h
@@ -430,6 +430,7 @@ protected:
 	bool bMouseDragShouldPushPull = false;
 	FRay3d DragProjectionAxis;
 	double DragProjectedStartParam;
+	int32 DragStartExtrudeAmount = 0;
 
 	void ApplyFlipSelection();
 	void ApplySlide(int32 NumBlocks);
