ModelingTools: Clear TargetWorld member during Shutdown of SingleSelectionMeshEditingTool and MultiSelectionMeshEditingTool

#rnx

diff --git a/Engine/Plugins/Experimental/HairModelingToolset/Source/HairModelingToolset/Private/GenerateLODMeshesTool.cpp b/Engine/Plugins/Experimental/HairModelingToolset/Source/HairModelingToolset/Private/GenerateLODMeshesTool.cpp
--- a/Engine/Plugins/Experimental/HairModelingToolset/Source/HairModelingToolset/Private/GenerateLODMeshesTool.cpp
+++ b/Engine/Plugins/Experimental/HairModelingToolset/Source/HairModelingToolset/Private/GenerateLODMeshesTool.cpp
@@ -138,7 +138,7 @@ void UGenerateLODMeshesTool::Setup()
 }
 
 
-void UGenerateLODMeshesTool::Shutdown(EToolShutdownType ShutdownType)
+void UGenerateLODMeshesTool::OnShutdown(EToolShutdownType ShutdownType)
 {
 	SimplifyProperties->SaveProperties(this);
 	Cast<IPrimitiveComponentBackedTarget>(Target)->SetOwnerVisibility(true);
diff --git a/Engine/Plugins/Experimental/HairModelingToolset/Source/HairModelingToolset/Public/GenerateLODMeshesTool.h b/Engine/Plugins/Experimental/HairModelingToolset/Source/HairModelingToolset/Public/GenerateLODMeshesTool.h
--- a/Engine/Plugins/Experimental/HairModelingToolset/Source/HairModelingToolset/Public/GenerateLODMeshesTool.h
+++ b/Engine/Plugins/Experimental/HairModelingToolset/Source/HairModelingToolset/Public/GenerateLODMeshesTool.h
@@ -152,7 +152,7 @@ class HAIRMODELINGTOOLSET_API UGenerateLODMeshesTool : public USingleSelectionMe
 
 public:
 	virtual void Setup() override;
-	virtual void Shutdown(EToolShutdownType ShutdownType) override;
+	virtual void OnShutdown(EToolShutdownType ShutdownType) override;
 
 	virtual void OnTick(float DeltaTime) override;
 	virtual void Render(IToolsContextRenderAPI* RenderAPI) override;
diff --git a/Engine/Plugins/Experimental/MeshLODToolset/Source/MeshLODToolset/Private/Tools/GenerateStaticMeshLODAssetTool.cpp b/Engine/Plugins/Experimental/MeshLODToolset/Source/MeshLODToolset/Private/Tools/GenerateStaticMeshLODAssetTool.cpp
--- a/Engine/Plugins/Experimental/MeshLODToolset/Source/MeshLODToolset/Private/Tools/GenerateStaticMeshLODAssetTool.cpp
+++ b/Engine/Plugins/Experimental/MeshLODToolset/Source/MeshLODToolset/Private/Tools/GenerateStaticMeshLODAssetTool.cpp
@@ -456,7 +456,7 @@ void UGenerateStaticMeshLODAssetTool::OnSettingsModified()
 
 
 
-void UGenerateStaticMeshLODAssetTool::Shutdown(EToolShutdownType ShutdownType)
+void UGenerateStaticMeshLODAssetTool::OnShutdown(EToolShutdownType ShutdownType)
 {
 	if (OutputProperties)
 	{
diff --git a/Engine/Plugins/Experimental/MeshLODToolset/Source/MeshLODToolset/Private/Tools/LODManagerTool.cpp b/Engine/Plugins/Experimental/MeshLODToolset/Source/MeshLODToolset/Private/Tools/LODManagerTool.cpp
--- a/Engine/Plugins/Experimental/MeshLODToolset/Source/MeshLODToolset/Private/Tools/LODManagerTool.cpp
+++ b/Engine/Plugins/Experimental/MeshLODToolset/Source/MeshLODToolset/Private/Tools/LODManagerTool.cpp
@@ -134,7 +134,7 @@ void ULODManagerTool::Setup()
 
 
 
-void ULODManagerTool::Shutdown(EToolShutdownType ShutdownType)
+void ULODManagerTool::OnShutdown(EToolShutdownType ShutdownType)
 {
 	if (LODPreview)
 	{
diff --git a/Engine/Plugins/Experimental/MeshLODToolset/Source/MeshLODToolset/Public/Tools/GenerateStaticMeshLODAssetTool.h b/Engine/Plugins/Experimental/MeshLODToolset/Source/MeshLODToolset/Public/Tools/GenerateStaticMeshLODAssetTool.h
--- a/Engine/Plugins/Experimental/MeshLODToolset/Source/MeshLODToolset/Public/Tools/GenerateStaticMeshLODAssetTool.h
+++ b/Engine/Plugins/Experimental/MeshLODToolset/Source/MeshLODToolset/Public/Tools/GenerateStaticMeshLODAssetTool.h
@@ -270,7 +270,7 @@ public:
 	virtual void SetUseAssetEditorMode(bool bEnable);
 
 	virtual void Setup() override;
-	virtual void Shutdown(EToolShutdownType ShutdownType) override;
+	virtual void OnShutdown(EToolShutdownType ShutdownType) override;
 
 	virtual void OnTick(float DeltaTime);
 
diff --git a/Engine/Plugins/Experimental/MeshLODToolset/Source/MeshLODToolset/Public/Tools/LODManagerTool.h b/Engine/Plugins/Experimental/MeshLODToolset/Source/MeshLODToolset/Public/Tools/LODManagerTool.h
--- a/Engine/Plugins/Experimental/MeshLODToolset/Source/MeshLODToolset/Public/Tools/LODManagerTool.h
+++ b/Engine/Plugins/Experimental/MeshLODToolset/Source/MeshLODToolset/Public/Tools/LODManagerTool.h
@@ -176,7 +176,7 @@ public:
 	ULODManagerTool();
 
 	virtual void Setup() override;
-	virtual void Shutdown(EToolShutdownType ShutdownType) override;
+	virtual void OnShutdown(EToolShutdownType ShutdownType) override;
 	virtual void OnTick(float DeltaTime) override;
 
 	virtual bool HasCancel() const override { return false; }
diff --git a/Engine/Plugins/Experimental/MeshModelingToolsetExp/Source/MeshModelingToolsEditorOnly/Private/AddPivotActorTool.cpp b/Engine/Plugins/Experimental/MeshModelingToolsetExp/Source/MeshModelingToolsEditorOnly/Private/AddPivotActorTool.cpp
--- a/Engine/Plugins/Experimental/MeshModelingToolsetExp/Source/MeshModelingToolsEditorOnly/Private/AddPivotActorTool.cpp
+++ b/Engine/Plugins/Experimental/MeshModelingToolsetExp/Source/MeshModelingToolsEditorOnly/Private/AddPivotActorTool.cpp
@@ -148,7 +148,7 @@ void UAddPivotActorTool::Setup()
 	DragAlignmentMechanic->AddToGizmo(TransformGizmo);
 }
 
-void UAddPivotActorTool::Shutdown(EToolShutdownType ShutdownType)
+void UAddPivotActorTool::OnShutdown(EToolShutdownType ShutdownType)
 {
 	using namespace AddPivotActorToolLocals;
 
diff --git a/Engine/Plugins/Experimental/MeshModelingToolsetExp/Source/MeshModelingToolsEditorOnly/Private/AttributeEditorTool.cpp b/Engine/Plugins/Experimental/MeshModelingToolsetExp/Source/MeshModelingToolsEditorOnly/Private/AttributeEditorTool.cpp
--- a/Engine/Plugins/Experimental/MeshModelingToolsetExp/Source/MeshModelingToolsEditorOnly/Private/AttributeEditorTool.cpp
+++ b/Engine/Plugins/Experimental/MeshModelingToolsetExp/Source/MeshModelingToolsEditorOnly/Private/AttributeEditorTool.cpp
@@ -358,7 +358,7 @@ void UAttributeEditorTool::InitializeAttributeLists()
 
 
 
-void UAttributeEditorTool::Shutdown(EToolShutdownType ShutdownType)
+void UAttributeEditorTool::OnShutdown(EToolShutdownType ShutdownType)
 {
 }
 
diff --git a/Engine/Plugins/Experimental/MeshModelingToolsetExp/Source/MeshModelingToolsEditorOnly/Private/MergeMeshesTool.cpp b/Engine/Plugins/Experimental/MeshModelingToolsetExp/Source/MeshModelingToolsEditorOnly/Private/MergeMeshesTool.cpp
--- a/Engine/Plugins/Experimental/MeshModelingToolsetExp/Source/MeshModelingToolsEditorOnly/Private/MergeMeshesTool.cpp
+++ b/Engine/Plugins/Experimental/MeshModelingToolsetExp/Source/MeshModelingToolsEditorOnly/Private/MergeMeshesTool.cpp
@@ -109,7 +109,7 @@ void UMergeMeshesTool::Setup()
 }
 
 
-void UMergeMeshesTool::Shutdown(EToolShutdownType ShutdownType)
+void UMergeMeshesTool::OnShutdown(EToolShutdownType ShutdownType)
 {
 	MergeProps->SaveProperties(this);
 	HandleSourcesProperties->SaveProperties(this);
diff --git a/Engine/Plugins/Experimental/MeshModelingToolsetExp/Source/MeshModelingToolsEditorOnly/Private/MeshTangentsTool.cpp b/Engine/Plugins/Experimental/MeshModelingToolsetExp/Source/MeshModelingToolsEditorOnly/Private/MeshTangentsTool.cpp
--- a/Engine/Plugins/Experimental/MeshModelingToolsetExp/Source/MeshModelingToolsEditorOnly/Private/MeshTangentsTool.cpp
+++ b/Engine/Plugins/Experimental/MeshModelingToolsetExp/Source/MeshModelingToolsEditorOnly/Private/MeshTangentsTool.cpp
@@ -111,7 +111,7 @@ void UMeshTangentsTool::Setup()
 
 
 
-void UMeshTangentsTool::Shutdown(EToolShutdownType ShutdownType)
+void UMeshTangentsTool::OnShutdown(EToolShutdownType ShutdownType)
 {
 	PreviewGeometry->Disconnect();
 	PreviewMesh->Disconnect();
diff --git a/Engine/Plugins/Experimental/MeshModelingToolsetExp/Source/MeshModelingToolsEditorOnly/Private/MeshToVolumeTool.cpp b/Engine/Plugins/Experimental/MeshModelingToolsetExp/Source/MeshModelingToolsEditorOnly/Private/MeshToVolumeTool.cpp
--- a/Engine/Plugins/Experimental/MeshModelingToolsetExp/Source/MeshModelingToolsEditorOnly/Private/MeshToVolumeTool.cpp
+++ b/Engine/Plugins/Experimental/MeshModelingToolsetExp/Source/MeshModelingToolsEditorOnly/Private/MeshToVolumeTool.cpp
@@ -91,7 +91,7 @@ void UMeshToVolumeTool::Setup()
 		EToolMessageLevel::UserNotification);
 }
 
-void UMeshToVolumeTool::Shutdown(EToolShutdownType ShutdownType)
+void UMeshToVolumeTool::OnShutdown(EToolShutdownType ShutdownType)
 {
 	Settings->SaveProperties(this);
 	HandleSourcesProperties->SaveProperties(this);
diff --git a/Engine/Plugins/Experimental/MeshModelingToolsetExp/Source/MeshModelingToolsEditorOnly/Private/ParameterizeMeshTool.cpp b/Engine/Plugins/Experimental/MeshModelingToolsetExp/Source/MeshModelingToolsEditorOnly/Private/ParameterizeMeshTool.cpp
--- a/Engine/Plugins/Experimental/MeshModelingToolsetExp/Source/MeshModelingToolsEditorOnly/Private/ParameterizeMeshTool.cpp
+++ b/Engine/Plugins/Experimental/MeshModelingToolsetExp/Source/MeshModelingToolsEditorOnly/Private/ParameterizeMeshTool.cpp
@@ -149,7 +149,7 @@ void UParameterizeMeshTool::OnMethodTypeChanged()
 }
 
 
-void UParameterizeMeshTool::Shutdown(EToolShutdownType ShutdownType)
+void UParameterizeMeshTool::OnShutdown(EToolShutdownType ShutdownType)
 {
 	if (UVLayoutView)
 	{
diff --git a/Engine/Plugins/Experimental/MeshModelingToolsetExp/Source/MeshModelingToolsEditorOnly/Private/PolygonOnMeshTool.cpp b/Engine/Plugins/Experimental/MeshModelingToolsetExp/Source/MeshModelingToolsEditorOnly/Private/PolygonOnMeshTool.cpp
--- a/Engine/Plugins/Experimental/MeshModelingToolsetExp/Source/MeshModelingToolsEditorOnly/Private/PolygonOnMeshTool.cpp
+++ b/Engine/Plugins/Experimental/MeshModelingToolsetExp/Source/MeshModelingToolsEditorOnly/Private/PolygonOnMeshTool.cpp
@@ -228,7 +228,7 @@ void UPolygonOnMeshTool::SetupPreview()
 }
 
 
-void UPolygonOnMeshTool::Shutdown(EToolShutdownType ShutdownType)
+void UPolygonOnMeshTool::OnShutdown(EToolShutdownType ShutdownType)
 {
 	PlaneMechanic->Shutdown();
 	if (DrawPolygonMechanic != nullptr)
diff --git a/Engine/Plugins/Experimental/MeshModelingToolsetExp/Source/MeshModelingToolsEditorOnly/Private/SimplifyMeshTool.cpp b/Engine/Plugins/Experimental/MeshModelingToolsetExp/Source/MeshModelingToolsEditorOnly/Private/SimplifyMeshTool.cpp
--- a/Engine/Plugins/Experimental/MeshModelingToolsetExp/Source/MeshModelingToolsEditorOnly/Private/SimplifyMeshTool.cpp
+++ b/Engine/Plugins/Experimental/MeshModelingToolsetExp/Source/MeshModelingToolsEditorOnly/Private/SimplifyMeshTool.cpp
@@ -159,7 +159,7 @@ bool USimplifyMeshTool::CanAccept() const
 }
 
 
-void USimplifyMeshTool::Shutdown(EToolShutdownType ShutdownType)
+void USimplifyMeshTool::OnShutdown(EToolShutdownType ShutdownType)
 {
 	SimplifyProperties->SaveProperties(this);
 	
diff --git a/Engine/Plugins/Experimental/MeshModelingToolsetExp/Source/MeshModelingToolsEditorOnly/Private/SubdividePolyTool.cpp b/Engine/Plugins/Experimental/MeshModelingToolsetExp/Source/MeshModelingToolsEditorOnly/Private/SubdividePolyTool.cpp
--- a/Engine/Plugins/Experimental/MeshModelingToolsetExp/Source/MeshModelingToolsEditorOnly/Private/SubdividePolyTool.cpp
+++ b/Engine/Plugins/Experimental/MeshModelingToolsetExp/Source/MeshModelingToolsEditorOnly/Private/SubdividePolyTool.cpp
@@ -377,7 +377,7 @@ void USubdividePolyTool::CreateOrUpdatePreviewGeometry()
 }
 
 
-void USubdividePolyTool::Shutdown(EToolShutdownType ShutdownType)
+void USubdividePolyTool::OnShutdown(EToolShutdownType ShutdownType)
 {
 	if (Properties)
 	{
diff --git a/Engine/Plugins/Experimental/MeshModelingToolsetExp/Source/MeshModelingToolsEditorOnly/Private/VoxelCSGMeshesTool.cpp b/Engine/Plugins/Experimental/MeshModelingToolsetExp/Source/MeshModelingToolsEditorOnly/Private/VoxelCSGMeshesTool.cpp
--- a/Engine/Plugins/Experimental/MeshModelingToolsetExp/Source/MeshModelingToolsEditorOnly/Private/VoxelCSGMeshesTool.cpp
+++ b/Engine/Plugins/Experimental/MeshModelingToolsetExp/Source/MeshModelingToolsEditorOnly/Private/VoxelCSGMeshesTool.cpp
@@ -94,7 +94,7 @@ void UVoxelCSGMeshesTool::Setup()
 }
 
 
-void UVoxelCSGMeshesTool::Shutdown(EToolShutdownType ShutdownType)
+void UVoxelCSGMeshesTool::OnShutdown(EToolShutdownType ShutdownType)
 {
 	CSGProps->SaveProperties(this);
 	HandleSourcesProperties->SaveProperties(this);
diff --git a/Engine/Plugins/Experimental/MeshModelingToolsetExp/Source/MeshModelingToolsEditorOnly/Public/AddPivotActorTool.h b/Engine/Plugins/Experimental/MeshModelingToolsetExp/Source/MeshModelingToolsEditorOnly/Public/AddPivotActorTool.h
--- a/Engine/Plugins/Experimental/MeshModelingToolsetExp/Source/MeshModelingToolsEditorOnly/Public/AddPivotActorTool.h
+++ b/Engine/Plugins/Experimental/MeshModelingToolsetExp/Source/MeshModelingToolsEditorOnly/Public/AddPivotActorTool.h
@@ -41,7 +41,7 @@ public:
 	virtual void SetPivotRepositionMode(AActor* PivotActor) { ExistingPivotActor = PivotActor; }
 
 	virtual void Setup() override;
-	virtual void Shutdown(EToolShutdownType ShutdownType) override;
+	virtual void OnShutdown(EToolShutdownType ShutdownType) override;
 
 	virtual void Render(IToolsContextRenderAPI* RenderAPI) override;
 
diff --git a/Engine/Plugins/Experimental/MeshModelingToolsetExp/Source/MeshModelingToolsEditorOnly/Public/AttributeEditorTool.h b/Engine/Plugins/Experimental/MeshModelingToolsetExp/Source/MeshModelingToolsEditorOnly/Public/AttributeEditorTool.h
--- a/Engine/Plugins/Experimental/MeshModelingToolsetExp/Source/MeshModelingToolsEditorOnly/Public/AttributeEditorTool.h
+++ b/Engine/Plugins/Experimental/MeshModelingToolsetExp/Source/MeshModelingToolsEditorOnly/Public/AttributeEditorTool.h
@@ -336,7 +336,7 @@ public:
 	UAttributeEditorTool();
 
 	virtual void Setup() override;
-	virtual void Shutdown(EToolShutdownType ShutdownType) override;
+	virtual void OnShutdown(EToolShutdownType ShutdownType) override;
 	virtual void OnTick(float DeltaTime) override;
 
 	virtual bool HasCancel() const override { return false; }
diff --git a/Engine/Plugins/Experimental/MeshModelingToolsetExp/Source/MeshModelingToolsEditorOnly/Public/MergeMeshesTool.h b/Engine/Plugins/Experimental/MeshModelingToolsetExp/Source/MeshModelingToolsEditorOnly/Public/MergeMeshesTool.h
--- a/Engine/Plugins/Experimental/MeshModelingToolsetExp/Source/MeshModelingToolsEditorOnly/Public/MergeMeshesTool.h
+++ b/Engine/Plugins/Experimental/MeshModelingToolsetExp/Source/MeshModelingToolsEditorOnly/Public/MergeMeshesTool.h
@@ -67,7 +67,7 @@ public:
 	UMergeMeshesTool();
 
 	virtual void Setup() override;
-	virtual void Shutdown(EToolShutdownType ShutdownType) override;
+	virtual void OnShutdown(EToolShutdownType ShutdownType) override;
 
 	virtual void OnTick(float DeltaTime) override;
 
diff --git a/Engine/Plugins/Experimental/MeshModelingToolsetExp/Source/MeshModelingToolsEditorOnly/Public/MeshTangentsTool.h b/Engine/Plugins/Experimental/MeshModelingToolsetExp/Source/MeshModelingToolsEditorOnly/Public/MeshTangentsTool.h
--- a/Engine/Plugins/Experimental/MeshModelingToolsetExp/Source/MeshModelingToolsEditorOnly/Public/MeshTangentsTool.h
+++ b/Engine/Plugins/Experimental/MeshModelingToolsetExp/Source/MeshModelingToolsEditorOnly/Public/MeshTangentsTool.h
@@ -96,7 +96,7 @@ public:
 	UMeshTangentsTool();
 
 	virtual void Setup() override;
-	virtual void Shutdown(EToolShutdownType ShutdownType) override;
+	virtual void OnShutdown(EToolShutdownType ShutdownType) override;
 
 	virtual void OnTick(float DeltaTime) override;
 	virtual void Render(IToolsContextRenderAPI* RenderAPI) override;
diff --git a/Engine/Plugins/Experimental/MeshModelingToolsetExp/Source/MeshModelingToolsEditorOnly/Public/MeshToVolumeTool.h b/Engine/Plugins/Experimental/MeshModelingToolsetExp/Source/MeshModelingToolsEditorOnly/Public/MeshToVolumeTool.h
--- a/Engine/Plugins/Experimental/MeshModelingToolsetExp/Source/MeshModelingToolsEditorOnly/Public/MeshToVolumeTool.h
+++ b/Engine/Plugins/Experimental/MeshModelingToolsetExp/Source/MeshModelingToolsEditorOnly/Public/MeshToVolumeTool.h
@@ -77,7 +77,7 @@ public:
 	UMeshToVolumeTool();
 
 	virtual void Setup() override;
-	virtual void Shutdown(EToolShutdownType ShutdownType) override;
+	virtual void OnShutdown(EToolShutdownType ShutdownType) override;
 
 	virtual void Render(IToolsContextRenderAPI* RenderAPI) override;
 	virtual void OnTick(float DeltaTime) override;
diff --git a/Engine/Plugins/Experimental/MeshModelingToolsetExp/Source/MeshModelingToolsEditorOnly/Public/ParameterizeMeshTool.h b/Engine/Plugins/Experimental/MeshModelingToolsetExp/Source/MeshModelingToolsEditorOnly/Public/ParameterizeMeshTool.h
--- a/Engine/Plugins/Experimental/MeshModelingToolsetExp/Source/MeshModelingToolsEditorOnly/Public/ParameterizeMeshTool.h
+++ b/Engine/Plugins/Experimental/MeshModelingToolsetExp/Source/MeshModelingToolsEditorOnly/Public/ParameterizeMeshTool.h
@@ -54,7 +54,7 @@ class MESHMODELINGTOOLSEDITORONLY_API UParameterizeMeshTool : public USingleSele
 
 public:
 	virtual void Setup() override;
-	virtual void Shutdown(EToolShutdownType ShutdownType) override;
+	virtual void OnShutdown(EToolShutdownType ShutdownType) override;
 
 	virtual void Render(IToolsContextRenderAPI* RenderAPI) override;
 	virtual void OnTick(float DeltaTime) override;
diff --git a/Engine/Plugins/Experimental/MeshModelingToolsetExp/Source/MeshModelingToolsEditorOnly/Public/PolygonOnMeshTool.h b/Engine/Plugins/Experimental/MeshModelingToolsetExp/Source/MeshModelingToolsEditorOnly/Public/PolygonOnMeshTool.h
--- a/Engine/Plugins/Experimental/MeshModelingToolsetExp/Source/MeshModelingToolsEditorOnly/Public/PolygonOnMeshTool.h
+++ b/Engine/Plugins/Experimental/MeshModelingToolsetExp/Source/MeshModelingToolsEditorOnly/Public/PolygonOnMeshTool.h
@@ -152,7 +152,7 @@ public:
 	UPolygonOnMeshTool();
 
 	virtual void Setup() override;
-	virtual void Shutdown(EToolShutdownType ShutdownType) override;
+	virtual void OnShutdown(EToolShutdownType ShutdownType) override;
 
 	virtual void SetWorld(UWorld* World);
 
diff --git a/Engine/Plugins/Experimental/MeshModelingToolsetExp/Source/MeshModelingToolsEditorOnly/Public/SimplifyMeshTool.h b/Engine/Plugins/Experimental/MeshModelingToolsetExp/Source/MeshModelingToolsEditorOnly/Public/SimplifyMeshTool.h
--- a/Engine/Plugins/Experimental/MeshModelingToolsetExp/Source/MeshModelingToolsEditorOnly/Public/SimplifyMeshTool.h
+++ b/Engine/Plugins/Experimental/MeshModelingToolsetExp/Source/MeshModelingToolsEditorOnly/Public/SimplifyMeshTool.h
@@ -111,7 +111,7 @@ class MESHMODELINGTOOLSEDITORONLY_API USimplifyMeshTool : public USingleSelectio
 
 public:
 	virtual void Setup() override;
-	virtual void Shutdown(EToolShutdownType ShutdownType) override;
+	virtual void OnShutdown(EToolShutdownType ShutdownType) override;
 
 	virtual void OnTick(float DeltaTime) override;
 
diff --git a/Engine/Plugins/Experimental/MeshModelingToolsetExp/Source/MeshModelingToolsEditorOnly/Public/SubdividePolyTool.h b/Engine/Plugins/Experimental/MeshModelingToolsetExp/Source/MeshModelingToolsEditorOnly/Public/SubdividePolyTool.h
--- a/Engine/Plugins/Experimental/MeshModelingToolsetExp/Source/MeshModelingToolsEditorOnly/Public/SubdividePolyTool.h
+++ b/Engine/Plugins/Experimental/MeshModelingToolsetExp/Source/MeshModelingToolsEditorOnly/Public/SubdividePolyTool.h
@@ -78,7 +78,7 @@ class MESHMODELINGTOOLSEDITORONLY_API USubdividePolyTool : public USingleSelecti
 public:
 
 	virtual void Setup() override;
-	virtual void Shutdown(EToolShutdownType ShutdownType) override;
+	virtual void OnShutdown(EToolShutdownType ShutdownType) override;
 
 	virtual bool HasCancel() const override { return true; }
 	virtual bool HasAccept() const override { return true; }
diff --git a/Engine/Plugins/Experimental/MeshModelingToolsetExp/Source/MeshModelingToolsEditorOnly/Public/VoxelCSGMeshesTool.h b/Engine/Plugins/Experimental/MeshModelingToolsetExp/Source/MeshModelingToolsEditorOnly/Public/VoxelCSGMeshesTool.h
--- a/Engine/Plugins/Experimental/MeshModelingToolsetExp/Source/MeshModelingToolsEditorOnly/Public/VoxelCSGMeshesTool.h
+++ b/Engine/Plugins/Experimental/MeshModelingToolsetExp/Source/MeshModelingToolsEditorOnly/Public/VoxelCSGMeshesTool.h
@@ -89,7 +89,7 @@ public:
 	UVoxelCSGMeshesTool();
 
 	virtual void Setup() override;
-	virtual void Shutdown(EToolShutdownType ShutdownType) override;
+	virtual void OnShutdown(EToolShutdownType ShutdownType) override;
 
 	virtual void OnTick(float DeltaTime) override;
 	virtual void Render(IToolsContextRenderAPI* RenderAPI) override;
diff --git a/Engine/Plugins/Experimental/MeshModelingToolsetExp/Source/MeshModelingToolsExp/Private/AlignObjectsTool.cpp b/Engine/Plugins/Experimental/MeshModelingToolsetExp/Source/MeshModelingToolsExp/Private/AlignObjectsTool.cpp
--- a/Engine/Plugins/Experimental/MeshModelingToolsetExp/Source/MeshModelingToolsExp/Private/AlignObjectsTool.cpp
+++ b/Engine/Plugins/Experimental/MeshModelingToolsetExp/Source/MeshModelingToolsExp/Private/AlignObjectsTool.cpp
@@ -73,7 +73,7 @@ void UAlignObjectsTool::Setup()
 
 
 
-void UAlignObjectsTool::Shutdown(EToolShutdownType ShutdownType)
+void UAlignObjectsTool::OnShutdown(EToolShutdownType ShutdownType)
 {
 	AlignProps->SaveProperties(this);
 
diff --git a/Engine/Plugins/Experimental/MeshModelingToolsetExp/Source/MeshModelingToolsExp/Private/BakeMeshAttributeMapsTool.cpp b/Engine/Plugins/Experimental/MeshModelingToolsetExp/Source/MeshModelingToolsExp/Private/BakeMeshAttributeMapsTool.cpp
--- a/Engine/Plugins/Experimental/MeshModelingToolsetExp/Source/MeshModelingToolsExp/Private/BakeMeshAttributeMapsTool.cpp
+++ b/Engine/Plugins/Experimental/MeshModelingToolsetExp/Source/MeshModelingToolsExp/Private/BakeMeshAttributeMapsTool.cpp
@@ -459,11 +459,11 @@ TUniquePtr<UE::Geometry::TGenericDataOperator<FMeshMapBaker>> UBakeMeshAttribute
 }
 
 
-void UBakeMeshAttributeMapsTool::Shutdown(EToolShutdownType ShutdownType)
+void UBakeMeshAttributeMapsTool::OnShutdown(EToolShutdownType ShutdownType)
 {
 	TRACE_CPUPROFILER_EVENT_SCOPE(UBakeMeshAttributeMapsTool::Shutdown);
 	
-	Super::Shutdown(ShutdownType);
+	Super::OnShutdown(ShutdownType);
 	
 	Settings->SaveProperties(this);
 	InputMeshSettings->SaveProperties(this);
diff --git a/Engine/Plugins/Experimental/MeshModelingToolsetExp/Source/MeshModelingToolsExp/Private/BakeMeshAttributeMapsToolBase.cpp b/Engine/Plugins/Experimental/MeshModelingToolsetExp/Source/MeshModelingToolsExp/Private/BakeMeshAttributeMapsToolBase.cpp
--- a/Engine/Plugins/Experimental/MeshModelingToolsetExp/Source/MeshModelingToolsExp/Private/BakeMeshAttributeMapsToolBase.cpp
+++ b/Engine/Plugins/Experimental/MeshModelingToolsetExp/Source/MeshModelingToolsExp/Private/BakeMeshAttributeMapsToolBase.cpp
@@ -87,7 +87,7 @@ void UBakeMeshAttributeMapsToolBase::PostSetup()
 	GatherAnalytics(BakeAnalytics.MeshSettings);
 }
 
-void UBakeMeshAttributeMapsToolBase::Shutdown(EToolShutdownType ShutdownType)
+void UBakeMeshAttributeMapsToolBase::OnShutdown(EToolShutdownType ShutdownType)
 {
 	VisualizationProps->SaveProperties(this);
 
diff --git a/Engine/Plugins/Experimental/MeshModelingToolsetExp/Source/MeshModelingToolsExp/Private/BakeMeshAttributeVertexTool.cpp b/Engine/Plugins/Experimental/MeshModelingToolsetExp/Source/MeshModelingToolsExp/Private/BakeMeshAttributeVertexTool.cpp
--- a/Engine/Plugins/Experimental/MeshModelingToolsetExp/Source/MeshModelingToolsExp/Private/BakeMeshAttributeVertexTool.cpp
+++ b/Engine/Plugins/Experimental/MeshModelingToolsetExp/Source/MeshModelingToolsExp/Private/BakeMeshAttributeVertexTool.cpp
@@ -360,7 +360,7 @@ void UBakeMeshAttributeVertexTool::Setup()
 	GatherAnalytics(BakeAnalytics.MeshSettings);
 }
 
-void UBakeMeshAttributeVertexTool::Shutdown(EToolShutdownType ShutdownType)
+void UBakeMeshAttributeVertexTool::OnShutdown(EToolShutdownType ShutdownType)
 {
 	TRACE_CPUPROFILER_EVENT_SCOPE(UBakeMeshAttributeVertexTool::Shutdown);
 	
diff --git a/Engine/Plugins/Experimental/MeshModelingToolsetExp/Source/MeshModelingToolsExp/Private/BakeMultiMeshAttributeMapsTool.cpp b/Engine/Plugins/Experimental/MeshModelingToolsetExp/Source/MeshModelingToolsExp/Private/BakeMultiMeshAttributeMapsTool.cpp
--- a/Engine/Plugins/Experimental/MeshModelingToolsetExp/Source/MeshModelingToolsExp/Private/BakeMultiMeshAttributeMapsTool.cpp
+++ b/Engine/Plugins/Experimental/MeshModelingToolsetExp/Source/MeshModelingToolsExp/Private/BakeMultiMeshAttributeMapsTool.cpp
@@ -523,11 +523,11 @@ TUniquePtr<UE::Geometry::TGenericDataOperator<FMeshMapBaker>> UBakeMultiMeshAttr
 }
 
 
-void UBakeMultiMeshAttributeMapsTool::Shutdown(EToolShutdownType ShutdownType)
+void UBakeMultiMeshAttributeMapsTool::OnShutdown(EToolShutdownType ShutdownType)
 {
 	TRACE_CPUPROFILER_EVENT_SCOPE(UBakeMultiMeshAttributeMapsTool::Shutdown);
 	
-	Super::Shutdown(ShutdownType);
+	Super::OnShutdown(ShutdownType);
 	
 	Settings->SaveProperties(this);
 	InputMeshSettings->SaveProperties(this);
diff --git a/Engine/Plugins/Experimental/MeshModelingToolsetExp/Source/MeshModelingToolsExp/Private/BakeTransformTool.cpp b/Engine/Plugins/Experimental/MeshModelingToolsetExp/Source/MeshModelingToolsExp/Private/BakeTransformTool.cpp
--- a/Engine/Plugins/Experimental/MeshModelingToolsetExp/Source/MeshModelingToolsExp/Private/BakeTransformTool.cpp
+++ b/Engine/Plugins/Experimental/MeshModelingToolsetExp/Source/MeshModelingToolsExp/Private/BakeTransformTool.cpp
@@ -94,7 +94,7 @@ void UBakeTransformTool::Setup()
 
 
 
-void UBakeTransformTool::Shutdown(EToolShutdownType ShutdownType)
+void UBakeTransformTool::OnShutdown(EToolShutdownType ShutdownType)
 {
 	if (ShutdownType == EToolShutdownType::Accept)
 	{
diff --git a/Engine/Plugins/Experimental/MeshModelingToolsetExp/Source/MeshModelingToolsExp/Private/ConvertMeshesTool.cpp b/Engine/Plugins/Experimental/MeshModelingToolsetExp/Source/MeshModelingToolsExp/Private/ConvertMeshesTool.cpp
--- a/Engine/Plugins/Experimental/MeshModelingToolsetExp/Source/MeshModelingToolsExp/Private/ConvertMeshesTool.cpp
+++ b/Engine/Plugins/Experimental/MeshModelingToolsetExp/Source/MeshModelingToolsExp/Private/ConvertMeshesTool.cpp
@@ -67,7 +67,7 @@ bool UConvertMeshesTool::CanAccept() const
 	return Super::CanAccept();
 }
 
-void UConvertMeshesTool::Shutdown(EToolShutdownType ShutdownType)
+void UConvertMeshesTool::OnShutdown(EToolShutdownType ShutdownType)
 {
 	OutputTypeProperties->SaveProperties(this, TEXT("ConvertMeshesTool"));
 	BasicProperties->SaveProperties(this);
diff --git a/Engine/Plugins/Experimental/MeshModelingToolsetExp/Source/MeshModelingToolsExp/Private/ConvertToPolygonsTool.cpp b/Engine/Plugins/Experimental/MeshModelingToolsetExp/Source/MeshModelingToolsExp/Private/ConvertToPolygonsTool.cpp
--- a/Engine/Plugins/Experimental/MeshModelingToolsetExp/Source/MeshModelingToolsExp/Private/ConvertToPolygonsTool.cpp
+++ b/Engine/Plugins/Experimental/MeshModelingToolsetExp/Source/MeshModelingToolsExp/Private/ConvertToPolygonsTool.cpp
@@ -259,7 +259,7 @@ void UConvertToPolygonsTool::UpdateOpParameters(FConvertToPolygonsOp& ConvertToP
 
 
 
-void UConvertToPolygonsTool::Shutdown(EToolShutdownType ShutdownType)
+void UConvertToPolygonsTool::OnShutdown(EToolShutdownType ShutdownType)
 {
 	Settings->SaveProperties(this);
 	UE::ToolTarget::ShowSourceObject(Target);
diff --git a/Engine/Plugins/Experimental/MeshModelingToolsetExp/Source/MeshModelingToolsExp/Private/DisplaceMeshTool.cpp b/Engine/Plugins/Experimental/MeshModelingToolsetExp/Source/MeshModelingToolsExp/Private/DisplaceMeshTool.cpp
--- a/Engine/Plugins/Experimental/MeshModelingToolsetExp/Source/MeshModelingToolsExp/Private/DisplaceMeshTool.cpp
+++ b/Engine/Plugins/Experimental/MeshModelingToolsetExp/Source/MeshModelingToolsExp/Private/DisplaceMeshTool.cpp
@@ -852,7 +852,7 @@ void UDisplaceMeshTool::Setup()
 		EToolMessageLevel::UserNotification);
 }
 
-void UDisplaceMeshTool::Shutdown(EToolShutdownType ShutdownType)
+void UDisplaceMeshTool::OnShutdown(EToolShutdownType ShutdownType)
 {
 	// unhook any active listener for contrast curve
 	DisconnectActiveContrastCurveTarget();
diff --git a/Engine/Plugins/Experimental/MeshModelingToolsetExp/Source/MeshModelingToolsExp/Private/EditNormalsTool.cpp b/Engine/Plugins/Experimental/MeshModelingToolsetExp/Source/MeshModelingToolsExp/Private/EditNormalsTool.cpp
--- a/Engine/Plugins/Experimental/MeshModelingToolsetExp/Source/MeshModelingToolsExp/Private/EditNormalsTool.cpp
+++ b/Engine/Plugins/Experimental/MeshModelingToolsetExp/Source/MeshModelingToolsExp/Private/EditNormalsTool.cpp
@@ -141,7 +141,7 @@ void UEditNormalsTool::UpdateNumPreviews()
 }
 
 
-void UEditNormalsTool::Shutdown(EToolShutdownType ShutdownType)
+void UEditNormalsTool::OnShutdown(EToolShutdownType ShutdownType)
 {
 	BasicProperties->SaveProperties(this);
 	AdvancedProperties->SaveProperties(this);
diff --git a/Engine/Plugins/Experimental/MeshModelingToolsetExp/Source/MeshModelingToolsExp/Private/EditPivotTool.cpp b/Engine/Plugins/Experimental/MeshModelingToolsetExp/Source/MeshModelingToolsExp/Private/EditPivotTool.cpp
--- a/Engine/Plugins/Experimental/MeshModelingToolsetExp/Source/MeshModelingToolsExp/Private/EditPivotTool.cpp
+++ b/Engine/Plugins/Experimental/MeshModelingToolsetExp/Source/MeshModelingToolsExp/Private/EditPivotTool.cpp
@@ -120,7 +120,7 @@ void UEditPivotTool::Setup()
 
 
 
-void UEditPivotTool::Shutdown(EToolShutdownType ShutdownType)
+void UEditPivotTool::OnShutdown(EToolShutdownType ShutdownType)
 {
 	DragAlignmentMechanic->Shutdown();
 
diff --git a/Engine/Plugins/Experimental/MeshModelingToolsetExp/Source/MeshModelingToolsExp/Private/HoleFillTool.cpp b/Engine/Plugins/Experimental/MeshModelingToolsetExp/Source/MeshModelingToolsExp/Private/HoleFillTool.cpp
--- a/Engine/Plugins/Experimental/MeshModelingToolsetExp/Source/MeshModelingToolsExp/Private/HoleFillTool.cpp
+++ b/Engine/Plugins/Experimental/MeshModelingToolsetExp/Source/MeshModelingToolsExp/Private/HoleFillTool.cpp
@@ -252,7 +252,7 @@ bool UHoleFillTool::CanAccept() const
 	return Super::CanAccept() && Preview->HaveValidResult();
 }
 
-void UHoleFillTool::Shutdown(EToolShutdownType ShutdownType)
+void UHoleFillTool::OnShutdown(EToolShutdownType ShutdownType)
 {
 	Properties->SaveProperties(this);
 	SmoothHoleFillProperties->SaveProperties(this);
diff --git a/Engine/Plugins/Experimental/MeshModelingToolsetExp/Source/MeshModelingToolsExp/Private/LatticeDeformerTool.cpp b/Engine/Plugins/Experimental/MeshModelingToolsetExp/Source/MeshModelingToolsExp/Private/LatticeDeformerTool.cpp
--- a/Engine/Plugins/Experimental/MeshModelingToolsetExp/Source/MeshModelingToolsExp/Private/LatticeDeformerTool.cpp
+++ b/Engine/Plugins/Experimental/MeshModelingToolsetExp/Source/MeshModelingToolsExp/Private/LatticeDeformerTool.cpp
@@ -336,7 +336,7 @@ void ULatticeDeformerTool::SoftDeformLattice()
 }
 
 
-void ULatticeDeformerTool::Shutdown(EToolShutdownType ShutdownType)
+void ULatticeDeformerTool::OnShutdown(EToolShutdownType ShutdownType)
 {
 	Settings->SaveProperties(this);
 	ControlPointsMechanic->Shutdown();
diff --git a/Engine/Plugins/Experimental/MeshModelingToolsetExp/Source/MeshModelingToolsExp/Private/MeshBoundaryToolBase.cpp b/Engine/Plugins/Experimental/MeshModelingToolsetExp/Source/MeshModelingToolsExp/Private/MeshBoundaryToolBase.cpp
--- a/Engine/Plugins/Experimental/MeshModelingToolsetExp/Source/MeshModelingToolsExp/Private/MeshBoundaryToolBase.cpp
+++ b/Engine/Plugins/Experimental/MeshModelingToolsetExp/Source/MeshModelingToolsExp/Private/MeshBoundaryToolBase.cpp
@@ -53,7 +53,7 @@ void UMeshBoundaryToolBase::Setup()
 	SelectionMechanic->OnSelectionChanged.AddUObject(this, &UMeshBoundaryToolBase::OnSelectionChanged);
 }
 
-void UMeshBoundaryToolBase::Shutdown(EToolShutdownType ShutdownType)
+void UMeshBoundaryToolBase::OnShutdown(EToolShutdownType ShutdownType)
 {
 	if (SelectionMechanic)
 	{
diff --git a/Engine/Plugins/Experimental/MeshModelingToolsetExp/Source/MeshModelingToolsExp/Private/MeshInspectorTool.cpp b/Engine/Plugins/Experimental/MeshModelingToolsetExp/Source/MeshModelingToolsExp/Private/MeshInspectorTool.cpp
--- a/Engine/Plugins/Experimental/MeshModelingToolsetExp/Source/MeshModelingToolsExp/Private/MeshInspectorTool.cpp
+++ b/Engine/Plugins/Experimental/MeshModelingToolsetExp/Source/MeshModelingToolsExp/Private/MeshInspectorTool.cpp
@@ -108,7 +108,7 @@ void UMeshInspectorTool::Setup()
 }
 
 
-void UMeshInspectorTool::Shutdown(EToolShutdownType ShutdownType)
+void UMeshInspectorTool::OnShutdown(EToolShutdownType ShutdownType)
 {
 	UE::ToolTarget::ShowSourceObject(Target);
 
diff --git a/Engine/Plugins/Experimental/MeshModelingToolsetExp/Source/MeshModelingToolsExp/Private/MeshSpaceDeformerTool.cpp b/Engine/Plugins/Experimental/MeshModelingToolsetExp/Source/MeshModelingToolsExp/Private/MeshSpaceDeformerTool.cpp
--- a/Engine/Plugins/Experimental/MeshModelingToolsetExp/Source/MeshModelingToolsExp/Private/MeshSpaceDeformerTool.cpp
+++ b/Engine/Plugins/Experimental/MeshModelingToolsetExp/Source/MeshModelingToolsExp/Private/MeshSpaceDeformerTool.cpp
@@ -313,7 +313,7 @@ void UMeshSpaceDeformerTool::Setup()
 		EToolMessageLevel::UserNotification);
 }
 
-void UMeshSpaceDeformerTool::Shutdown(EToolShutdownType ShutdownType)
+void UMeshSpaceDeformerTool::OnShutdown(EToolShutdownType ShutdownType)
 {
 	Settings->SaveProperties(this);
 
diff --git a/Engine/Plugins/Experimental/MeshModelingToolsetExp/Source/MeshModelingToolsExp/Private/MirrorTool.cpp b/Engine/Plugins/Experimental/MeshModelingToolsetExp/Source/MeshModelingToolsExp/Private/MirrorTool.cpp
--- a/Engine/Plugins/Experimental/MeshModelingToolsetExp/Source/MeshModelingToolsExp/Private/MirrorTool.cpp
+++ b/Engine/Plugins/Experimental/MeshModelingToolsetExp/Source/MeshModelingToolsExp/Private/MirrorTool.cpp
@@ -327,7 +327,7 @@ void UMirrorTool::CheckAndDisplayWarnings()
 	}
 }
 
-void UMirrorTool::Shutdown(EToolShutdownType ShutdownType)
+void UMirrorTool::OnShutdown(EToolShutdownType ShutdownType)
 {
 	Settings->SaveProperties(this);
 
diff --git a/Engine/Plugins/Experimental/MeshModelingToolsetExp/Source/MeshModelingToolsExp/Private/Physics/ExtractCollisionGeometryTool.cpp b/Engine/Plugins/Experimental/MeshModelingToolsetExp/Source/MeshModelingToolsExp/Private/Physics/ExtractCollisionGeometryTool.cpp
--- a/Engine/Plugins/Experimental/MeshModelingToolsetExp/Source/MeshModelingToolsExp/Private/Physics/ExtractCollisionGeometryTool.cpp
+++ b/Engine/Plugins/Experimental/MeshModelingToolsetExp/Source/MeshModelingToolsExp/Private/Physics/ExtractCollisionGeometryTool.cpp
@@ -117,7 +117,7 @@ void UExtractCollisionGeometryTool::Setup()
 }
 
 
-void UExtractCollisionGeometryTool::Shutdown(EToolShutdownType ShutdownType)
+void UExtractCollisionGeometryTool::OnShutdown(EToolShutdownType ShutdownType)
 {
 	using FTransform3d = UE::Geometry::FTransform3d;
 	OutputTypeProperties->SaveProperties(this, TEXT("ExtractCollisionTool"));
diff --git a/Engine/Plugins/Experimental/MeshModelingToolsetExp/Source/MeshModelingToolsExp/Private/Physics/PhysicsInspectorTool.cpp b/Engine/Plugins/Experimental/MeshModelingToolsetExp/Source/MeshModelingToolsExp/Private/Physics/PhysicsInspectorTool.cpp
--- a/Engine/Plugins/Experimental/MeshModelingToolsetExp/Source/MeshModelingToolsExp/Private/Physics/PhysicsInspectorTool.cpp
+++ b/Engine/Plugins/Experimental/MeshModelingToolsetExp/Source/MeshModelingToolsExp/Private/Physics/PhysicsInspectorTool.cpp
@@ -82,7 +82,7 @@ void UPhysicsInspectorTool::Setup()
 }
 
 
-void UPhysicsInspectorTool::Shutdown(EToolShutdownType ShutdownType)
+void UPhysicsInspectorTool::OnShutdown(EToolShutdownType ShutdownType)
 {
 	VizSettings->SaveProperties(this);
 
diff --git a/Engine/Plugins/Experimental/MeshModelingToolsetExp/Source/MeshModelingToolsExp/Private/Physics/SetCollisionGeometryTool.cpp b/Engine/Plugins/Experimental/MeshModelingToolsetExp/Source/MeshModelingToolsExp/Private/Physics/SetCollisionGeometryTool.cpp
--- a/Engine/Plugins/Experimental/MeshModelingToolsetExp/Source/MeshModelingToolsExp/Private/Physics/SetCollisionGeometryTool.cpp
+++ b/Engine/Plugins/Experimental/MeshModelingToolsetExp/Source/MeshModelingToolsExp/Private/Physics/SetCollisionGeometryTool.cpp
@@ -150,7 +150,7 @@ void USetCollisionGeometryTool::Setup()
 }
 
 
-void USetCollisionGeometryTool::Shutdown(EToolShutdownType ShutdownType)
+void USetCollisionGeometryTool::OnShutdown(EToolShutdownType ShutdownType)
 {
 	VizSettings->SaveProperties(this);
 	Settings->SaveProperties(this);
diff --git a/Engine/Plugins/Experimental/MeshModelingToolsetExp/Source/MeshModelingToolsExp/Private/PlaneCutTool.cpp b/Engine/Plugins/Experimental/MeshModelingToolsetExp/Source/MeshModelingToolsExp/Private/PlaneCutTool.cpp
--- a/Engine/Plugins/Experimental/MeshModelingToolsetExp/Source/MeshModelingToolsExp/Private/PlaneCutTool.cpp
+++ b/Engine/Plugins/Experimental/MeshModelingToolsetExp/Source/MeshModelingToolsExp/Private/PlaneCutTool.cpp
@@ -242,7 +242,7 @@ void UPlaneCutTool::DoCut()
 
 
 
-void UPlaneCutTool::Shutdown(EToolShutdownType ShutdownType)
+void UPlaneCutTool::OnShutdown(EToolShutdownType ShutdownType)
 {
 	PlaneMechanic->Shutdown();
 	BasicProperties->SaveProperties(this);
diff --git a/Engine/Plugins/Experimental/MeshModelingToolsetExp/Source/MeshModelingToolsExp/Private/RemeshMeshTool.cpp b/Engine/Plugins/Experimental/MeshModelingToolsetExp/Source/MeshModelingToolsExp/Private/RemeshMeshTool.cpp
--- a/Engine/Plugins/Experimental/MeshModelingToolsetExp/Source/MeshModelingToolsExp/Private/RemeshMeshTool.cpp
+++ b/Engine/Plugins/Experimental/MeshModelingToolsetExp/Source/MeshModelingToolsExp/Private/RemeshMeshTool.cpp
@@ -159,7 +159,7 @@ void URemeshMeshTool::Setup()
 		EToolMessageLevel::UserNotification);
 }
 
-void URemeshMeshTool::Shutdown(EToolShutdownType ShutdownType)
+void URemeshMeshTool::OnShutdown(EToolShutdownType ShutdownType)
 {
 	BasicProperties->SaveProperties(this);
 
diff --git a/Engine/Plugins/Experimental/MeshModelingToolsetExp/Source/MeshModelingToolsExp/Private/RemoveOccludedTrianglesTool.cpp b/Engine/Plugins/Experimental/MeshModelingToolsetExp/Source/MeshModelingToolsExp/Private/RemoveOccludedTrianglesTool.cpp
--- a/Engine/Plugins/Experimental/MeshModelingToolsetExp/Source/MeshModelingToolsExp/Private/RemoveOccludedTrianglesTool.cpp
+++ b/Engine/Plugins/Experimental/MeshModelingToolsetExp/Source/MeshModelingToolsExp/Private/RemoveOccludedTrianglesTool.cpp
@@ -290,7 +290,7 @@ void URemoveOccludedTrianglesTool::SetupPreviews()
 
 
 
-void URemoveOccludedTrianglesTool::Shutdown(EToolShutdownType ShutdownType)
+void URemoveOccludedTrianglesTool::OnShutdown(EToolShutdownType ShutdownType)
 {
 	if (ShutdownType == EToolShutdownType::Accept && AreAllTargetsValid() == false)
 	{
diff --git a/Engine/Plugins/Experimental/MeshModelingToolsetExp/Source/MeshModelingToolsExp/Private/RevolveBoundaryTool.cpp b/Engine/Plugins/Experimental/MeshModelingToolsetExp/Source/MeshModelingToolsExp/Private/RevolveBoundaryTool.cpp
--- a/Engine/Plugins/Experimental/MeshModelingToolsetExp/Source/MeshModelingToolsExp/Private/RevolveBoundaryTool.cpp
+++ b/Engine/Plugins/Experimental/MeshModelingToolsetExp/Source/MeshModelingToolsExp/Private/RevolveBoundaryTool.cpp
@@ -265,9 +265,9 @@ void URevolveBoundaryTool::StartPreview()
 	Preview->InvalidateResult();
 }
 
-void URevolveBoundaryTool::Shutdown(EToolShutdownType ShutdownType)
+void URevolveBoundaryTool::OnShutdown(EToolShutdownType ShutdownType)
 {
-	UMeshBoundaryToolBase::Shutdown(ShutdownType);
+	UMeshBoundaryToolBase::OnShutdown(ShutdownType);
 
 	OutputTypeProperties->SaveProperties(this);
 
diff --git a/Engine/Plugins/Experimental/MeshModelingToolsetExp/Source/MeshModelingToolsExp/Private/SkinWeightsBindingTool.cpp b/Engine/Plugins/Experimental/MeshModelingToolsetExp/Source/MeshModelingToolsExp/Private/SkinWeightsBindingTool.cpp
--- a/Engine/Plugins/Experimental/MeshModelingToolsetExp/Source/MeshModelingToolsExp/Private/SkinWeightsBindingTool.cpp
+++ b/Engine/Plugins/Experimental/MeshModelingToolsetExp/Source/MeshModelingToolsExp/Private/SkinWeightsBindingTool.cpp
@@ -864,7 +864,7 @@ void USkinWeightsBindingTool::Setup()
 }
 
 
-void USkinWeightsBindingTool::Shutdown(EToolShutdownType ShutdownType)
+void USkinWeightsBindingTool::OnShutdown(EToolShutdownType ShutdownType)
 {
 	Properties->SaveProperties(this);
 
diff --git a/Engine/Plugins/Experimental/MeshModelingToolsetExp/Source/MeshModelingToolsExp/Private/SplitMeshesTool.cpp b/Engine/Plugins/Experimental/MeshModelingToolsetExp/Source/MeshModelingToolsExp/Private/SplitMeshesTool.cpp
--- a/Engine/Plugins/Experimental/MeshModelingToolsetExp/Source/MeshModelingToolsExp/Private/SplitMeshesTool.cpp
+++ b/Engine/Plugins/Experimental/MeshModelingToolsetExp/Source/MeshModelingToolsExp/Private/SplitMeshesTool.cpp
@@ -80,7 +80,7 @@ bool USplitMeshesTool::CanAccept() const
 	return Super::CanAccept();
 }
 
-void USplitMeshesTool::Shutdown(EToolShutdownType ShutdownType)
+void USplitMeshesTool::OnShutdown(EToolShutdownType ShutdownType)
 {
 	OutputTypeProperties->SaveProperties(this, TEXT("OutputTypeFromInputTool"));
 	BasicProperties->SaveProperties(this);
diff --git a/Engine/Plugins/Experimental/MeshModelingToolsetExp/Source/MeshModelingToolsExp/Private/TransferMeshTool.cpp b/Engine/Plugins/Experimental/MeshModelingToolsetExp/Source/MeshModelingToolsExp/Private/TransferMeshTool.cpp
--- a/Engine/Plugins/Experimental/MeshModelingToolsetExp/Source/MeshModelingToolsExp/Private/TransferMeshTool.cpp
+++ b/Engine/Plugins/Experimental/MeshModelingToolsetExp/Source/MeshModelingToolsExp/Private/TransferMeshTool.cpp
@@ -110,7 +110,7 @@ bool UTransferMeshTool::CanAccept() const
 	return Super::CanAccept();
 }
 
-void UTransferMeshTool::Shutdown(EToolShutdownType ShutdownType)
+void UTransferMeshTool::OnShutdown(EToolShutdownType ShutdownType)
 {
 	BasicProperties->SaveProperties(this);
 
diff --git a/Engine/Plugins/Experimental/MeshModelingToolsetExp/Source/MeshModelingToolsExp/Private/TransformMeshesTool.cpp b/Engine/Plugins/Experimental/MeshModelingToolsetExp/Source/MeshModelingToolsExp/Private/TransformMeshesTool.cpp
--- a/Engine/Plugins/Experimental/MeshModelingToolsetExp/Source/MeshModelingToolsExp/Private/TransformMeshesTool.cpp
+++ b/Engine/Plugins/Experimental/MeshModelingToolsetExp/Source/MeshModelingToolsExp/Private/TransformMeshesTool.cpp
@@ -99,7 +99,7 @@ void UTransformMeshesTool::Setup()
 
 
 
-void UTransformMeshesTool::Shutdown(EToolShutdownType ShutdownType)
+void UTransformMeshesTool::OnShutdown(EToolShutdownType ShutdownType)
 {
 	TransformProps->SaveProperties(this);
 
diff --git a/Engine/Plugins/Experimental/MeshModelingToolsetExp/Source/MeshModelingToolsExp/Private/WeldMeshEdgesTool.cpp b/Engine/Plugins/Experimental/MeshModelingToolsetExp/Source/MeshModelingToolsExp/Private/WeldMeshEdgesTool.cpp
--- a/Engine/Plugins/Experimental/MeshModelingToolsetExp/Source/MeshModelingToolsExp/Private/WeldMeshEdgesTool.cpp
+++ b/Engine/Plugins/Experimental/MeshModelingToolsetExp/Source/MeshModelingToolsExp/Private/WeldMeshEdgesTool.cpp
@@ -170,7 +170,7 @@ void UWeldMeshEdgesTool::Setup()
 }
 
 
-void UWeldMeshEdgesTool::Shutdown(EToolShutdownType ShutdownType)
+void UWeldMeshEdgesTool::OnShutdown(EToolShutdownType ShutdownType)
 {
 	Settings->SaveProperties(this);
 	UE::ToolTarget::ShowSourceObject(Target);
diff --git a/Engine/Plugins/Experimental/MeshModelingToolsetExp/Source/MeshModelingToolsExp/Public/AlignObjectsTool.h b/Engine/Plugins/Experimental/MeshModelingToolsetExp/Source/MeshModelingToolsExp/Public/AlignObjectsTool.h
--- a/Engine/Plugins/Experimental/MeshModelingToolsetExp/Source/MeshModelingToolsExp/Public/AlignObjectsTool.h
+++ b/Engine/Plugins/Experimental/MeshModelingToolsetExp/Source/MeshModelingToolsExp/Public/AlignObjectsTool.h
@@ -110,7 +110,7 @@ public:
 	virtual void RegisterActions(FInteractiveToolActionSet& ActionSet) override;
 
 	virtual void Setup() override;
-	virtual void Shutdown(EToolShutdownType ShutdownType) override;
+	virtual void OnShutdown(EToolShutdownType ShutdownType) override;
 
 	virtual void OnTick(float DeltaTime) override;
 	virtual void Render(IToolsContextRenderAPI* RenderAPI) override;
diff --git a/Engine/Plugins/Experimental/MeshModelingToolsetExp/Source/MeshModelingToolsExp/Public/BakeMeshAttributeMapsTool.h b/Engine/Plugins/Experimental/MeshModelingToolsetExp/Source/MeshModelingToolsExp/Public/BakeMeshAttributeMapsTool.h
--- a/Engine/Plugins/Experimental/MeshModelingToolsetExp/Source/MeshModelingToolsExp/Public/BakeMeshAttributeMapsTool.h
+++ b/Engine/Plugins/Experimental/MeshModelingToolsetExp/Source/MeshModelingToolsExp/Public/BakeMeshAttributeMapsTool.h
@@ -94,7 +94,7 @@ public:
 
 	// Begin UInteractiveTool interface
 	virtual void Setup() override;
-	virtual void Shutdown(EToolShutdownType ShutdownType) override;
+	virtual void OnShutdown(EToolShutdownType ShutdownType) override;
 
 	virtual bool HasCancel() const override { return true; }
 	virtual bool HasAccept() const override { return true; }
diff --git a/Engine/Plugins/Experimental/MeshModelingToolsetExp/Source/MeshModelingToolsExp/Public/BakeMeshAttributeMapsToolBase.h b/Engine/Plugins/Experimental/MeshModelingToolsetExp/Source/MeshModelingToolsExp/Public/BakeMeshAttributeMapsToolBase.h
--- a/Engine/Plugins/Experimental/MeshModelingToolsetExp/Source/MeshModelingToolsExp/Public/BakeMeshAttributeMapsToolBase.h
+++ b/Engine/Plugins/Experimental/MeshModelingToolsetExp/Source/MeshModelingToolsExp/Public/BakeMeshAttributeMapsToolBase.h
@@ -80,7 +80,7 @@ public:
 
 	// Begin UInteractiveTool interface
 	virtual void Setup() override;
-	virtual void Shutdown(EToolShutdownType ShutdownType) override;
+	virtual void OnShutdown(EToolShutdownType ShutdownType) override;
 
 	virtual void OnTick(float DeltaTime) override;
 	virtual void Render(IToolsContextRenderAPI* RenderAPI) override;
diff --git a/Engine/Plugins/Experimental/MeshModelingToolsetExp/Source/MeshModelingToolsExp/Public/BakeMeshAttributeVertexTool.h b/Engine/Plugins/Experimental/MeshModelingToolsetExp/Source/MeshModelingToolsExp/Public/BakeMeshAttributeVertexTool.h
--- a/Engine/Plugins/Experimental/MeshModelingToolsetExp/Source/MeshModelingToolsExp/Public/BakeMeshAttributeVertexTool.h
+++ b/Engine/Plugins/Experimental/MeshModelingToolsetExp/Source/MeshModelingToolsExp/Public/BakeMeshAttributeVertexTool.h
@@ -120,7 +120,7 @@ public:
 
 	// Begin UInteractiveTool interface
 	virtual void Setup() override;
-	virtual void Shutdown(EToolShutdownType ShutdownType) override;
+	virtual void OnShutdown(EToolShutdownType ShutdownType) override;
 
 	virtual void OnTick(float DeltaTime) override;
 	virtual void Render(IToolsContextRenderAPI* RenderAPI) override;
diff --git a/Engine/Plugins/Experimental/MeshModelingToolsetExp/Source/MeshModelingToolsExp/Public/BakeMultiMeshAttributeMapsTool.h b/Engine/Plugins/Experimental/MeshModelingToolsetExp/Source/MeshModelingToolsExp/Public/BakeMultiMeshAttributeMapsTool.h
--- a/Engine/Plugins/Experimental/MeshModelingToolsetExp/Source/MeshModelingToolsExp/Public/BakeMultiMeshAttributeMapsTool.h
+++ b/Engine/Plugins/Experimental/MeshModelingToolsetExp/Source/MeshModelingToolsExp/Public/BakeMultiMeshAttributeMapsTool.h
@@ -159,7 +159,7 @@ public:
 
 	// Begin UInteractiveTool interface
 	virtual void Setup() override;
-	virtual void Shutdown(EToolShutdownType ShutdownType) override;
+	virtual void OnShutdown(EToolShutdownType ShutdownType) override;
 
 	virtual bool HasCancel() const override { return true; }
 	virtual bool HasAccept() const override { return true; }
diff --git a/Engine/Plugins/Experimental/MeshModelingToolsetExp/Source/MeshModelingToolsExp/Public/BakeTransformTool.h b/Engine/Plugins/Experimental/MeshModelingToolsetExp/Source/MeshModelingToolsExp/Public/BakeTransformTool.h
--- a/Engine/Plugins/Experimental/MeshModelingToolsetExp/Source/MeshModelingToolsExp/Public/BakeTransformTool.h
+++ b/Engine/Plugins/Experimental/MeshModelingToolsetExp/Source/MeshModelingToolsExp/Public/BakeTransformTool.h
@@ -79,7 +79,7 @@ public:
 	UBakeTransformTool();
 
 	virtual void Setup() override;
-	virtual void Shutdown(EToolShutdownType ShutdownType) override;
+	virtual void OnShutdown(EToolShutdownType ShutdownType) override;
 
 	virtual bool HasCancel() const override { return true; }
 	virtual bool HasAccept() const override { return true; }
diff --git a/Engine/Plugins/Experimental/MeshModelingToolsetExp/Source/MeshModelingToolsExp/Public/ConvertMeshesTool.h b/Engine/Plugins/Experimental/MeshModelingToolsetExp/Source/MeshModelingToolsExp/Public/ConvertMeshesTool.h
--- a/Engine/Plugins/Experimental/MeshModelingToolsetExp/Source/MeshModelingToolsExp/Public/ConvertMeshesTool.h
+++ b/Engine/Plugins/Experimental/MeshModelingToolsetExp/Source/MeshModelingToolsExp/Public/ConvertMeshesTool.h
@@ -46,7 +46,7 @@ class MESHMODELINGTOOLSEXP_API UConvertMeshesTool : public UMultiSelectionMeshEd
 
 public:
 	virtual void Setup() override;
-	virtual void Shutdown(EToolShutdownType ShutdownType) override;
+	virtual void OnShutdown(EToolShutdownType ShutdownType) override;
 
 	virtual bool HasCancel() const override { return true; }
 	virtual bool HasAccept() const override { return true; }
diff --git a/Engine/Plugins/Experimental/MeshModelingToolsetExp/Source/MeshModelingToolsExp/Public/ConvertToPolygonsTool.h b/Engine/Plugins/Experimental/MeshModelingToolsetExp/Source/MeshModelingToolsExp/Public/ConvertToPolygonsTool.h
--- a/Engine/Plugins/Experimental/MeshModelingToolsetExp/Source/MeshModelingToolsExp/Public/ConvertToPolygonsTool.h
+++ b/Engine/Plugins/Experimental/MeshModelingToolsetExp/Source/MeshModelingToolsExp/Public/ConvertToPolygonsTool.h
@@ -116,7 +116,7 @@ public:
 	UConvertToPolygonsTool();
 
 	virtual void Setup() override;
-	virtual void Shutdown(EToolShutdownType ShutdownType) override;
+	virtual void OnShutdown(EToolShutdownType ShutdownType) override;
 
 	virtual void OnTick(float DeltaTime) override;
 
diff --git a/Engine/Plugins/Experimental/MeshModelingToolsetExp/Source/MeshModelingToolsExp/Public/DisplaceMeshTool.h b/Engine/Plugins/Experimental/MeshModelingToolsetExp/Source/MeshModelingToolsExp/Public/DisplaceMeshTool.h
--- a/Engine/Plugins/Experimental/MeshModelingToolsetExp/Source/MeshModelingToolsExp/Public/DisplaceMeshTool.h
+++ b/Engine/Plugins/Experimental/MeshModelingToolsetExp/Source/MeshModelingToolsExp/Public/DisplaceMeshTool.h
@@ -272,7 +272,7 @@ class MESHMODELINGTOOLSEXP_API UDisplaceMeshTool : public USingleSelectionMeshEd
 	GENERATED_BODY()
 public:
 	void Setup() final;
-	void Shutdown(EToolShutdownType ShutdownType) final;
+	void OnShutdown(EToolShutdownType ShutdownType) final;
 	void OnTick(float DeltaTime) final;
 
 	bool HasCancel() const final { return true; }
diff --git a/Engine/Plugins/Experimental/MeshModelingToolsetExp/Source/MeshModelingToolsExp/Public/EditNormalsTool.h b/Engine/Plugins/Experimental/MeshModelingToolsetExp/Source/MeshModelingToolsExp/Public/EditNormalsTool.h
--- a/Engine/Plugins/Experimental/MeshModelingToolsetExp/Source/MeshModelingToolsExp/Public/EditNormalsTool.h
+++ b/Engine/Plugins/Experimental/MeshModelingToolsetExp/Source/MeshModelingToolsExp/Public/EditNormalsTool.h
@@ -130,7 +130,7 @@ public:
 	UEditNormalsTool();
 
 	virtual void Setup() override;
-	virtual void Shutdown(EToolShutdownType ShutdownType) override;
+	virtual void OnShutdown(EToolShutdownType ShutdownType) override;
 
 	virtual void OnTick(float DeltaTime) override;
 	virtual void Render(IToolsContextRenderAPI* RenderAPI) override;
diff --git a/Engine/Plugins/Experimental/MeshModelingToolsetExp/Source/MeshModelingToolsExp/Public/EditPivotTool.h b/Engine/Plugins/Experimental/MeshModelingToolsetExp/Source/MeshModelingToolsExp/Public/EditPivotTool.h
--- a/Engine/Plugins/Experimental/MeshModelingToolsetExp/Source/MeshModelingToolsExp/Public/EditPivotTool.h
+++ b/Engine/Plugins/Experimental/MeshModelingToolsetExp/Source/MeshModelingToolsExp/Public/EditPivotTool.h
@@ -157,7 +157,7 @@ public:
 	virtual void RegisterActions(FInteractiveToolActionSet& ActionSet) override;
 
 	virtual void Setup() override;
-	virtual void Shutdown(EToolShutdownType ShutdownType) override;
+	virtual void OnShutdown(EToolShutdownType ShutdownType) override;
 
 	virtual void OnTick(float DeltaTime) override;
 	virtual void Render(IToolsContextRenderAPI* RenderAPI) override;
diff --git a/Engine/Plugins/Experimental/MeshModelingToolsetExp/Source/MeshModelingToolsExp/Public/HoleFillTool.h b/Engine/Plugins/Experimental/MeshModelingToolsetExp/Source/MeshModelingToolsExp/Public/HoleFillTool.h
--- a/Engine/Plugins/Experimental/MeshModelingToolsetExp/Source/MeshModelingToolsExp/Public/HoleFillTool.h
+++ b/Engine/Plugins/Experimental/MeshModelingToolsetExp/Source/MeshModelingToolsExp/Public/HoleFillTool.h
@@ -228,7 +228,7 @@ public:
 	bool HasCancel() const override { return true; }
 	bool HasAccept() const override { return true; }
 	bool CanAccept() const override;
-	void Shutdown(EToolShutdownType ShutdownType) override;
+	void OnShutdown(EToolShutdownType ShutdownType) override;
 
 	void OnSelectionModified();
 
diff --git a/Engine/Plugins/Experimental/MeshModelingToolsetExp/Source/MeshModelingToolsExp/Public/LatticeDeformerTool.h b/Engine/Plugins/Experimental/MeshModelingToolsetExp/Source/MeshModelingToolsExp/Public/LatticeDeformerTool.h
--- a/Engine/Plugins/Experimental/MeshModelingToolsetExp/Source/MeshModelingToolsExp/Public/LatticeDeformerTool.h
+++ b/Engine/Plugins/Experimental/MeshModelingToolsetExp/Source/MeshModelingToolsExp/Public/LatticeDeformerTool.h
@@ -142,7 +142,7 @@ public:
 	virtual bool CanAccept() const override;
 
 	virtual void Setup() override;
-	virtual void Shutdown(EToolShutdownType ShutdownType) override;
+	virtual void OnShutdown(EToolShutdownType ShutdownType) override;
 
 	virtual void OnTick(float DeltaTime) override;
 	virtual void Render(IToolsContextRenderAPI* RenderAPI) override;
diff --git a/Engine/Plugins/Experimental/MeshModelingToolsetExp/Source/MeshModelingToolsExp/Public/MeshBoundaryToolBase.h b/Engine/Plugins/Experimental/MeshModelingToolsetExp/Source/MeshModelingToolsExp/Public/MeshBoundaryToolBase.h
--- a/Engine/Plugins/Experimental/MeshModelingToolsetExp/Source/MeshModelingToolsExp/Public/MeshBoundaryToolBase.h
+++ b/Engine/Plugins/Experimental/MeshModelingToolsetExp/Source/MeshModelingToolsExp/Public/MeshBoundaryToolBase.h
@@ -33,7 +33,7 @@ public:
 	// CanAccept() needs to be provided by child.
 
 	virtual void Setup() override;
-	virtual void Shutdown(EToolShutdownType ShutdownType) override;
+	virtual void OnShutdown(EToolShutdownType ShutdownType) override;
 
 	virtual void Render(IToolsContextRenderAPI* RenderAPI) override;
 
diff --git a/Engine/Plugins/Experimental/MeshModelingToolsetExp/Source/MeshModelingToolsExp/Public/MeshInspectorTool.h b/Engine/Plugins/Experimental/MeshModelingToolsetExp/Source/MeshModelingToolsExp/Public/MeshInspectorTool.h
--- a/Engine/Plugins/Experimental/MeshModelingToolsetExp/Source/MeshModelingToolsExp/Public/MeshInspectorTool.h
+++ b/Engine/Plugins/Experimental/MeshModelingToolsetExp/Source/MeshModelingToolsExp/Public/MeshInspectorTool.h
@@ -104,7 +104,7 @@ public:
 	virtual void RegisterActions(FInteractiveToolActionSet& ActionSet) override;
 
 	virtual void Setup() override;
-	virtual void Shutdown(EToolShutdownType ShutdownType) override;
+	virtual void OnShutdown(EToolShutdownType ShutdownType) override;
 
 	virtual void Render(IToolsContextRenderAPI* RenderAPI) override;
 
diff --git a/Engine/Plugins/Experimental/MeshModelingToolsetExp/Source/MeshModelingToolsExp/Public/MeshSpaceDeformerTool.h b/Engine/Plugins/Experimental/MeshModelingToolsetExp/Source/MeshModelingToolsExp/Public/MeshSpaceDeformerTool.h
--- a/Engine/Plugins/Experimental/MeshModelingToolsetExp/Source/MeshModelingToolsExp/Public/MeshSpaceDeformerTool.h
+++ b/Engine/Plugins/Experimental/MeshModelingToolsetExp/Source/MeshModelingToolsExp/Public/MeshSpaceDeformerTool.h
@@ -205,7 +205,7 @@ public:
 	UMeshSpaceDeformerTool();
 
 	virtual void Setup() override;
-	virtual void Shutdown(EToolShutdownType ShutdownType) override;
+	virtual void OnShutdown(EToolShutdownType ShutdownType) override;
 
 	void RequestAction(EMeshSpaceDeformerToolAction ActionType);
 
diff --git a/Engine/Plugins/Experimental/MeshModelingToolsetExp/Source/MeshModelingToolsExp/Public/MirrorTool.h b/Engine/Plugins/Experimental/MeshModelingToolsetExp/Source/MeshModelingToolsExp/Public/MirrorTool.h
--- a/Engine/Plugins/Experimental/MeshModelingToolsetExp/Source/MeshModelingToolsExp/Public/MirrorTool.h
+++ b/Engine/Plugins/Experimental/MeshModelingToolsetExp/Source/MeshModelingToolsExp/Public/MirrorTool.h
@@ -188,7 +188,7 @@ public:
 	UMirrorTool();
 
 	virtual void Setup() override;
-	virtual void Shutdown(EToolShutdownType ShutdownType) override;
+	virtual void OnShutdown(EToolShutdownType ShutdownType) override;
 
 	virtual void OnTick(float DeltaTime) override;
 	virtual void Render(IToolsContextRenderAPI* RenderAPI) override;
diff --git a/Engine/Plugins/Experimental/MeshModelingToolsetExp/Source/MeshModelingToolsExp/Public/Physics/ExtractCollisionGeometryTool.h b/Engine/Plugins/Experimental/MeshModelingToolsetExp/Source/MeshModelingToolsExp/Public/Physics/ExtractCollisionGeometryTool.h
--- a/Engine/Plugins/Experimental/MeshModelingToolsetExp/Source/MeshModelingToolsExp/Public/Physics/ExtractCollisionGeometryTool.h
+++ b/Engine/Plugins/Experimental/MeshModelingToolsetExp/Source/MeshModelingToolsExp/Public/Physics/ExtractCollisionGeometryTool.h
@@ -72,7 +72,7 @@ class MESHMODELINGTOOLSEXP_API UExtractCollisionGeometryTool : public USingleSel
 	GENERATED_BODY()
 public:
 	virtual void Setup() override;
-	virtual void Shutdown(EToolShutdownType ShutdownType) override;
+	virtual void OnShutdown(EToolShutdownType ShutdownType) override;
 
 	virtual void OnTick(float DeltaTime) override;
 
diff --git a/Engine/Plugins/Experimental/MeshModelingToolsetExp/Source/MeshModelingToolsExp/Public/Physics/PhysicsInspectorTool.h b/Engine/Plugins/Experimental/MeshModelingToolsetExp/Source/MeshModelingToolsExp/Public/Physics/PhysicsInspectorTool.h
--- a/Engine/Plugins/Experimental/MeshModelingToolsetExp/Source/MeshModelingToolsExp/Public/Physics/PhysicsInspectorTool.h
+++ b/Engine/Plugins/Experimental/MeshModelingToolsetExp/Source/MeshModelingToolsExp/Public/Physics/PhysicsInspectorTool.h
@@ -34,7 +34,7 @@ class MESHMODELINGTOOLSEXP_API UPhysicsInspectorTool : public UMultiSelectionMes
 	GENERATED_BODY()
 public:
 	virtual void Setup() override;
-	virtual void Shutdown(EToolShutdownType ShutdownType) override;
+	virtual void OnShutdown(EToolShutdownType ShutdownType) override;
 
 	virtual void OnTick(float DeltaTime) override;
 
diff --git a/Engine/Plugins/Experimental/MeshModelingToolsetExp/Source/MeshModelingToolsExp/Public/Physics/SetCollisionGeometryTool.h b/Engine/Plugins/Experimental/MeshModelingToolsetExp/Source/MeshModelingToolsExp/Public/Physics/SetCollisionGeometryTool.h
--- a/Engine/Plugins/Experimental/MeshModelingToolsetExp/Source/MeshModelingToolsExp/Public/Physics/SetCollisionGeometryTool.h
+++ b/Engine/Plugins/Experimental/MeshModelingToolsetExp/Source/MeshModelingToolsExp/Public/Physics/SetCollisionGeometryTool.h
@@ -146,7 +146,7 @@ class MESHMODELINGTOOLSEXP_API USetCollisionGeometryTool : public UMultiSelectio
 	GENERATED_BODY()
 public:
 	virtual void Setup() override;
-	virtual void Shutdown(EToolShutdownType ShutdownType) override;
+	virtual void OnShutdown(EToolShutdownType ShutdownType) override;
 
 	virtual void OnTick(float DeltaTime) override;
 
diff --git a/Engine/Plugins/Experimental/MeshModelingToolsetExp/Source/MeshModelingToolsExp/Public/PlaneCutTool.h b/Engine/Plugins/Experimental/MeshModelingToolsetExp/Source/MeshModelingToolsExp/Public/PlaneCutTool.h
--- a/Engine/Plugins/Experimental/MeshModelingToolsetExp/Source/MeshModelingToolsExp/Public/PlaneCutTool.h
+++ b/Engine/Plugins/Experimental/MeshModelingToolsetExp/Source/MeshModelingToolsExp/Public/PlaneCutTool.h
@@ -125,7 +125,7 @@ public:
 	UPlaneCutTool();
 
 	virtual void Setup() override;
-	virtual void Shutdown(EToolShutdownType ShutdownType) override;
+	virtual void OnShutdown(EToolShutdownType ShutdownType) override;
 
 	virtual void RegisterActions(FInteractiveToolActionSet& ActionSet) override;
 
diff --git a/Engine/Plugins/Experimental/MeshModelingToolsetExp/Source/MeshModelingToolsExp/Public/RemeshMeshTool.h b/Engine/Plugins/Experimental/MeshModelingToolsetExp/Source/MeshModelingToolsExp/Public/RemeshMeshTool.h
--- a/Engine/Plugins/Experimental/MeshModelingToolsetExp/Source/MeshModelingToolsExp/Public/RemeshMeshTool.h
+++ b/Engine/Plugins/Experimental/MeshModelingToolsetExp/Source/MeshModelingToolsExp/Public/RemeshMeshTool.h
@@ -108,7 +108,7 @@ public:
 	URemeshMeshTool(const FObjectInitializer&);
 
 	virtual void Setup() override;
-	virtual void Shutdown(EToolShutdownType ShutdownType) override;
+	virtual void OnShutdown(EToolShutdownType ShutdownType) override;
 
 	virtual void OnTick(float DeltaTime) override;
 
diff --git a/Engine/Plugins/Experimental/MeshModelingToolsetExp/Source/MeshModelingToolsExp/Public/RemoveOccludedTrianglesTool.h b/Engine/Plugins/Experimental/MeshModelingToolsetExp/Source/MeshModelingToolsExp/Public/RemoveOccludedTrianglesTool.h
--- a/Engine/Plugins/Experimental/MeshModelingToolsetExp/Source/MeshModelingToolsExp/Public/RemoveOccludedTrianglesTool.h
+++ b/Engine/Plugins/Experimental/MeshModelingToolsetExp/Source/MeshModelingToolsExp/Public/RemoveOccludedTrianglesTool.h
@@ -165,7 +165,7 @@ public:
 	URemoveOccludedTrianglesTool();
 
 	virtual void Setup() override;
-	virtual void Shutdown(EToolShutdownType ShutdownType) override;
+	virtual void OnShutdown(EToolShutdownType ShutdownType) override;
 
 	virtual void OnTick(float DeltaTime) override;
 
diff --git a/Engine/Plugins/Experimental/MeshModelingToolsetExp/Source/MeshModelingToolsExp/Public/RevolveBoundaryTool.h b/Engine/Plugins/Experimental/MeshModelingToolsetExp/Source/MeshModelingToolsExp/Public/RevolveBoundaryTool.h
--- a/Engine/Plugins/Experimental/MeshModelingToolsetExp/Source/MeshModelingToolsExp/Public/RevolveBoundaryTool.h
+++ b/Engine/Plugins/Experimental/MeshModelingToolsetExp/Source/MeshModelingToolsExp/Public/RevolveBoundaryTool.h
@@ -88,7 +88,7 @@ public:
 	virtual bool CanAccept() const override;
 
 	virtual void Setup() override;
-	virtual void Shutdown(EToolShutdownType ShutdownType) override;
+	virtual void OnShutdown(EToolShutdownType ShutdownType) override;
 
 	virtual void Render(IToolsContextRenderAPI* RenderAPI) override;
 	virtual void OnTick(float DeltaTime) override;
diff --git a/Engine/Plugins/Experimental/MeshModelingToolsetExp/Source/MeshModelingToolsExp/Public/SkinWeightsBindingTool.h b/Engine/Plugins/Experimental/MeshModelingToolsetExp/Source/MeshModelingToolsExp/Public/SkinWeightsBindingTool.h
--- a/Engine/Plugins/Experimental/MeshModelingToolsetExp/Source/MeshModelingToolsExp/Public/SkinWeightsBindingTool.h
+++ b/Engine/Plugins/Experimental/MeshModelingToolsetExp/Source/MeshModelingToolsExp/Public/SkinWeightsBindingTool.h
@@ -89,7 +89,7 @@ public:
 	~USkinWeightsBindingTool();
 	
 	void Setup() override;
-	void Shutdown(EToolShutdownType ShutdownType) override;
+	void OnShutdown(EToolShutdownType ShutdownType) override;
 	
 	void OnTick(float DeltaTime) override;
 	void Render(IToolsContextRenderAPI* RenderAPI) override;
diff --git a/Engine/Plugins/Experimental/MeshModelingToolsetExp/Source/MeshModelingToolsExp/Public/SplitMeshesTool.h b/Engine/Plugins/Experimental/MeshModelingToolsetExp/Source/MeshModelingToolsExp/Public/SplitMeshesTool.h
--- a/Engine/Plugins/Experimental/MeshModelingToolsetExp/Source/MeshModelingToolsExp/Public/SplitMeshesTool.h
+++ b/Engine/Plugins/Experimental/MeshModelingToolsetExp/Source/MeshModelingToolsExp/Public/SplitMeshesTool.h
@@ -42,7 +42,7 @@ class MESHMODELINGTOOLSEXP_API USplitMeshesTool : public UMultiSelectionMeshEdit
 
 public:
 	virtual void Setup() override;
-	virtual void Shutdown(EToolShutdownType ShutdownType) override;
+	virtual void OnShutdown(EToolShutdownType ShutdownType) override;
 
 	virtual bool HasCancel() const override { return true; }
 	virtual bool HasAccept() const override { return true; }
diff --git a/Engine/Plugins/Experimental/MeshModelingToolsetExp/Source/MeshModelingToolsExp/Public/TransferMeshTool.h b/Engine/Plugins/Experimental/MeshModelingToolsetExp/Source/MeshModelingToolsExp/Public/TransferMeshTool.h
--- a/Engine/Plugins/Experimental/MeshModelingToolsetExp/Source/MeshModelingToolsExp/Public/TransferMeshTool.h
+++ b/Engine/Plugins/Experimental/MeshModelingToolsetExp/Source/MeshModelingToolsExp/Public/TransferMeshTool.h
@@ -81,7 +81,7 @@ class MESHMODELINGTOOLSEXP_API UTransferMeshTool : public UMultiSelectionMeshEdi
 
 public:
 	virtual void Setup() override;
-	virtual void Shutdown(EToolShutdownType ShutdownType) override;
+	virtual void OnShutdown(EToolShutdownType ShutdownType) override;
 
 	virtual bool HasCancel() const override { return true; }
 	virtual bool HasAccept() const override { return true; }
diff --git a/Engine/Plugins/Experimental/MeshModelingToolsetExp/Source/MeshModelingToolsExp/Public/TransformMeshesTool.h b/Engine/Plugins/Experimental/MeshModelingToolsetExp/Source/MeshModelingToolsExp/Public/TransformMeshesTool.h
--- a/Engine/Plugins/Experimental/MeshModelingToolsetExp/Source/MeshModelingToolsExp/Public/TransformMeshesTool.h
+++ b/Engine/Plugins/Experimental/MeshModelingToolsetExp/Source/MeshModelingToolsExp/Public/TransformMeshesTool.h
@@ -158,7 +158,7 @@ public:
 	virtual void RegisterActions(FInteractiveToolActionSet& ActionSet) override;
 
 	virtual void Setup() override;
-	virtual void Shutdown(EToolShutdownType ShutdownType) override;
+	virtual void OnShutdown(EToolShutdownType ShutdownType) override;
 
 	virtual void Render(IToolsContextRenderAPI* RenderAPI) override;
 
diff --git a/Engine/Plugins/Experimental/MeshModelingToolsetExp/Source/MeshModelingToolsExp/Public/WeldMeshEdgesTool.h b/Engine/Plugins/Experimental/MeshModelingToolsetExp/Source/MeshModelingToolsExp/Public/WeldMeshEdgesTool.h
--- a/Engine/Plugins/Experimental/MeshModelingToolsetExp/Source/MeshModelingToolsExp/Public/WeldMeshEdgesTool.h
+++ b/Engine/Plugins/Experimental/MeshModelingToolsetExp/Source/MeshModelingToolsExp/Public/WeldMeshEdgesTool.h
@@ -73,7 +73,7 @@ public:
 	UWeldMeshEdgesTool();
 
 	virtual void Setup() override;
-	virtual void Shutdown(EToolShutdownType ShutdownType) override;
+	virtual void OnShutdown(EToolShutdownType ShutdownType) override;
 
 	virtual void OnTick(float DeltaTime) override;
 
diff --git a/Engine/Plugins/Runtime/MeshModelingToolset/Source/MeshModelingTools/Private/CSGMeshesTool.cpp b/Engine/Plugins/Runtime/MeshModelingToolset/Source/MeshModelingTools/Private/CSGMeshesTool.cpp
--- a/Engine/Plugins/Runtime/MeshModelingToolset/Source/MeshModelingTools/Private/CSGMeshesTool.cpp
+++ b/Engine/Plugins/Runtime/MeshModelingToolset/Source/MeshModelingTools/Private/CSGMeshesTool.cpp
@@ -376,9 +376,9 @@ FText UCSGMeshesTool::GetActionName() const
 
 
 
-void UCSGMeshesTool::Shutdown(EToolShutdownType ShutdownType)
+void UCSGMeshesTool::OnShutdown(EToolShutdownType ShutdownType)
 {
-	Super::Shutdown(ShutdownType);
+	Super::OnShutdown(ShutdownType);
 
 	for (UPreviewMesh* MeshPreview : OriginalMeshPreviews)
 	{
diff --git a/Engine/Plugins/Runtime/MeshModelingToolset/Source/MeshModelingTools/Private/CombineMeshesTool.cpp b/Engine/Plugins/Runtime/MeshModelingToolset/Source/MeshModelingTools/Private/CombineMeshesTool.cpp
--- a/Engine/Plugins/Runtime/MeshModelingToolset/Source/MeshModelingTools/Private/CombineMeshesTool.cpp
+++ b/Engine/Plugins/Runtime/MeshModelingToolset/Source/MeshModelingTools/Private/CombineMeshesTool.cpp
@@ -122,7 +122,7 @@ void UCombineMeshesTool::Setup()
 }
 
 
-void UCombineMeshesTool::Shutdown(EToolShutdownType ShutdownType)
+void UCombineMeshesTool::OnShutdown(EToolShutdownType ShutdownType)
 {
 	BasicProperties->SaveProperties(this);
 	OutputTypeProperties->SaveProperties(this, TEXT("OutputTypeFromInputTool"));
diff --git a/Engine/Plugins/Runtime/MeshModelingToolset/Source/MeshModelingTools/Private/CutMeshWithMeshTool.cpp b/Engine/Plugins/Runtime/MeshModelingToolset/Source/MeshModelingTools/Private/CutMeshWithMeshTool.cpp
--- a/Engine/Plugins/Runtime/MeshModelingToolset/Source/MeshModelingTools/Private/CutMeshWithMeshTool.cpp
+++ b/Engine/Plugins/Runtime/MeshModelingToolset/Source/MeshModelingTools/Private/CutMeshWithMeshTool.cpp
@@ -326,7 +326,7 @@ FText UCutMeshWithMeshTool::GetActionName() const
 }
 
 
-void UCutMeshWithMeshTool::Shutdown(EToolShutdownType ShutdownType)
+void UCutMeshWithMeshTool::OnShutdown(EToolShutdownType ShutdownType)
 {
 	SaveProperties();
 	HandleSourcesProperties->SaveProperties(this);
diff --git a/Engine/Plugins/Runtime/MeshModelingToolset/Source/MeshModelingTools/Private/EditMeshPolygonsTool.cpp b/Engine/Plugins/Runtime/MeshModelingToolset/Source/MeshModelingTools/Private/EditMeshPolygonsTool.cpp
--- a/Engine/Plugins/Runtime/MeshModelingToolset/Source/MeshModelingTools/Private/EditMeshPolygonsTool.cpp
+++ b/Engine/Plugins/Runtime/MeshModelingToolset/Source/MeshModelingTools/Private/EditMeshPolygonsTool.cpp
@@ -447,7 +447,7 @@ bool UEditMeshPolygonsTool::IsToolInputSelectionUsable(const UPersistentMeshSele
 		&& !InputSelectionIn->IsEmpty();
 }
 
-void UEditMeshPolygonsTool::Shutdown(EToolShutdownType ShutdownType)
+void UEditMeshPolygonsTool::OnShutdown(EToolShutdownType ShutdownType)
 {
 	using namespace EditMeshPolygonsToolLocals;
 
diff --git a/Engine/Plugins/Runtime/MeshModelingToolset/Source/MeshModelingTools/Private/RecomputeUVsTool.cpp b/Engine/Plugins/Runtime/MeshModelingToolset/Source/MeshModelingTools/Private/RecomputeUVsTool.cpp
--- a/Engine/Plugins/Runtime/MeshModelingToolset/Source/MeshModelingTools/Private/RecomputeUVsTool.cpp
+++ b/Engine/Plugins/Runtime/MeshModelingToolset/Source/MeshModelingTools/Private/RecomputeUVsTool.cpp
@@ -134,7 +134,7 @@ void URecomputeUVsTool::OnPropertyModified(UObject* PropertySet, FProperty* Prop
 }
 
 
-void URecomputeUVsTool::Shutdown(EToolShutdownType ShutdownType)
+void URecomputeUVsTool::OnShutdown(EToolShutdownType ShutdownType)
 {
 	if (UVLayoutView)
 	{
diff --git a/Engine/Plugins/Runtime/MeshModelingToolset/Source/MeshModelingTools/Private/UVLayoutTool.cpp b/Engine/Plugins/Runtime/MeshModelingToolset/Source/MeshModelingTools/Private/UVLayoutTool.cpp
--- a/Engine/Plugins/Runtime/MeshModelingToolset/Source/MeshModelingTools/Private/UVLayoutTool.cpp
+++ b/Engine/Plugins/Runtime/MeshModelingToolset/Source/MeshModelingTools/Private/UVLayoutTool.cpp
@@ -154,7 +154,7 @@ void UUVLayoutTool::UpdateNumPreviews()
 }
 
 
-void UUVLayoutTool::Shutdown(EToolShutdownType ShutdownType)
+void UUVLayoutTool::OnShutdown(EToolShutdownType ShutdownType)
 {
 	if (UVLayoutView)
 	{
diff --git a/Engine/Plugins/Runtime/MeshModelingToolset/Source/MeshModelingTools/Private/UVProjectionTool.cpp b/Engine/Plugins/Runtime/MeshModelingToolset/Source/MeshModelingTools/Private/UVProjectionTool.cpp
--- a/Engine/Plugins/Runtime/MeshModelingToolset/Source/MeshModelingTools/Private/UVProjectionTool.cpp
+++ b/Engine/Plugins/Runtime/MeshModelingToolset/Source/MeshModelingTools/Private/UVProjectionTool.cpp
@@ -197,7 +197,7 @@ void UUVProjectionTool::UpdateNumPreviews()
 }
 
 
-void UUVProjectionTool::Shutdown(EToolShutdownType ShutdownType)
+void UUVProjectionTool::OnShutdown(EToolShutdownType ShutdownType)
 {
 	UVChannelProperties->SaveProperties(this);
 	BasicProperties->SavedDimensions = BasicProperties->Dimensions;
diff --git a/Engine/Plugins/Runtime/MeshModelingToolset/Source/MeshModelingTools/Public/CSGMeshesTool.h b/Engine/Plugins/Runtime/MeshModelingToolset/Source/MeshModelingTools/Public/CSGMeshesTool.h
--- a/Engine/Plugins/Runtime/MeshModelingToolset/Source/MeshModelingTools/Public/CSGMeshesTool.h
+++ b/Engine/Plugins/Runtime/MeshModelingToolset/Source/MeshModelingTools/Public/CSGMeshesTool.h
@@ -113,7 +113,7 @@ public:
 
 protected:
 
-	virtual void Shutdown(EToolShutdownType ShutdownType) override;
+	virtual void OnShutdown(EToolShutdownType ShutdownType) override;
 
 	virtual void OnPropertyModified(UObject* PropertySet, FProperty* Property) override;
 
diff --git a/Engine/Plugins/Runtime/MeshModelingToolset/Source/MeshModelingTools/Public/CombineMeshesTool.h b/Engine/Plugins/Runtime/MeshModelingToolset/Source/MeshModelingTools/Public/CombineMeshesTool.h
--- a/Engine/Plugins/Runtime/MeshModelingToolset/Source/MeshModelingTools/Public/CombineMeshesTool.h
+++ b/Engine/Plugins/Runtime/MeshModelingToolset/Source/MeshModelingTools/Public/CombineMeshesTool.h
@@ -76,7 +76,7 @@ public:
 	virtual void SetDuplicateMode(bool bDuplicateMode);
 
 	virtual void Setup() override;
-	virtual void Shutdown(EToolShutdownType ShutdownType) override;
+	virtual void OnShutdown(EToolShutdownType ShutdownType) override;
 
 	virtual bool HasCancel() const override { return true; }
 	virtual bool HasAccept() const override { return true; }
diff --git a/Engine/Plugins/Runtime/MeshModelingToolset/Source/MeshModelingTools/Public/CutMeshWithMeshTool.h b/Engine/Plugins/Runtime/MeshModelingToolset/Source/MeshModelingTools/Public/CutMeshWithMeshTool.h
--- a/Engine/Plugins/Runtime/MeshModelingToolset/Source/MeshModelingTools/Public/CutMeshWithMeshTool.h
+++ b/Engine/Plugins/Runtime/MeshModelingToolset/Source/MeshModelingTools/Public/CutMeshWithMeshTool.h
@@ -75,7 +75,7 @@ protected:
 	// IDynamicMeshOperatorFactory API
 	virtual TUniquePtr<UE::Geometry::FDynamicMeshOperator> MakeNewOperator() override;
 
-	virtual void Shutdown(EToolShutdownType ShutdownType) override;
+	virtual void OnShutdown(EToolShutdownType ShutdownType) override;
 
 	void UpdateVisualization();
 
diff --git a/Engine/Plugins/Runtime/MeshModelingToolset/Source/MeshModelingTools/Public/EditMeshPolygonsTool.h b/Engine/Plugins/Runtime/MeshModelingToolset/Source/MeshModelingTools/Public/EditMeshPolygonsTool.h
--- a/Engine/Plugins/Runtime/MeshModelingToolset/Source/MeshModelingTools/Public/EditMeshPolygonsTool.h
+++ b/Engine/Plugins/Runtime/MeshModelingToolset/Source/MeshModelingTools/Public/EditMeshPolygonsTool.h
@@ -441,7 +441,7 @@ public:
 	virtual void SetWorld(UWorld* World) { this->TargetWorld = World; }
 
 	virtual void Setup() override;
-	virtual void Shutdown(EToolShutdownType ShutdownType) override;
+	virtual void OnShutdown(EToolShutdownType ShutdownType) override;
 
 	virtual void OnTick(float DeltaTime) override;
 	virtual void Render(IToolsContextRenderAPI* RenderAPI) override;
diff --git a/Engine/Plugins/Runtime/MeshModelingToolset/Source/MeshModelingTools/Public/RecomputeUVsTool.h b/Engine/Plugins/Runtime/MeshModelingToolset/Source/MeshModelingTools/Public/RecomputeUVsTool.h
--- a/Engine/Plugins/Runtime/MeshModelingToolset/Source/MeshModelingTools/Public/RecomputeUVsTool.h
+++ b/Engine/Plugins/Runtime/MeshModelingToolset/Source/MeshModelingTools/Public/RecomputeUVsTool.h
@@ -47,7 +47,7 @@ class MESHMODELINGTOOLS_API URecomputeUVsTool : public USingleSelectionMeshEditi
 
 public:
 	virtual void Setup() override;
-	virtual void Shutdown(EToolShutdownType ShutdownType) override;
+	virtual void OnShutdown(EToolShutdownType ShutdownType) override;
 
 	virtual void Render(IToolsContextRenderAPI* RenderAPI) override;
 	virtual void OnTick(float DeltaTime) override;
diff --git a/Engine/Plugins/Runtime/MeshModelingToolset/Source/MeshModelingTools/Public/UVLayoutTool.h b/Engine/Plugins/Runtime/MeshModelingToolset/Source/MeshModelingTools/Public/UVLayoutTool.h
--- a/Engine/Plugins/Runtime/MeshModelingToolset/Source/MeshModelingTools/Public/UVLayoutTool.h
+++ b/Engine/Plugins/Runtime/MeshModelingToolset/Source/MeshModelingTools/Public/UVLayoutTool.h
@@ -50,7 +50,7 @@ public:
 	UUVLayoutTool();
 
 	virtual void Setup() override;
-	virtual void Shutdown(EToolShutdownType ShutdownType) override;
+	virtual void OnShutdown(EToolShutdownType ShutdownType) override;
 
 	virtual void OnTick(float DeltaTime) override;
 	virtual void Render(IToolsContextRenderAPI* RenderAPI) override;
diff --git a/Engine/Plugins/Runtime/MeshModelingToolset/Source/MeshModelingTools/Public/UVProjectionTool.h b/Engine/Plugins/Runtime/MeshModelingToolset/Source/MeshModelingTools/Public/UVProjectionTool.h
--- a/Engine/Plugins/Runtime/MeshModelingToolset/Source/MeshModelingTools/Public/UVProjectionTool.h
+++ b/Engine/Plugins/Runtime/MeshModelingToolset/Source/MeshModelingTools/Public/UVProjectionTool.h
@@ -210,7 +210,7 @@ public:
 	friend UUVProjectionOperatorFactory;
 
 	virtual void Setup() override;
-	virtual void Shutdown(EToolShutdownType ShutdownType) override;
+	virtual void OnShutdown(EToolShutdownType ShutdownType) override;
 
 	virtual void OnTick(float DeltaTime) override;
 	virtual void Render(IToolsContextRenderAPI* RenderAPI) override;
diff --git a/Engine/Plugins/Runtime/MeshModelingToolset/Source/ModelingComponents/Private/BaseTools/BaseCreateFromSelectedTool.cpp b/Engine/Plugins/Runtime/MeshModelingToolset/Source/ModelingComponents/Private/BaseTools/BaseCreateFromSelectedTool.cpp
--- a/Engine/Plugins/Runtime/MeshModelingToolset/Source/ModelingComponents/Private/BaseTools/BaseCreateFromSelectedTool.cpp
+++ b/Engine/Plugins/Runtime/MeshModelingToolset/Source/ModelingComponents/Private/BaseTools/BaseCreateFromSelectedTool.cpp
@@ -284,7 +284,7 @@ void UBaseCreateFromSelectedTool::OnPropertyModified(UObject* PropertySet, FProp
 }
 
 
-void UBaseCreateFromSelectedTool::Shutdown(EToolShutdownType ShutdownType)
+void UBaseCreateFromSelectedTool::OnShutdown(EToolShutdownType ShutdownType)
 {
 	SaveProperties();
 	HandleSourcesProperties->SaveProperties(this);
diff --git a/Engine/Plugins/Runtime/MeshModelingToolset/Source/ModelingComponents/Private/BaseTools/MultiSelectionMeshEditingTool.cpp b/Engine/Plugins/Runtime/MeshModelingToolset/Source/ModelingComponents/Private/BaseTools/MultiSelectionMeshEditingTool.cpp
--- a/Engine/Plugins/Runtime/MeshModelingToolset/Source/ModelingComponents/Private/BaseTools/MultiSelectionMeshEditingTool.cpp
+++ b/Engine/Plugins/Runtime/MeshModelingToolset/Source/ModelingComponents/Private/BaseTools/MultiSelectionMeshEditingTool.cpp
@@ -49,6 +49,17 @@ void UMultiSelectionMeshEditingToolBuilder::InitializeNewTool(UMultiSelectionMes
  * Tool
  */
 
+void UMultiSelectionMeshEditingTool::Shutdown(EToolShutdownType ShutdownType)
+{
+	OnShutdown(ShutdownType);
+	TargetWorld = nullptr;
+}
+
+void UMultiSelectionMeshEditingTool::OnShutdown(EToolShutdownType ShutdownType)
+{
+}
+
+
 void UMultiSelectionMeshEditingTool::SetWorld(UWorld* World)
 {
 	TargetWorld = World;
diff --git a/Engine/Plugins/Runtime/MeshModelingToolset/Source/ModelingComponents/Private/BaseTools/SingleSelectionMeshEditingTool.cpp b/Engine/Plugins/Runtime/MeshModelingToolset/Source/ModelingComponents/Private/BaseTools/SingleSelectionMeshEditingTool.cpp
--- a/Engine/Plugins/Runtime/MeshModelingToolset/Source/ModelingComponents/Private/BaseTools/SingleSelectionMeshEditingTool.cpp
+++ b/Engine/Plugins/Runtime/MeshModelingToolset/Source/ModelingComponents/Private/BaseTools/SingleSelectionMeshEditingTool.cpp
@@ -54,6 +54,18 @@ void USingleSelectionMeshEditingToolBuilder::InitializeNewTool(USingleSelectionM
 	}
 }
 
+
+void USingleSelectionMeshEditingTool::Shutdown(EToolShutdownType ShutdownType)
+{
+	OnShutdown(ShutdownType);
+	TargetWorld = nullptr;
+}
+
+void USingleSelectionMeshEditingTool::OnShutdown(EToolShutdownType ShutdownType)
+{
+}
+
+
 void USingleSelectionMeshEditingTool::SetWorld(UWorld* World)
 {
 	TargetWorld = World;
diff --git a/Engine/Plugins/Runtime/MeshModelingToolset/Source/ModelingComponents/Public/BaseTools/BaseCreateFromSelectedTool.h b/Engine/Plugins/Runtime/MeshModelingToolset/Source/ModelingComponents/Public/BaseTools/BaseCreateFromSelectedTool.h
--- a/Engine/Plugins/Runtime/MeshModelingToolset/Source/ModelingComponents/Public/BaseTools/BaseCreateFromSelectedTool.h
+++ b/Engine/Plugins/Runtime/MeshModelingToolset/Source/ModelingComponents/Public/BaseTools/BaseCreateFromSelectedTool.h
@@ -106,7 +106,7 @@ public:
 	//
 
 	virtual void Setup() override;
-	virtual void Shutdown(EToolShutdownType ShutdownType) override;
+	virtual void OnShutdown(EToolShutdownType ShutdownType) override;
 
 	virtual void OnTick(float DeltaTime) override;
 
diff --git a/Engine/Plugins/Runtime/MeshModelingToolset/Source/ModelingComponents/Public/BaseTools/MultiSelectionMeshEditingTool.h b/Engine/Plugins/Runtime/MeshModelingToolset/Source/ModelingComponents/Public/BaseTools/MultiSelectionMeshEditingTool.h
--- a/Engine/Plugins/Runtime/MeshModelingToolset/Source/ModelingComponents/Public/BaseTools/MultiSelectionMeshEditingTool.h
+++ b/Engine/Plugins/Runtime/MeshModelingToolset/Source/ModelingComponents/Public/BaseTools/MultiSelectionMeshEditingTool.h
@@ -44,6 +44,9 @@ class MODELINGCOMPONENTS_API UMultiSelectionMeshEditingTool : public UMultiSelec
 {
 	GENERATED_BODY()
 public:
+	virtual void Shutdown(EToolShutdownType ShutdownType) override;
+    virtual void OnShutdown(EToolShutdownType ShutdownType);
+
 	virtual void SetWorld(UWorld* World);
 	virtual UWorld* GetTargetWorld();
 
diff --git a/Engine/Plugins/Runtime/MeshModelingToolset/Source/ModelingComponents/Public/BaseTools/SingleSelectionMeshEditingTool.h b/Engine/Plugins/Runtime/MeshModelingToolset/Source/ModelingComponents/Public/BaseTools/SingleSelectionMeshEditingTool.h
--- a/Engine/Plugins/Runtime/MeshModelingToolset/Source/ModelingComponents/Public/BaseTools/SingleSelectionMeshEditingTool.h
+++ b/Engine/Plugins/Runtime/MeshModelingToolset/Source/ModelingComponents/Public/BaseTools/SingleSelectionMeshEditingTool.h
@@ -47,8 +47,10 @@ UCLASS()
 class MODELINGCOMPONENTS_API USingleSelectionMeshEditingTool : public USingleSelectionTool
 {
 	GENERATED_BODY()
-
 public:
+	virtual void Shutdown(EToolShutdownType ShutdownType) override;
+	virtual void OnShutdown(EToolShutdownType ShutdownType);
+
 	virtual void SetWorld(UWorld* World);
 	virtual UWorld* GetTargetWorld();
 
