ModelingTools: Fix crash when reloading level during a BakeTx or BakeAll compute. Ensure that TargetMeshUVCharts is still valid for the background compute if the tool is destroyed.

#rnx

diff --git a/Engine/Plugins/Experimental/MeshModelingToolsetExp/Source/MeshModelingToolsExp/Private/BakeMeshAttributeMapsTool.cpp b/Engine/Plugins/Experimental/MeshModelingToolsetExp/Source/MeshModelingToolsExp/Private/BakeMeshAttributeMapsTool.cpp
--- a/Engine/Plugins/Experimental/MeshModelingToolsetExp/Source/MeshModelingToolsExp/Private/BakeMeshAttributeMapsTool.cpp
+++ b/Engine/Plugins/Experimental/MeshModelingToolsetExp/Source/MeshModelingToolsExp/Private/BakeMeshAttributeMapsTool.cpp
@@ -98,11 +98,11 @@ public:
 	TSharedPtr<UE::Geometry::FDynamicMesh3, ESPMode::ThreadSafe> DetailMesh;
 	TSharedPtr<UE::Geometry::FDynamicMeshAABBTree3, ESPMode::ThreadSafe> DetailSpatial;
 	TSharedPtr<UE::Geometry::TMeshTangents<double>, ESPMode::ThreadSafe> DetailMeshTangents;
-	UE::Geometry::FDynamicMesh3* BaseMesh;
+	UE::Geometry::FDynamicMesh3* BaseMesh = nullptr;
 	TUniquePtr<UE::Geometry::FMeshMapBaker> Baker;
 	UBakeMeshAttributeMapsTool::FBakeSettings BakeSettings;
 	TSharedPtr<UE::Geometry::TMeshTangents<double>, ESPMode::ThreadSafe> BaseMeshTangents;
-	TArray<int32>* BaseMeshUVCharts;
+	TSharedPtr<TArray<int32>, ESPMode::ThreadSafe> BaseMeshUVCharts;
 
 	// Map Type settings
 	EBakeMapType Maps;
@@ -133,7 +133,7 @@ public:
 		Baker->SetDimensions(BakeSettings.Dimensions);
 		Baker->SetProjectionDistance(BakeSettings.ProjectionDistance);
 		Baker->SetSamplesPerPixel(BakeSettings.SamplesPerPixel);
-		Baker->SetTargetMeshUVCharts(BaseMeshUVCharts);
+		Baker->SetTargetMeshUVCharts(BaseMeshUVCharts.Get());
 		Baker->SetTargetMeshTangents(BaseMeshTangents);
 		
 		FMeshBakerDynamicMeshSampler DetailSampler(DetailMesh.Get(), DetailSpatial.Get(), DetailMeshTangents.Get());
@@ -403,7 +403,7 @@ TUniquePtr<UE::Geometry::TGenericDataOperator<FMeshMapBaker>> UBakeMeshAttribute
 	Op->DetailSpatial = DetailSpatial;
 	Op->BaseMesh = &TargetMesh;
 	Op->BakeSettings = CachedBakeSettings;
-	Op->BaseMeshUVCharts = &TargetMeshUVCharts;
+	Op->BaseMeshUVCharts = TargetMeshUVCharts;
 
 	constexpr EBakeMapType RequiresTangents = EBakeMapType::TangentSpaceNormal | EBakeMapType::BentNormal;
 	if ((bool)(CachedBakeSettings.BakeMapTypes & RequiresTangents))
diff --git a/Engine/Plugins/Experimental/MeshModelingToolsetExp/Source/MeshModelingToolsExp/Private/BakeMeshAttributeMapsToolBase.cpp b/Engine/Plugins/Experimental/MeshModelingToolsetExp/Source/MeshModelingToolsExp/Private/BakeMeshAttributeMapsToolBase.cpp
--- a/Engine/Plugins/Experimental/MeshModelingToolsetExp/Source/MeshModelingToolsExp/Private/BakeMeshAttributeMapsToolBase.cpp
+++ b/Engine/Plugins/Experimental/MeshModelingToolsetExp/Source/MeshModelingToolsExp/Private/BakeMeshAttributeMapsToolBase.cpp
@@ -82,7 +82,8 @@ void UBakeMeshAttributeMapsToolBase::PostSetup()
 
 	// Initialize UV charts
 	// TODO: Compute UV charts asynchronously
-	FMeshMapBaker::ComputeUVCharts(TargetMesh, TargetMeshUVCharts);
+	TargetMeshUVCharts = MakeShared<TArray<int32>, ESPMode::ThreadSafe>();
+	FMeshMapBaker::ComputeUVCharts(TargetMesh, *TargetMeshUVCharts);
 
 	GatherAnalytics(BakeAnalytics.MeshSettings);
 }
diff --git a/Engine/Plugins/Experimental/MeshModelingToolsetExp/Source/MeshModelingToolsExp/Private/BakeMultiMeshAttributeMapsTool.cpp b/Engine/Plugins/Experimental/MeshModelingToolsetExp/Source/MeshModelingToolsExp/Private/BakeMultiMeshAttributeMapsTool.cpp
--- a/Engine/Plugins/Experimental/MeshModelingToolsetExp/Source/MeshModelingToolsExp/Private/BakeMultiMeshAttributeMapsTool.cpp
+++ b/Engine/Plugins/Experimental/MeshModelingToolsetExp/Source/MeshModelingToolsExp/Private/BakeMultiMeshAttributeMapsTool.cpp
@@ -315,11 +315,11 @@ class FMultiMeshMapBakerOp : public TGenericDataOperator<FMeshMapBaker>
 public:
 	// General bake settings
 	FMeshSceneAdapter* DetailMeshScene = nullptr;
-	UE::Geometry::FDynamicMesh3* BaseMesh;
+	UE::Geometry::FDynamicMesh3* BaseMesh = nullptr;
 	TSharedPtr<UE::Geometry::FMeshTangentsd, ESPMode::ThreadSafe> BaseMeshTangents;
 	TUniquePtr<UE::Geometry::FMeshMapBaker> Baker;
 	UBakeMultiMeshAttributeMapsTool::FBakeSettings BakeSettings;
-	TArray<int32>* BaseMeshUVCharts;
+	TSharedPtr<TArray<int32>, ESPMode::ThreadSafe> BaseMeshUVCharts;
 
 	// Detail bake data
 	TArray<TSharedPtr<UE::Geometry::TImageBuilder<FVector4f>>> CachedColorImages;
@@ -338,7 +338,7 @@ public:
 		Baker->SetProjectionDistance(BakeSettings.ProjectionDistance);
 		Baker->SetSamplesPerPixel(BakeSettings.SamplesPerPixel);
 		Baker->SetTargetMeshTangents(BaseMeshTangents);
-		Baker->SetTargetMeshUVCharts(BaseMeshUVCharts);
+		Baker->SetTargetMeshUVCharts(BaseMeshUVCharts.Get());
 		
 		FMeshBakerMeshSceneSampler DetailSampler(DetailMeshScene);
 		Baker->SetDetailSampler(&DetailSampler);
@@ -505,7 +505,7 @@ TUniquePtr<UE::Geometry::TGenericDataOperator<FMeshMapBaker>> UBakeMultiMeshAttr
 	TUniquePtr<FMultiMeshMapBakerOp> Op = MakeUnique<FMultiMeshMapBakerOp>();
 	Op->DetailMeshScene = &DetailMeshScene;
 	Op->BaseMesh = &TargetMesh;
-	Op->BaseMeshUVCharts = &TargetMeshUVCharts;
+	Op->BaseMeshUVCharts = TargetMeshUVCharts;
 	Op->BakeSettings = CachedBakeSettings;
 
 	constexpr EBakeMapType RequiresTangents = EBakeMapType::TangentSpaceNormal | EBakeMapType::BentNormal;
diff --git a/Engine/Plugins/Experimental/MeshModelingToolsetExp/Source/MeshModelingToolsExp/Public/BakeMeshAttributeMapsToolBase.h b/Engine/Plugins/Experimental/MeshModelingToolsetExp/Source/MeshModelingToolsExp/Public/BakeMeshAttributeMapsToolBase.h
--- a/Engine/Plugins/Experimental/MeshModelingToolsetExp/Source/MeshModelingToolsExp/Public/BakeMeshAttributeMapsToolBase.h
+++ b/Engine/Plugins/Experimental/MeshModelingToolsetExp/Source/MeshModelingToolsExp/Public/BakeMeshAttributeMapsToolBase.h
@@ -176,7 +176,7 @@ protected:
 	};
 	FBakeSettings CachedBakeSettings;
 
-	TArray<int32> TargetMeshUVCharts;
+	TSharedPtr<TArray<int32>, ESPMode::ThreadSafe> TargetMeshUVCharts;
 
 	/**
 	 * To be invoked by client when bake map types change.
