ModelingTools: Fix SComboPanel layout to ensure that the border of ComboPanelIcons does not exceed the icon dimensions. Also, adjust SComboPanel background color per UX design guidance.

#rnx

diff --git a/Engine/Plugins/Experimental/MeshModelingToolsetExp/Source/ModelingEditorUI/Private/ModelingWidgets/SComboPanel.cpp b/Engine/Plugins/Experimental/MeshModelingToolsetExp/Source/ModelingEditorUI/Private/ModelingWidgets/SComboPanel.cpp
--- a/Engine/Plugins/Experimental/MeshModelingToolsetExp/Source/ModelingEditorUI/Private/ModelingWidgets/SComboPanel.cpp
+++ b/Engine/Plugins/Experimental/MeshModelingToolsetExp/Source/ModelingEditorUI/Private/ModelingWidgets/SComboPanel.cpp
@@ -112,7 +112,6 @@ void SComboPanelIconTile::Construct(const FArguments& InArgs)
 				SNew(SVerticalBox)
 				// Thumbnail
 				+ SVerticalBox::Slot()
-				.AutoHeight()
 				.HAlign(HAlign_Center)
 				.VAlign(VAlign_Center)
 				[
@@ -125,6 +124,11 @@ void SComboPanelIconTile::Construct(const FArguments& InArgs)
 						.BorderImage(FAppStyle::Get().GetBrush("ProjectBrowser.ProjectTile.ThumbnailAreaBackground"))
 						.HAlign(HAlign_Center)
 						.VAlign(VAlign_Top)
+						// TODO: Sort out the color transformation for BorderBackgroundColor to use FAppStyle
+						// There is some color transformation beyond just SRGB that cause the colors to appear
+						// darker than the actual color.
+						//.BorderBackgroundColor(FAppStyle::Get().GetSlateColor("Colors.Recessed"))
+						.BorderBackgroundColor(FLinearColor(0.16f, 0.16f, 0.16f))
 						[
 							SNew(SImage)
 							.Image(Image)
