MeshModelingToolset: Clean-up module dependencies

diff --git a/Engine/Plugins/Editor/ModelingToolsEditorMode/Source/ModelingToolsEditorMode/ModelingToolsEditorMode.Build.cs b/Engine/Plugins/Editor/ModelingToolsEditorMode/Source/ModelingToolsEditorMode/ModelingToolsEditorMode.Build.cs
--- a/Engine/Plugins/Editor/ModelingToolsEditorMode/Source/ModelingToolsEditorMode/ModelingToolsEditorMode.Build.cs
+++ b/Engine/Plugins/Editor/ModelingToolsEditorMode/Source/ModelingToolsEditorMode/ModelingToolsEditorMode.Build.cs
@@ -7,55 +7,46 @@ public class ModelingToolsEditorMode : ModuleRules
 	public ModelingToolsEditorMode(ReadOnlyTargetRules Target) : base(Target)
 	{
 		PCHUsage = ModuleRules.PCHUsageMode.UseExplicitOrSharedPCHs;
-		
+
 		PublicIncludePaths.AddRange(
-			new string[] {
-                //"ContentBrowser"
+			new string[]
+			{
 				// ... add public include paths required here ...
 			}
-			);
-				
-		
+		);
+
 		PrivateIncludePaths.AddRange(
-			new string[] {
+			new string[]
+			{
 				// ... add other private include paths required here ...
 			}
-			);
-			
-		
+		);
+
 		PublicDependencyModuleNames.AddRange(
 			new string[]
 			{
 				"Core"
 				// ... add other public dependencies that you statically link with here ...
 			}
-			);
-			
-		
+		);
+
 		PrivateDependencyModuleNames.AddRange(
 			new string[]
 			{
-                "Core",
-                "CoreUObject",
-                "ApplicationCore",
-                "Slate",
-                "SlateCore",
-                "Engine",
-                "InputCore",
+				"CoreUObject",
+				"Slate",
+				"SlateCore",
+				"Engine",
+				"InputCore",
 				"EditorFramework",
 				"UnrealEd",
-                "ContentBrowser",
 				"ContentBrowserData",
-                "LevelEditor",
 				"StatusBar",
 				"EditorStyle",
-                "Projects",
+				"Projects",
 				"TypedElementRuntime",
-				"MeshDescription",
-				"StaticMeshDescription",
-                "InteractiveToolsFramework",
+				"InteractiveToolsFramework",
 				"EditorInteractiveToolsFramework",
-				"GeometryFramework",
 				"ModelingComponents",
 				"ModelingComponentsEditorOnly",
 				"MeshModelingTools",
@@ -67,21 +58,11 @@ public class ModelingToolsEditorMode : ModuleRules
 				"EditorWidgets",
 				"ModelingEditorUI",
 				"StylusInput",
-				"DeveloperSettings"
-
+				"DeveloperSettings",
 				// ... add private dependencies that you statically link with here ...	
 			}
-			);
-		
-		
-		DynamicallyLoadedModuleNames.AddRange(
-			new string[]
-			{
-				"Settings"
-			}
-			);
+		);
 
 		PublicDefinitions.Add("WITH_PROXYLOD=" + (Target.Platform == UnrealTargetPlatform.Win64 ? '1' : '0'));
-
 	}
-}
+}
\ No newline at end of file
diff --git a/Engine/Plugins/Experimental/MeshLODToolset/Source/MeshLODToolset/MeshLODToolset.Build.cs b/Engine/Plugins/Experimental/MeshLODToolset/Source/MeshLODToolset/MeshLODToolset.Build.cs
--- a/Engine/Plugins/Experimental/MeshLODToolset/Source/MeshLODToolset/MeshLODToolset.Build.cs
+++ b/Engine/Plugins/Experimental/MeshLODToolset/Source/MeshLODToolset/MeshLODToolset.Build.cs
@@ -7,84 +7,59 @@ public class MeshLODToolset : ModuleRules
 	public MeshLODToolset(ReadOnlyTargetRules Target) : base(Target)
 	{
 		PCHUsage = ModuleRules.PCHUsageMode.UseExplicitOrSharedPCHs;
-		
+
 		PublicIncludePaths.AddRange(
-			new string[] {
+			new string[]
+			{
 				// ... add public include paths required here ...
 			}
-			);
-				
-		
+		);
+
 		PrivateIncludePaths.AddRange(
-			new string[] {
+			new string[]
+			{
 				// ... add other private include paths required here ...
 			}
-			);
-			
-		
+		);
+
 		PublicDependencyModuleNames.AddRange(
 			new string[]
 			{
 				"Core",
-				"Eigen",
-                "InteractiveToolsFramework",
+				"InteractiveToolsFramework",
 				"GeometryCore",
-				"GeometryAlgorithms",
 				"DynamicMesh",
 				"MeshConversion",
 				"MeshConversionEngineTypes",
 				"MeshDescription",
-                "StaticMeshDescription",
+				"StaticMeshDescription",
 				"ModelingComponents",
-				"ModelingOperators",
 				"MeshModelingToolsExp",
 				"GeometryFlowCore",
 				"GeometryFlowMeshProcessing",
-				"GeometryFlowMeshProcessingEditor"
-
+				"GeometryFlowMeshProcessingEditor",
 				// ... add other public dependencies that you statically link with here ...
 			}
-            );
-			
-		
+		);
+
 		PrivateDependencyModuleNames.AddRange(
 			new string[]
 			{
-
-				//"RenderCore",
-				//"RHI",
-				//"MeshUtilities",    // temp for saving mesh asset
-				//"UnrealEd",
-				//"MeshBuilder",
-                //"MeshDescriptionOperations",
-				//"MeshUtilitiesCommon",
-				//"MeshReductionInterface", // for UE4 standard simplification
-                //"ProxyLODMeshReduction", // for mesh merging voxel-based csg
-				//"Slate",
-				//"SlateCore",
-
 				"CoreUObject",
 				"Engine",
-				"ModelingOperators",
-				"InputCore",
-				"PhysicsCore",
 				"RenderCore",
-
 				"UnrealEd",
 				"EditorScriptingUtilities",
-				"AssetTools",
-				"Slate"			// for FSlateNotificationManager, etc.
-
+				"Slate",
 				// ... add private dependencies that you statically link with here ...
 			}
-			);
-		
-		
+		);
+
 		DynamicallyLoadedModuleNames.AddRange(
 			new string[]
 			{
 				// ... add any modules that your module loads dynamically here ...
 			}
-			);
+		);
 	}
-}
+}
\ No newline at end of file
diff --git a/Engine/Plugins/Experimental/MeshModelingToolsetExp/Source/MeshModelingToolsEditorOnlyExp/MeshModelingToolsEditorOnlyExp.Build.cs b/Engine/Plugins/Experimental/MeshModelingToolsetExp/Source/MeshModelingToolsEditorOnlyExp/MeshModelingToolsEditorOnlyExp.Build.cs
--- a/Engine/Plugins/Experimental/MeshModelingToolsetExp/Source/MeshModelingToolsEditorOnlyExp/MeshModelingToolsEditorOnlyExp.Build.cs
+++ b/Engine/Plugins/Experimental/MeshModelingToolsetExp/Source/MeshModelingToolsEditorOnlyExp/MeshModelingToolsEditorOnlyExp.Build.cs
@@ -7,27 +7,27 @@ public class MeshModelingToolsEditorOnlyExp : ModuleRules
 	public MeshModelingToolsEditorOnlyExp(ReadOnlyTargetRules Target) : base(Target)
 	{
 		PCHUsage = ModuleRules.PCHUsageMode.UseExplicitOrSharedPCHs;
-		
+
 		PublicIncludePaths.AddRange(
-			new string[] {
+			new string[]
+			{
 				// ... add public include paths required here ...
 			}
-			);
-				
-		
+		);
+
 		PrivateIncludePaths.AddRange(
-			new string[] {
+			new string[]
+			{
 				// ... add other private include paths required here ...
 			}
-			);
-			
-		
+		);
+
 		PublicDependencyModuleNames.AddRange(
 			new string[]
 			{
 				"Core",
-                "InteractiveToolsFramework",
-                "MeshDescription",
+				"InteractiveToolsFramework",
+				"MeshDescription",
 				"StaticMeshDescription",
 				"SkeletalMeshDescription",
 				"GeometryCore",
@@ -42,46 +42,26 @@ public class MeshModelingToolsEditorOnlyExp : ModuleRules
 				"ModelingOperatorsEditorOnly",
 				// ... add other public dependencies that you statically link with here ...
 			}
-            );
-			
-		
+		);
+
 		PrivateDependencyModuleNames.AddRange(
 			new string[]
 			{
 				"BSPUtils",
 				"CoreUObject",
 				"Engine",
-                "RenderCore",
-                "RHI",
+				"RenderCore",
 				"InputCore",
-
-				"MeshUtilities",		// tangents calculation
-				"EditorFramework",
 				"UnrealEd",
-				"MeshBuilder",
-				"MeshUtilitiesCommon",  
-				"MeshReductionInterface", // for UE standard simplification 
-				//"Slate",
-				//"SlateCore",
 				// ... add private dependencies that you statically link with here ...	
 			}
-			);
-
-		if (Target.Platform == UnrealTargetPlatform.Win64)
-		{
-			PrivateDependencyModuleNames.AddRange(
-				new string[]
-				{
-					"ProxyLODMeshReduction", // currently Win64-only
-				}
-				);
-		}
+		);
 
 		DynamicallyLoadedModuleNames.AddRange(
 			new string[]
 			{
 				// ... add any modules that your module loads dynamically here ...
 			}
-			);
+		);
 	}
-}
+}
\ No newline at end of file
diff --git a/Engine/Plugins/Experimental/MeshModelingToolsetExp/Source/MeshModelingToolsExp/MeshModelingToolsExp.Build.cs b/Engine/Plugins/Experimental/MeshModelingToolsetExp/Source/MeshModelingToolsExp/MeshModelingToolsExp.Build.cs
--- a/Engine/Plugins/Experimental/MeshModelingToolsetExp/Source/MeshModelingToolsExp/MeshModelingToolsExp.Build.cs
+++ b/Engine/Plugins/Experimental/MeshModelingToolsetExp/Source/MeshModelingToolsExp/MeshModelingToolsExp.Build.cs
@@ -7,87 +7,57 @@ public class MeshModelingToolsExp : ModuleRules
 	public MeshModelingToolsExp(ReadOnlyTargetRules Target) : base(Target)
 	{
 		PCHUsage = ModuleRules.PCHUsageMode.UseExplicitOrSharedPCHs;
-		
+
 		PublicIncludePaths.AddRange(
-			new string[] {
+			new string[]
+			{
 				// ... add public include paths required here ...
 			}
-			);
-				
-		PublicIncludePathModuleNames.AddRange(
-			new string[] {
-				"AnimationCore",			// For the BoneWeights.h include
+		);
+
+		PrivateIncludePaths.AddRange(
+			new string[]
+			{
+				// ... add other private include paths required here ...
 			}
 		);
-	
+
 		PublicDependencyModuleNames.AddRange(
 			new string[]
 			{
-				// These dependencies were commented out when we split out MeshModelingToolsEditorOnly
-				// Many of them are editor only and should not be depended on here.  If you re-add any
-				// dependent modules, please confirm by "launching" or packaging that you are not introducing
-				// editor dependencies
-				//
-				//"MeshDescription",
-				//"ProxyLODMeshReduction", // currently required to be public due to IVoxelBasedCSG API
-
 				"Core",
-				"Eigen",
-                "InteractiveToolsFramework",
+				"InteractiveToolsFramework",
 				"GeometryCore",
 				"GeometryFramework",
-				"GeometryAlgorithms",
 				"DynamicMesh",
 				"MeshConversion",
 				"MeshDescription",
 				"MeshModelingTools",
-                "StaticMeshDescription",
-                "SkeletalMeshDescription",
+				"StaticMeshDescription",
+				"SkeletalMeshDescription",
 				"ModelingComponents",
 				"ModelingOperators",
-
 				// ... add other public dependencies that you statically link with here ...
 			}
-            );
-			
-		
+		);
+
 		PrivateDependencyModuleNames.AddRange(
 			new string[]
 			{
-				// These dependencies were commented out when we split out MeshModelingToolsEditorOnly
-				// Many of them are editor only and should not be depended on here.  If you re-add any
-				// dependent modules, please confirm by "launching" or packaging that you are not introducing
-				// editor dependencies
-				//
-				//"RenderCore",
-				//"RHI",
-				//"MeshUtilities",    // temp for saving mesh asset
-				//"UnrealEd",
-				//"MeshBuilder",
-                //"MeshDescriptionOperations",
-				//"MeshUtilitiesCommon",
-				//"MeshReductionInterface", // for UE standard simplification
-                //"ProxyLODMeshReduction", // for mesh merging voxel-based csg
-				//"Slate",
-				//"SlateCore",
-
 				"CoreUObject",
 				"Engine",
 				"RenderCore",
-				"ModelingOperators",
 				"InputCore",
 				"PhysicsCore",
-
 				// ... add private dependencies that you statically link with here ...
 			}
-			);
-		
-		
+		);
+
 		DynamicallyLoadedModuleNames.AddRange(
 			new string[]
 			{
 				// ... add any modules that your module loads dynamically here ...
 			}
-			);
+		);
 	}
-}
+}
\ No newline at end of file
diff --git a/Engine/Plugins/Runtime/MeshModelingToolset/Source/MeshModelingTools/MeshModelingTools.Build.cs b/Engine/Plugins/Runtime/MeshModelingToolset/Source/MeshModelingTools/MeshModelingTools.Build.cs
--- a/Engine/Plugins/Runtime/MeshModelingToolset/Source/MeshModelingTools/MeshModelingTools.Build.cs
+++ b/Engine/Plugins/Runtime/MeshModelingToolset/Source/MeshModelingTools/MeshModelingTools.Build.cs
@@ -7,26 +7,27 @@ public class MeshModelingTools : ModuleRules
 	public MeshModelingTools(ReadOnlyTargetRules Target) : base(Target)
 	{
 		PCHUsage = ModuleRules.PCHUsageMode.UseExplicitOrSharedPCHs;
-		
+
 		PublicIncludePaths.AddRange(
-			new string[] {
+			new string[]
+			{
 				// ... add public include paths required here ...
 			}
-			);
-				
-		PublicIncludePathModuleNames.AddRange(
-			new string[] {
-				// "AnimationCore",			// For the BoneWeights.h include
+		);
+
+		PrivateIncludePaths.AddRange(
+			new string[]
+			{
+				// ... add other private include paths required here ...
 			}
 		);
-	
+
 		PublicDependencyModuleNames.AddRange(
 			new string[]
 			{
 				"Core",
-				// "Eigen",
 				"InputCore",
-                "InteractiveToolsFramework",
+				"InteractiveToolsFramework",
 				"GeometryCore",
 				"GeometryFramework",
 				"GeometryAlgorithms",
@@ -34,35 +35,26 @@ public class MeshModelingTools : ModuleRules
 				"MeshConversion",
 				"MeshDescription",
 				"StaticMeshDescription",
-				// "SkeletalMeshDescription",
 				"ModelingComponents",
 				"ModelingOperators",
-
 				// ... add other public dependencies that you statically link with here ...
 			}
-            );
-			
-		
+		);
+
 		PrivateDependencyModuleNames.AddRange(
 			new string[]
 			{
 				"CoreUObject",
 				"Engine",
-				// "RenderCore",
-				// "ModelingOperators",
-				// "InputCore",
-				// "PhysicsCore",
-
 				// ... add private dependencies that you statically link with here ...
 			}
-			);
-		
-		
+		);
+
 		DynamicallyLoadedModuleNames.AddRange(
 			new string[]
 			{
 				// ... add any modules that your module loads dynamically here ...
 			}
-			);
+		);
 	}
-}
+}
\ No newline at end of file
diff --git a/Engine/Plugins/Runtime/MeshModelingToolset/Source/ModelingComponents/ModelingComponents.Build.cs b/Engine/Plugins/Runtime/MeshModelingToolset/Source/ModelingComponents/ModelingComponents.Build.cs
--- a/Engine/Plugins/Runtime/MeshModelingToolset/Source/ModelingComponents/ModelingComponents.Build.cs
+++ b/Engine/Plugins/Runtime/MeshModelingToolset/Source/ModelingComponents/ModelingComponents.Build.cs
@@ -7,28 +7,28 @@ public class ModelingComponents : ModuleRules
 	public ModelingComponents(ReadOnlyTargetRules Target) : base(Target)
 	{
 		PCHUsage = ModuleRules.PCHUsageMode.UseExplicitOrSharedPCHs;
-		
+
 		PublicIncludePaths.AddRange(
-			new string[] {
+			new string[]
+			{
 				// ... add public include paths required here ...
 			}
-			);
-				
-		
+		);
+
 		PrivateIncludePaths.AddRange(
-			new string[] {
+			new string[]
+			{
 				// ... add other private include paths required here ...
 			}
-			);
-			
-		
+		);
+
 		PublicDependencyModuleNames.AddRange(
 			new string[]
 			{
 				"Core",
 				"PhysicsCore",
 				"InteractiveToolsFramework",
-                "MeshDescription",
+				"MeshDescription",
 				"StaticMeshDescription",
 				"GeometryCore",
 				"GeometryFramework",
@@ -36,36 +36,28 @@ public class ModelingComponents : ModuleRules
 				"DynamicMesh",
 				"MeshConversion",
 				"ModelingOperators",
-				"DeveloperSettings"
+				"DeveloperSettings",
 				// ... add other public dependencies that you statically link with here ...
 			}
-            );
-			
-		
+		);
+
 		PrivateDependencyModuleNames.AddRange(
 			new string[]
 			{
 				"CoreUObject",
 				"Engine",
-				"InputCore",
-                "RenderCore",
-                "RHI",
-				"ImageWriteQueue"
-
-				//"MeshUtilities",    // temp for saving mesh asset
-				//"UnrealEd"
-				//"Slate",
-				//"SlateCore",
-				// ... add private dependencies that you statically link with here ...	
+				"RenderCore",
+				"RHI",
+				"ImageWriteQueue",
+				// ... add private dependencies that you statically link with here ...
 			}
-			);
-		
-		
+		);
+
 		DynamicallyLoadedModuleNames.AddRange(
 			new string[]
 			{
 				// ... add any modules that your module loads dynamically here ...
 			}
-			);
+		);
 	}
-}
+}
\ No newline at end of file
diff --git a/Engine/Plugins/Runtime/MeshModelingToolset/Source/ModelingOperators/ModelingOperators.Build.cs b/Engine/Plugins/Runtime/MeshModelingToolset/Source/ModelingOperators/ModelingOperators.Build.cs
--- a/Engine/Plugins/Runtime/MeshModelingToolset/Source/ModelingOperators/ModelingOperators.Build.cs
+++ b/Engine/Plugins/Runtime/MeshModelingToolset/Source/ModelingOperators/ModelingOperators.Build.cs
@@ -7,68 +7,48 @@ public class ModelingOperators : ModuleRules
 	public ModelingOperators(ReadOnlyTargetRules Target) : base(Target)
 	{
 		PCHUsage = ModuleRules.PCHUsageMode.UseExplicitOrSharedPCHs;
-		
+
 		PublicIncludePaths.AddRange(
-			new string[] {
+			new string[]
+			{
 				// ... add public include paths required here ...
 			}
-			);
-				
-		
+		);
+
 		PrivateIncludePaths.AddRange(
-			new string[] {
+			new string[]
+			{
 				// ... add other private include paths required here ...
 			}
-			);
-			
-		
+		);
+
 		PublicDependencyModuleNames.AddRange(
 			new string[]
 			{
-				// These dependencies were commented out when we split out ModelingOperatorsEditorOnly
-				// Many of them are editor only and should not be depended on here.  If you re-add any
-				// dependent modules, please confirm by "launching" or packaging that you are not introducing
-				// editor dependencies
-				//
-                //"ProxyLODMeshReduction", // currently required to be public due to IVoxelBasedCSG API
-                //"MeshDescription",
-
 				"Core",
-                "InteractiveToolsFramework",
+				"InteractiveToolsFramework",
 				"GeometryCore",
 				"DynamicMesh",
 				"MeshConversion",
-                "GeometryAlgorithms", // required for constrained Delaunay triangulation
+				"GeometryAlgorithms", // required for constrained Delaunay triangulation
 				// ... add other public dependencies that you statically link with here ...
 			}
-            );
-			
-		
+		);
+
 		PrivateDependencyModuleNames.AddRange(
 			new string[]
 			{
-				// These dependencies were commented out when we split out ModelingOperatorsEditorOnly
-				// Many of them are editor only and should not be depended on here.  If you re-add any
-				// dependent modules, please confirm by "launching" or packaging that you are not introducing
-				// editor dependencies
-				//
-				//"MeshBuilder",
-                //"MeshDescriptionOperations",
-				//"MeshUtilitiesCommon",
-				//"MeshReductionInterface", // for standard UE simplification
-
 				"CoreUObject",
 				"Engine",
-
+				// ... add private dependencies that you statically link with here ...
 			}
-			);
-		
-		
+		);
+
 		DynamicallyLoadedModuleNames.AddRange(
 			new string[]
 			{
 				// ... add any modules that your module loads dynamically here ...
 			}
-			);
+		);
 	}
-}
+}
\ No newline at end of file
