Dynamic sculpt tool: make sure to allow the async task that initializes the remsher to finish

#rnx

diff --git a/Engine/Plugins/Experimental/MeshModelingToolsetExp/Source/MeshModelingToolsExp/Private/DynamicMeshSculptTool.cpp b/Engine/Plugins/Experimental/MeshModelingToolsetExp/Source/MeshModelingToolsExp/Private/DynamicMeshSculptTool.cpp
--- a/Engine/Plugins/Experimental/MeshModelingToolsetExp/Source/MeshModelingToolsExp/Private/DynamicMeshSculptTool.cpp
+++ b/Engine/Plugins/Experimental/MeshModelingToolsetExp/Source/MeshModelingToolsExp/Private/DynamicMeshSculptTool.cpp
@@ -1554,8 +1554,9 @@ void UDynamicMeshSculptTool::OnTick(float DeltaTime)
 	bool bROIUpdatePending = false;
 	bool bOctreeUpdatePending = false;
 
-
 	TFuture<void> InitializeRemesher;
+	TFuture<void> PrecomputeRemeshROI;
+
 	if (bStampPending)
 	{
 		// if we don't have an active remesher for this brush stroke, create one
@@ -1575,7 +1576,7 @@ void UDynamicMeshSculptTool::OnTick(float DeltaTime)
 
 
 		// once we know ROI we can speculatively start initializing remesh ROI
-		TFuture<void> PrecomputeRemeshROI = Async(DynamicSculptToolAsyncExecTarget, [&]()
+		PrecomputeRemeshROI = Async(DynamicSculptToolAsyncExecTarget, [&]()
 		{
 			// make sure our remesher is initialized
 			InitializeRemesher.Wait();
@@ -1681,6 +1682,9 @@ void UDynamicMeshSculptTool::OnTick(float DeltaTime)
 		bMeshModified = false;
 	}
 
+	// Allow futures to finish in case bRemeshPending == false
+	InitializeRemesher.Wait();
+	PrecomputeRemeshROI.Wait();
 }
 
 void UDynamicMeshSculptTool::PrecomputeRemeshInfo()
