Fix a crash in DynaSculpt tool. Wait for async task to finish before starting a new one.

#rnx

diff --git a/Engine/Plugins/Experimental/MeshModelingToolsetExp/Source/MeshModelingToolsExp/Private/DynamicMeshSculptTool.cpp b/Engine/Plugins/Experimental/MeshModelingToolsetExp/Source/MeshModelingToolsExp/Private/DynamicMeshSculptTool.cpp
--- a/Engine/Plugins/Experimental/MeshModelingToolsetExp/Source/MeshModelingToolsExp/Private/DynamicMeshSculptTool.cpp
+++ b/Engine/Plugins/Experimental/MeshModelingToolsetExp/Source/MeshModelingToolsExp/Private/DynamicMeshSculptTool.cpp
@@ -2107,6 +2107,9 @@ void UDynamicMeshSculptTool::UpdateTarget()
 		}
 	}
 
+	// Allow any in-progress update to finish before starting a new one
+	PendingTargetUpdate.Wait();
+
 	// TODO: could have a second set of target meshes, and swap between them while we update the other one
 	PendingTargetUpdate = Async(DynamicSculptToolAsyncExecTarget, [this]()
 	{
