ApproximateActors: re-organize computation to reduce peak memory usage
- MeshSceneAdapter instance is now passed to GenerateApproximationMesh(), to be released as soon as possible
- generated SolidMesh is now discarded once Morphology result is computed (never used again)
- FColliderMesh is used to create projection-target mesh copy for Simplification w/ geometric tolerance, instead of DynamicMesh (fewer buffers)
- CapturePhotoSet is no longer run in parallel with mesh generation. This costs a bit of performance but significantly reduces peak memory usage
- Add Insights bookmarks

#rnx

diff --git a/Engine/Plugins/Experimental/MeshModelingToolsetExp/Source/GeometryProcessingAdapters/Private/ApproximateActorsImpl.cpp b/Engine/Plugins/Experimental/MeshModelingToolsetExp/Source/GeometryProcessingAdapters/Private/ApproximateActorsImpl.cpp
--- a/Engine/Plugins/Experimental/MeshModelingToolsetExp/Source/GeometryProcessingAdapters/Private/ApproximateActorsImpl.cpp
+++ b/Engine/Plugins/Experimental/MeshModelingToolsetExp/Source/GeometryProcessingAdapters/Private/ApproximateActorsImpl.cpp
@@ -6,6 +6,7 @@
 #include "DynamicMesh/DynamicMesh3.h"
 #include "DynamicMesh/MeshNormals.h"
 #include "DynamicMesh/MeshTangents.h"
+#include "DynamicMesh/ColliderMesh.h"
 #include "MeshSimplification.h"
 #include "MeshConstraintsUtil.h"
 #include "Implicit/Solidify.h"
@@ -377,12 +378,13 @@ struct FApproximationMeshData
 
 
 static TSharedPtr<FApproximationMeshData> GenerateApproximationMesh(
-	FMeshSceneAdapter& Scene,
+	TUniquePtr<FMeshSceneAdapter> Scene,
 	const IGeometryProcessing_ApproximateActors::FOptions& Options,
 	double ApproxAccuracy
 )
 {
 	FScopedSlowTask Progress(8.f, LOCTEXT("Generating Mesh", "Generating Mesh.."));
+	TRACE_BOOKMARK(TEXT("ApproximateActors-Collect Seed Points"));
 
 	TSharedPtr<FApproximationMeshData> Result = MakeShared<FApproximationMeshData>();
 
@@ -390,9 +392,9 @@ static TSharedPtr<FApproximationMeshData> GenerateApproximationMesh(
 	TArray<FVector3d> SeedPoints;
 	{
 		TRACE_CPUPROFILER_EVENT_SCOPE(ApproximateActorsImpl_Generate_SeedPoints);
-		Scene.CollectMeshSeedPoints(SeedPoints);
+		Scene->CollectMeshSeedPoints(SeedPoints);
 	}
-	FAxisAlignedBox3d SceneBounds = Scene.GetBoundingBox();
+	FAxisAlignedBox3d SceneBounds = Scene->GetBoundingBox();
 
 	// calculate a voxel size based on target world-space approximation accuracy
 	float WorldBoundsSize = SceneBounds.DiagonalLength();
@@ -422,31 +424,45 @@ static TSharedPtr<FApproximationMeshData> GenerateApproximationMesh(
 
 	Progress.EnterProgressFrame(1.f, LOCTEXT("SolidifyMesh", "Approximating Mesh..."));
 
-	FWindingNumberBasedSolidify Solidify(
-		[&Scene](const FVector3d& Position) { return Scene.FastWindingNumber(Position, true); },
-		SceneBounds, SeedPoints);
-	Solidify.SetCellSizeAndExtendBounds(SceneBounds, 2.0 * ApproxAccuracy, VoxelDimTarget);
-	Solidify.WindingThreshold = Options.WindingThreshold;
+	TRACE_BOOKMARK(TEXT("ApproximateActors-Solidify"));
 
 	FDynamicMesh3 SolidMesh;
 	{
+		// Do in local scope so that memory allocated in Solidify is released after SolidMesh is available
+		// TODO: SolidMesh could be replaced with a FColliderMesh if TImplicitMorphology would allow it
+
 		TRACE_CPUPROFILER_EVENT_SCOPE(ApproximateActorsImpl_Generate_Solidify);
+		FWindingNumberBasedSolidify Solidify(
+			[&Scene](const FVector3d& Position) { return Scene->FastWindingNumber(Position, true); },
+			SceneBounds, SeedPoints);
+		Solidify.SetCellSizeAndExtendBounds(SceneBounds, 2.0 * ApproxAccuracy, VoxelDimTarget);
+		Solidify.WindingThreshold = Options.WindingThreshold;
+
 		SolidMesh = FDynamicMesh3(&Solidify.Generate());
+		SolidMesh.DiscardAttributes();
 	}
-	SolidMesh.DiscardAttributes();
-	FDynamicMesh3* CurResultMesh = &SolidMesh;		// this pointer will be updated as we recompute the mesh
+
+	// CurResultMesh will point to the "current" result and we will update this pointer as we
+	// step through various stages of the generation process
+	FDynamicMesh3* CurResultMesh = &SolidMesh;		
 
 	if (Options.bVerbose)
 	{
 		UE_LOG(LogApproximateActors, Log, TEXT("Solidify mesh has %d triangles"), CurResultMesh->TriangleCount());
 	}
 
+	// we are done w/ the FMeshSceneAdapter now and can free it's memory
+	Scene.Reset();
+	SeedPoints.Empty();
+
 	Progress.EnterProgressFrame(1.f, LOCTEXT("ClosingMesh", "Topological Operations..."));
 
 	// do topological closure to fix small gaps/etc
 	FDynamicMesh3 MorphologyMesh;
 	if (Options.bApplyMorphology)
 	{
+		TRACE_BOOKMARK(TEXT("ApproximateActors-Morphology"));
+
 		TRACE_CPUPROFILER_EVENT_SCOPE(ApproximateActorsImpl_Generate_Morphology);
 		double MorphologyDistance = Options.MorphologyDistanceMeters * 100.0;		// convert to cm
 		FAxisAlignedBox3d MorphologyBounds = CurResultMesh->GetBounds();
@@ -458,7 +474,9 @@ static TSharedPtr<FApproximationMeshData> GenerateApproximationMesh(
 		ImplicitMorphology.SetCellSizesAndDistance(MorphologyBounds, MorphologyDistance, VoxelDimTarget, VoxelDimTarget);
 		MorphologyMesh = FDynamicMesh3(&ImplicitMorphology.Generate());
 		MorphologyMesh.DiscardAttributes();
+
 		CurResultMesh = &MorphologyMesh;
+		SolidMesh = FDynamicMesh3();		// we are done with SolidMesh now, release it's memory
 
 		if (Options.bVerbose)
 		{
@@ -477,6 +495,8 @@ static TSharedPtr<FApproximationMeshData> GenerateApproximationMesh(
 
 	Progress.EnterProgressFrame(1.f, LOCTEXT("SimplifyingMesh", "Simplifying Mesh..."));
 
+	TRACE_BOOKMARK(TEXT("ApproximateActors-Simplify"));
+
 	FVolPresMeshSimplification Simplifier(CurResultMesh);
 	Simplifier.ProjectionMode = FVolPresMeshSimplification::ETargetProjectionMode::NoProjection;
 	Simplifier.DEBUG_CHECK_LEVEL = 0;
@@ -512,10 +532,12 @@ static TSharedPtr<FApproximationMeshData> GenerateApproximationMesh(
 				Simplifier.SimplifyToTriangleCount(50000);
 			}
 
-			FDynamicMesh3 MeshCopy(*CurResultMesh);
-			FDynamicMeshAABBTree3 MeshCopySpatial(&MeshCopy, true);
-			FMeshProjectionTarget ProjectionTarget(&MeshCopy, &MeshCopySpatial);
-			Simplifier.SetProjectionTarget(&ProjectionTarget);
+			// make copy of mesh geometry to use for geometric error measurement
+			FColliderMesh Collider;
+			Collider.Initialize(*CurResultMesh);
+			FColliderMeshProjectionTarget ColliderTarget(&Collider);
+			Simplifier.SetProjectionTarget(&ColliderTarget);
+
 			Simplifier.GeometricErrorConstraint = FVolPresMeshSimplification::EGeometricErrorCriteria::PredictedPointToProjectionTarget;
 			Simplifier.GeometricErrorTolerance = UseTargetTolerance;
 			{
@@ -537,6 +559,7 @@ static TSharedPtr<FApproximationMeshData> GenerateApproximationMesh(
 
 
 	Progress.EnterProgressFrame(1.f, LOCTEXT("RemoveHidden", "Removing Hidden Geometry..."));
+	TRACE_BOOKMARK(TEXT("ApproximateActors-Remove Hidden"));
 
 	if (Options.OcclusionPolicy == IGeometryProcessing_ApproximateActors::EOcclusionPolicy::VisibilityBased)
 	{
@@ -625,6 +648,8 @@ static TSharedPtr<FApproximationMeshData> GenerateApproximationMesh(
 	}
 
 
+	TRACE_BOOKMARK(TEXT("ApproximateActors-Clip Ground"));
+
 	if (Options.GroundPlaneClippingPolicy == IGeometryProcessing_ApproximateActors::EGroundPlaneClippingPolicy::CutFaces ||
 		Options.GroundPlaneClippingPolicy == IGeometryProcessing_ApproximateActors::EGroundPlaneClippingPolicy::CutFacesAndFill)
 	{
@@ -643,6 +668,7 @@ static TSharedPtr<FApproximationMeshData> GenerateApproximationMesh(
 		}
 	}
 
+	TRACE_BOOKMARK(TEXT("ApproximateActors-Normals and UVs"));
 
 	// re-enable attributes
 	CurResultMesh->EnableAttributes();
@@ -975,19 +1001,19 @@ void FApproximateActorsImpl::GenerateApproximationForActorSet(const TArray<AActo
 			return;
 		}
 	}
+
 	//
-	// Future Optimizations
-	// 	   - can do most of the mesh processing at the same time as capturing the photo set (if that matters)
-	//     - some parts of mesh gen can be done simultaneously (maybe?)
+	// extract all visible meshes from the list of Actors and build up a "mesh scene" that represents the
+	// assembly of geometry, taking advantage of instancing where possible, but also pulling those meshes
+	// apart and processing them if necessary to (eg) thicken them, etc, so that later processing works better.
+	// FMeshSceneAdapter does all the heavy lifting, here it is just being configured and built
 	//
-
 	FScopedSlowTask Progress(11.f, LOCTEXT("ApproximatingActors", "Generating Actor Approximation..."));
-
 	Progress.EnterProgressFrame(1.f, LOCTEXT("BuildingScene", "Building Scene..."));
 
 	float ApproxAccuracy = Options.WorldSpaceApproximationAccuracyMeters * 100.0;		// convert to cm (UE Units)
 
-	FMeshSceneAdapter Scene;
+	TUniquePtr<FMeshSceneAdapter> Scene = MakeUnique<FMeshSceneAdapter>();
 	FMeshSceneAdapterBuildOptions SceneBuildOptions;
 	SceneBuildOptions.bThickenThinMeshes = Options.bAutoThickenThinParts;
 	SceneBuildOptions.DesiredMinThickness = Options.AutoThickenThicknessMeters * 100.0;		// convert to cm (UE Units)
@@ -998,21 +1024,25 @@ void FApproximateActorsImpl::GenerateApproximationForActorSet(const TArray<AActo
 	SceneBuildOptions.bPrintDebugMessages = Options.bVerbose;
 	{
 		TRACE_CPUPROFILER_EVENT_SCOPE(ApproximateActorsImpl_Generate_BuildScene);
-		Scene.AddActors(Actors);
-		Scene.Build(SceneBuildOptions);
+		TRACE_BOOKMARK(TEXT("ApproximateActors-Adding Actors"));
+		Scene->AddActors(Actors);
+		TRACE_BOOKMARK(TEXT("ApproximateActors-Building Scene"));
+		Scene->Build(SceneBuildOptions);
 	}
 
-	// todo: make optional
 	if (Options.bVerbose)
 	{
 		FMeshSceneAdapter::FStatistics Stats;
-		Scene.GetGeometryStatistics(Stats);
+		Scene->GetGeometryStatistics(Stats);
 		UE_LOG(LogApproximateActors, Log, TEXT("%lld triangles in %lld unique meshes, total %lld triangles in %lld instances"),
 			Stats.UniqueMeshTriangleCount, Stats.UniqueMeshCount, Stats.InstanceMeshTriangleCount, Stats.InstanceMeshCount);
 	}
 
+	// add a "base cap" if desired, this helps with (eg) large meshes with no base like a 3D scan of a mountain, etc
 	if (Options.BaseCappingPolicy != EBaseCappingPolicy::NoBaseCapping)
 	{
+		TRACE_BOOKMARK(TEXT("ApproximateActors-Capping"));
+
 		TRACE_CPUPROFILER_EVENT_SCOPE(ApproximateActorsImpl_Generate_Capping);
 
 		double UseThickness = 0.0;
@@ -1023,7 +1053,7 @@ void FApproximateActorsImpl::GenerateApproximationForActorSet(const TArray<AActo
 		}
 
 		double UseHeight = (Options.BaseHeightOverrideMeters != 0) ? (Options.BaseHeightOverrideMeters * 100.0) : (2.0 * ApproxAccuracy);
-		Scene.GenerateBaseClosingMesh(UseHeight, UseThickness);
+		Scene->GenerateBaseClosingMesh(UseHeight, UseThickness);
 	}
 
 	FDynamicMesh3 DebugMesh;
@@ -1032,18 +1062,27 @@ void FApproximateActorsImpl::GenerateApproximationForActorSet(const TArray<AActo
 	{
 		TRACE_CPUPROFILER_EVENT_SCOPE(ApproximateActorsImpl_Generate_DebugMesh);
 		DebugMesh.EnableAttributes();
-		Scene.GetAccumulatedMesh(DebugMesh);
+		Scene->GetAccumulatedMesh(DebugMesh);
 		FMeshNormals::InitializeMeshToPerTriangleNormals(&DebugMesh);
 		WriteDebugMesh = &DebugMesh;
 	}
 
-	// build spatial evaluation cache
-	Scene.BuildSpatialEvaluationCache();
+	// build spatial evaluation cache in the FMeshSceneAdapter, necessary for the mesh build below
+	TRACE_BOOKMARK(TEXT("ApproximateActors-Building Cache"));
+	Scene->BuildSpatialEvaluationCache();
+
+	//
+	// Generate a new mesh that approximates the entire scene represented by FMeshSceneAdapter.
+	//
+	TRACE_BOOKMARK(TEXT("ApproximateActors-Building Approx Mesh"));
+
+	// Pass ownership of the Scene to GenerateApproximationMesh() so that it can delete it as soon
+	// as possible (because it is often very large, memory-wise)
+	TSharedPtr<FApproximationMeshData> ApproximationMeshData = GenerateApproximationMesh( MoveTemp(Scene), Options, ApproxAccuracy);
 
-	// if we are only generating collision mesh, we are going to exit after mesh generation
+	// if we are only generating collision mesh, we are done now
 	if (Options.BasePolicy == IGeometryProcessing_ApproximateActors::EApproximationPolicy::CollisionMesh)
 	{
-		TSharedPtr<FApproximationMeshData> ApproximationMeshData = GenerateApproximationMesh(Scene, Options, ApproxAccuracy);
 		ResultsOut.ResultCode = ApproximationMeshData->ResultCode;
 		if (ResultsOut.ResultCode == EResultCode::Success)
 		{
@@ -1052,21 +1091,7 @@ void FApproximateActorsImpl::GenerateApproximationForActorSet(const TArray<AActo
 		return;
 	}
 
-	// launch async mesh compute which can run while we do (relatively) expensive render captures
-	TFuture<TSharedPtr<FApproximationMeshData>> MeshComputeFuture = Async(EAsyncExecution::Thread,
-		[&Scene, &Options, &ApproxAccuracy]() {
-			return GenerateApproximationMesh(Scene, Options, ApproxAccuracy);
-		});
-
-	Progress.EnterProgressFrame(1.f, LOCTEXT("CapturingScene", "Capturing Scene..."));
 
-	TUniquePtr<FSceneCapturePhotoSet> SceneCapture = CapturePhotoSet(Actors, Options);
-
-	Progress.EnterProgressFrame(1.f, LOCTEXT("BakingTextures", "Baking Textures..."));
-
-	// need to wait for mesh to finish computing
-	MeshComputeFuture.Wait();
-	TSharedPtr<FApproximationMeshData> ApproximationMeshData = MeshComputeFuture.Get();
 	if (ApproximationMeshData->ResultCode != EResultCode::Success)
 	{
 		ResultsOut.ResultCode = ApproximationMeshData->ResultCode;
@@ -1075,6 +1100,23 @@ void FApproximateActorsImpl::GenerateApproximationForActorSet(const TArray<AActo
 	FDynamicMesh3 FinalMesh = MoveTemp(ApproximationMeshData->Mesh);
 	FMeshTangentsd FinalMeshTangents = MoveTemp(ApproximationMeshData->Tangents);
 
+
+	//
+	// create a set of spatially located render captures of the scene ("photo set").
+	//
+	Progress.EnterProgressFrame(1.f, LOCTEXT("CapturingScene", "Capturing Scene..."));
+	TRACE_BOOKMARK(TEXT("ApproximateActors-Capture Photos"));
+
+	TUniquePtr<FSceneCapturePhotoSet> SceneCapture = CapturePhotoSet(Actors, Options);
+
+
+	//
+	// bake textures onto the generated approximation mesh by projecting/sampling
+	// the set of captured photos
+	//
+	Progress.EnterProgressFrame(1.f, LOCTEXT("BakingTextures", "Baking Textures..."));
+	TRACE_BOOKMARK(TEXT("ApproximateActors-Bake Textures"));
+
 	FOptions OverridenOptions = Options;
 
 	// evaluate required texture size if needed
@@ -1093,7 +1135,8 @@ void FApproximateActorsImpl::GenerateApproximationForActorSet(const TArray<AActo
 			OverridenOptions.TextureImageSize = BestTextureSize;
 		}
 	}
-	
+
+
 	// bake textures for Actor
 	FGeneratedResultTextures GeneratedTextures;
 	BakeTexturesFromPhotoCapture(SceneCapture, OverridenOptions,
@@ -1101,6 +1144,7 @@ void FApproximateActorsImpl::GenerateApproximationForActorSet(const TArray<AActo
 		&FinalMesh, &FinalMeshTangents);
 
 	Progress.EnterProgressFrame(1.f, LOCTEXT("Writing Assets", "Writing Assets..."));
+	TRACE_BOOKMARK(TEXT("ApproximateActors-Create Material"));
 
 	// Make material for textures by creating MIC of input material, or fall back to known material
 	UMaterialInterface* UseBaseMaterial = (Options.BakeMaterial != nullptr) ? 
@@ -1160,6 +1204,7 @@ void FApproximateActorsImpl::GenerateApproximationForActorSet(const TArray<AActo
 		}
 	};
 
+	TRACE_BOOKMARK(TEXT("ApproximateActors-Write Textures"));
 
 	// process the generated textures
 	if (Options.bBakeBaseColor && GeneratedTextures.BaseColorMap)
