MeshModelingTools: Fix duplicate Loctext entries in PolyExt and UVSelectTool and make Loctext keys in PolyExt clearer for localizers.

#rnx

diff --git a/Engine/Plugins/Experimental/MeshModelingToolsetExp/Source/MeshModelingToolsExp/Private/DrawPolyPathTool.cpp b/Engine/Plugins/Experimental/MeshModelingToolsetExp/Source/MeshModelingToolsExp/Private/DrawPolyPathTool.cpp
--- a/Engine/Plugins/Experimental/MeshModelingToolsetExp/Source/MeshModelingToolsExp/Private/DrawPolyPathTool.cpp
+++ b/Engine/Plugins/Experimental/MeshModelingToolsetExp/Source/MeshModelingToolsExp/Private/DrawPolyPathTool.cpp
@@ -306,7 +306,7 @@ void UDrawPolyPathTool::OnPropertyModified(UObject* PropertySet, FProperty* Prop
 		{
 			GetToolManager()->EmitObjectChange(this,
 				MakeUnique<FDrawPolyPathStateChange>(CurrentCurveTimestamp, EState::SettingRadius),
-				LOCTEXT("DrawPolyPathCancelRadius", "Cancel setting corner radius"));
+				LOCTEXT("CancelRadiusTransactionName", "Cancel Setting Corner Radius"));
 			OnCompleteRadius();
 		}
 		else if(State != EState::DrawingPath)
@@ -384,14 +384,14 @@ void UDrawPolyPathTool::OnClicked(const FInputDeviceRay& ClickPos)
 					bPathIsClosed = SurfacePathMechanic->LoopWasClosed();
 					GetToolManager()->EmitObjectChange(this, 
 						MakeUnique<FDrawPolyPathStateChange>(CurrentCurveTimestamp, EState::DrawingPath), 
-						LOCTEXT("DrawPolyPathFinishPath", "Finish path"));
+						LOCTEXT("FinishPathTransactionName", "Finish Path"));
 					OnCompleteSurfacePath();
 				}
 				else
 				{
 					GetToolManager()->EmitObjectChange(this, 
 						MakeUnique<FDrawPolyPathStateChange>(CurrentCurveTimestamp, EState::DrawingPath),
-						LOCTEXT("DrawPolyPathAddToPath", "Add point to path"));
+						LOCTEXT("AddToPathTransactionName", "Add Point to Path"));
 				}
 			}
 			break;
@@ -401,7 +401,7 @@ void UDrawPolyPathTool::OnClicked(const FInputDeviceRay& ClickPos)
 				// This can happen accidentally when the user has snapping turned on and it snaps to 0. 
 				// We'll ignore that click and show an error message.
 				GetToolManager()->DisplayMessage(
-					LOCTEXT("ZeroWidthPath", "Cannot set path width to 0."),
+					LOCTEXT("ZeroWidthPathError", "Cannot set path width to 0."),
 					EToolMessageLevel::UserError);
 			}
 			else
@@ -410,14 +410,14 @@ void UDrawPolyPathTool::OnClicked(const FInputDeviceRay& ClickPos)
 
 				GetToolManager()->EmitObjectChange(this,
 					MakeUnique<FDrawPolyPathStateChange>(CurrentCurveTimestamp, EState::SettingWidth),
-					LOCTEXT("DrawPolyPathBeginWidth", "Set path width"));
+					LOCTEXT("BeginWidthTransactionName", "Set Path Width"));
 				OnCompleteWidth();
 			}
 			break;
 		case EState::SettingRadius:
 			GetToolManager()->EmitObjectChange(this, 
 				MakeUnique<FDrawPolyPathStateChange>(CurrentCurveTimestamp, EState::SettingRadius), 
-				LOCTEXT("DrawPolyPathBeginRadius", "Set corner radius"));
+				LOCTEXT("BeginRadiusTransactionName", "Set Corner Radius"));
 			OnCompleteRadius();
 			break;
 		case EState::SettingHeight:
@@ -697,13 +697,13 @@ void UDrawPolyPathTool::BeginSettingWidth()
 	if (TransformProps->WidthMode == EDrawPolyPathWidthMode::Interactive)
 	{
 		GetToolManager()->DisplayMessage(
-			LOCTEXT("OnStartOffset", "Set the width of the path by clicking on the drawing plane. Hold Shift to ignore snapping."),
+			LOCTEXT("InteractiveSetWidthInstructions", "Set the width of the path by clicking on the drawing plane. Hold Shift to ignore snapping."),
 			EToolMessageLevel::UserNotification);
 	}
 	else
 	{
 		GetToolManager()->DisplayMessage(
-			LOCTEXT("OnStartOffset", "Click in viewport to accept fixed path width, or change it in details panel."),
+			LOCTEXT("FixedSetWidthInstructions", "Click in viewport to accept fixed path width, or change it in details panel."),
 			EToolMessageLevel::UserNotification);
 	}
 }
@@ -751,13 +751,13 @@ void UDrawPolyPathTool::BeginSettingRadius()
 	if (TransformProps->RadiusMode == EDrawPolyPathRadiusMode::Interactive)
 	{
 		GetToolManager()->DisplayMessage(
-			LOCTEXT("OnStartOffset", "Set the radius of the corners by clicking on the drawing plane. Hold Shift to ignore snapping."),
+			LOCTEXT("InteractiveSetRadiusInstructions", "Set the radius of the corners by clicking on the drawing plane. Hold Shift to ignore snapping."),
 			EToolMessageLevel::UserNotification);
 	}
 	else
 	{
 		GetToolManager()->DisplayMessage(
-			LOCTEXT("OnStartOffset", "Click in viewport to accept fixed corner radius, or change it in details panel."),
+			LOCTEXT("FixedSetRadiusInstructions", "Click in viewport to accept fixed corner radius, or change it in details panel."),
 			EToolMessageLevel::UserNotification);
 	}
 }
@@ -916,7 +916,7 @@ void UDrawPolyPathTool::BeginConstantExtrudeHeight()
 	ExtrudeHeightMechanic = nullptr;
 
 	GetToolManager()->DisplayMessage(
-		LOCTEXT("OnStartOffset", "Click in viewport to accept fixed path height, or change it in details panel."),
+		LOCTEXT("FixedSetHeightInstructions", "Click in viewport to accept fixed path height, or change it in details panel."),
 		EToolMessageLevel::UserNotification);
 }
 
@@ -1013,7 +1013,7 @@ void UDrawPolyPathTool::EmitNewObject()
 	MeshTransform.Origin = MeshTransform.ToPlane(Center, 2);
 	MeshTransforms::WorldToFrameCoords(PathMesh, MeshTransform);
 
-	GetToolManager()->BeginUndoTransaction(LOCTEXT("CreatePolyPath", "Create PolyPath"));
+	GetToolManager()->BeginUndoTransaction(LOCTEXT("CreatePolyPathTransactionName", "Create PolyPath"));
 
 	FCreateMeshObjectParams NewMeshObjectParams;
 	NewMeshObjectParams.TargetWorld = TargetWorld;
@@ -1055,7 +1055,7 @@ void UDrawPolyPathTool::EmitNewObject()
 void UDrawPolyPathTool::ShowStartupMessage()
 {
 	GetToolManager()->DisplayMessage(
-		LOCTEXT("OnStartDraw", "Draw a path on the drawing plane, set its width, and extrude it. Left-click to place path vertices, and click on the last or first vertex to complete the path. Hold Shift to ignore snapping while drawing."),
+		LOCTEXT("StartDrawInstructions", "Draw a path on the drawing plane, set its width, and extrude it. Left-click to place path vertices, and click on the last or first vertex to complete the path. Hold Shift to ignore snapping while drawing."),
 		EToolMessageLevel::UserNotification);
 }
 
@@ -1063,7 +1063,7 @@ void UDrawPolyPathTool::ShowStartupMessage()
 void UDrawPolyPathTool::ShowExtrudeMessage()
 {
 	GetToolManager()->DisplayMessage(
-		LOCTEXT("OnStartExtrude", "Set the height of the extrusion by positioning the mouse over the extrusion volume, or over objects to snap to their heights. Hold Shift to ignore snapping."),
+		LOCTEXT("InteractiveSetHeightInstructions", "Set the height of the extrusion by positioning the mouse over the extrusion volume, or over objects to snap to their heights. Hold Shift to ignore snapping."),
 		EToolMessageLevel::UserNotification);
 }
 
diff --git a/Engine/Plugins/Experimental/UVEditor/Source/UVEditorTools/Private/UVSelectTool.cpp b/Engine/Plugins/Experimental/UVEditor/Source/UVEditorTools/Private/UVSelectTool.cpp
--- a/Engine/Plugins/Experimental/UVEditor/Source/UVEditorTools/Private/UVSelectTool.cpp
+++ b/Engine/Plugins/Experimental/UVEditor/Source/UVEditorTools/Private/UVSelectTool.cpp
@@ -1031,7 +1031,7 @@ void UUVSelectTool::ApplySplit()
 	else
 	{
 		GetToolManager()->DisplayMessage(
-			LOCTEXT("SplitErrorSelectionEmpty", "Cannot split UVs. Selection must be edges or vertices."),
+			LOCTEXT("SplitErrorNotEdgeOrVert", "Cannot split UVs. Selection must be edges or vertices."),
 			EToolMessageLevel::UserWarning);
 	}
 }
