Mesh Selection Tool: add a Smooth Boundary operation

diff --git a/Engine/Plugins/Experimental/MeshModelingToolsetExp/Source/MeshModelingToolsExp/Private/MeshSelectionTool.cpp b/Engine/Plugins/Experimental/MeshModelingToolsetExp/Source/MeshModelingToolsExp/Private/MeshSelectionTool.cpp
--- a/Engine/Plugins/Experimental/MeshModelingToolsetExp/Source/MeshModelingToolsExp/Private/MeshSelectionTool.cpp
+++ b/Engine/Plugins/Experimental/MeshModelingToolsetExp/Source/MeshModelingToolsExp/Private/MeshSelectionTool.cpp
@@ -306,6 +306,14 @@ void UMeshSelectionTool::RegisterActions(FInteractiveToolActionSet& ActionSet)
 		LOCTEXT("OptimizeSelectionTooltip", "Optimize selection"),
 		EModifierKey::None, EKeys::O,
 		[this]() { OptimizeSelection(); });
+
+	ActionSet.RegisterAction(this, (int32)EMeshSelectionToolActions::SmoothBoundary,
+		TEXT("SmoothBoundary"),
+		LOCTEXT("SmoothBoundary", "Smooth Boundary"),
+		LOCTEXT("SmoothBoundaryTooltip", "Smooth Boundary"),
+		EModifierKey::None, EKeys::B,
+		[this]() { SmoothSelectionBoundary(); });
+
 }
 
 
@@ -1018,6 +1026,10 @@ void UMeshSelectionTool::ApplyAction(EMeshSelectionToolActions ActionType)
 			DisconnectSelectedTriangles();
 			break;
 
+		case EMeshSelectionToolActions::SmoothBoundary:
+			SmoothSelectionBoundary();
+			break;
+
 		case EMeshSelectionToolActions::SeparateSelected:
 			SeparateSelectedTriangles(true);
 			break;
@@ -1607,7 +1619,68 @@ void UMeshSelectionTool::AssignNewGroupToSelectedTriangles()
 	bHaveModifiedMesh = true;
 }
 
+void UMeshSelectionTool::SmoothSelectionBoundary()
+{
+	check(SelectionType == EMeshSelectionElementType::Face);
+	TArray<int32> SelectedFaces = Selection->GetElements(EMeshSelectionElementType::Face);
+	if (SelectedFaces.Num() == 0)
+	{
+		return;
+	}
+
+	FMeshRegionBoundaryLoops BoundaryLoops(PreviewMesh->GetMesh(), SelectedFaces, true);
+	if (BoundaryLoops.Loops.Num() == 0)
+	{
+		return;
+	}
+
+	TUniquePtr<FToolCommandChangeSequence> ChangeSeq = MakeUnique<FToolCommandChangeSequence>();
+
+	TUniquePtr<FMeshChange> MeshChange = PreviewMesh->TrackedEditMesh(
+		[&BoundaryLoops, this](FDynamicMesh3& Mesh, FDynamicMeshChangeTracker& ChangeTracker)
+	{
+		constexpr double Alpha = 0.75;
+
+		TMap<int, FVector3d> NewLoopPositions;
+
+		for (const FEdgeLoop& Loop : BoundaryLoops.Loops)
+		{
+			int NumLoopVertices = Loop.GetVertexCount();
+
+			for (int LoopVertexIndex = 0; LoopVertexIndex < NumLoopVertices; ++LoopVertexIndex)
+			{
+				const FVector3d PrevPoint = Loop.GetPrevVertex(LoopVertexIndex);
+				const FVector3d NextPoint = Loop.GetNextVertex(LoopVertexIndex);
+				const FVector3d Avg = 0.5 * (PrevPoint + NextPoint);
+
+				const FVector3d CurrPoint = Loop.GetVertex(LoopVertexIndex);
+
+				const FVector3d NewCurrPoint = (1.0 - Alpha) * CurrPoint + Alpha * Avg;
+				const int VertexIndex = Loop.Vertices[LoopVertexIndex];
+
+				// TODO: Reproject to original surface?
+				// TODO: Fix the UVs
+
+				NewLoopPositions.Add(VertexIndex, NewCurrPoint);
+			}
+		}
+
+		PreviewMesh->EditMesh([NewLoopPositions](FDynamicMesh3& Mesh)
+		{
+			for (const TPair<int, FVector3d>& NewVert : NewLoopPositions)
+			{
+				Mesh.SetVertex(NewVert.Key, NewVert.Value);
+			}
+		});
+	});
 
+	ChangeSeq->AppendChange(PreviewMesh, MoveTemp(MeshChange));
+
+	GetToolManager()->EmitObjectChange(this, MoveTemp(ChangeSeq), LOCTEXT("MeshSelectionToolSmoothBoundary", "Smooth Selection Boundary"));
+
+	OnExternalSelectionChange();
+	bHaveModifiedMesh = true;
+}
 
 
 
diff --git a/Engine/Plugins/Experimental/MeshModelingToolsetExp/Source/MeshModelingToolsExp/Public/MeshSelectionTool.h b/Engine/Plugins/Experimental/MeshModelingToolsetExp/Source/MeshModelingToolsExp/Public/MeshSelectionTool.h
--- a/Engine/Plugins/Experimental/MeshModelingToolsetExp/Source/MeshModelingToolsExp/Public/MeshSelectionTool.h
+++ b/Engine/Plugins/Experimental/MeshModelingToolsetExp/Source/MeshModelingToolsExp/Public/MeshSelectionTool.h
@@ -55,6 +55,7 @@ enum class EMeshSelectionToolActions
 	DuplicateSelected,
 	FlipSelected,
 	CreateGroup,
+	SmoothBoundary,
 
 	CycleSelectionMode,
 	CycleViewMode
@@ -146,7 +147,6 @@ public:
 	{
 		PostAction(EMeshSelectionToolActions::OptimizeSelection);
 	}
-
 };
 
 
@@ -200,6 +200,13 @@ public:
 		PostAction(EMeshSelectionToolActions::DuplicateSelected);
 	}
 
+	/** Smooth the selection border */
+	UFUNCTION(CallInEditor, Category = MeshEdits, meta = (DisplayPriority = 12))
+	void SmoothBorder()
+	{
+		PostAction(EMeshSelectionToolActions::SmoothBoundary);
+	}
+
 };
 
 
@@ -430,6 +437,8 @@ protected:
 	void FlipSelectedTriangles();
 	void AssignNewGroupToSelectedTriangles();
 
+	void SmoothSelectionBoundary();
+
 	TSharedPtr<UE::Geometry::FPolygroupSet, ESPMode::ThreadSafe> ActiveGroupSet;
 	void OnSelectedGroupLayerChanged();
 	void UpdateActiveGroupLayer();
