ModelingTools: Fixed race condition when setting override render material on preview mesh in BakeTexture & BakeAll.

#rnx

diff --git a/Engine/Plugins/Experimental/MeshModelingToolsetExp/Source/MeshModelingToolsExp/Private/BakeMeshAttributeMapsToolBase.cpp b/Engine/Plugins/Experimental/MeshModelingToolsetExp/Source/MeshModelingToolsExp/Private/BakeMeshAttributeMapsToolBase.cpp
--- a/Engine/Plugins/Experimental/MeshModelingToolsetExp/Source/MeshModelingToolsExp/Private/BakeMeshAttributeMapsToolBase.cpp
+++ b/Engine/Plugins/Experimental/MeshModelingToolsetExp/Source/MeshModelingToolsExp/Private/BakeMeshAttributeMapsToolBase.cpp
@@ -75,11 +75,6 @@ void UBakeMeshAttributeMapsToolBase::PostSetup()
 	AddToolPropertySource(VisualizationProps);
 	VisualizationProps->WatchProperty(VisualizationProps->bPreviewAsMaterial, [this](bool) { UpdateVisualization(); GetToolManager()->PostInvalidation(); });
 
-	// Initialize background compute
-	Compute = MakeUnique<TGenericDataBackgroundCompute<FMeshMapBaker>>();
-	Compute->Setup(this);
-	Compute->OnResultUpdated.AddLambda([this](const TUniquePtr<FMeshMapBaker>& NewResult) { OnMapsUpdated(NewResult); });
-
 	// Initialize UV charts
 	// TODO: Compute UV charts asynchronously
 	TargetMeshUVCharts = MakeShared<TArray<int32>, ESPMode::ThreadSafe>();
@@ -153,6 +148,13 @@ void UBakeMeshAttributeMapsToolBase::UpdateVisualization()
 
 void UBakeMeshAttributeMapsToolBase::InvalidateCompute()
 {
+	if (!Compute)
+	{
+		// Initialize background compute
+		Compute = MakeUnique<TGenericDataBackgroundCompute<FMeshMapBaker>>();
+		Compute->Setup(this);
+		Compute->OnResultUpdated.AddLambda([this](const TUniquePtr<FMeshMapBaker>& NewResult) { OnMapsUpdated(NewResult); });
+	}
 	Compute->InvalidateResult();
 	OpState = EBakeOpState::Clean;
 }
diff --git a/Engine/Plugins/Runtime/MeshModelingToolset/Source/ModelingComponents/Private/PreviewMesh.cpp b/Engine/Plugins/Runtime/MeshModelingToolset/Source/ModelingComponents/Private/PreviewMesh.cpp
--- a/Engine/Plugins/Runtime/MeshModelingToolset/Source/ModelingComponents/Private/PreviewMesh.cpp
+++ b/Engine/Plugins/Runtime/MeshModelingToolset/Source/ModelingComponents/Private/PreviewMesh.cpp
@@ -37,6 +37,14 @@ void UPreviewMesh::CreateInWorld(UWorld* World, const FTransform& WithTransform)
 	TemporaryParentActor = World->SpawnActor<APreviewMeshActor>(FVector::ZeroVector, Rotation, SpawnInfo);
 
 	DynamicMeshComponent = NewObject<UDynamicMeshComponent>(TemporaryParentActor);
+
+	// Disable VerifyUsedMaterials on the DynamicMeshSceneProxy. Material verification is prone
+	// to data races when materials are subject to change at a high frequency. Since the
+	// preview mesh material usage (override render materials) is particularly prone to these
+	// races and we are certain the component materials are updated appropriately, we disable
+	// used material verification.
+	DynamicMeshComponent->SetSceneProxyVerifyUsedMaterials(false);
+	
 	TemporaryParentActor->SetRootComponent(DynamicMeshComponent);
 	//DynamicMeshComponent->SetupAttachment(TemporaryParentActor->GetRootComponent());
 	DynamicMeshComponent->RegisterComponent();
diff --git a/Engine/Source/Runtime/GeometryFramework/Private/Components/DynamicMeshComponent.cpp b/Engine/Source/Runtime/GeometryFramework/Private/Components/DynamicMeshComponent.cpp
--- a/Engine/Source/Runtime/GeometryFramework/Private/Components/DynamicMeshComponent.cpp
+++ b/Engine/Source/Runtime/GeometryFramework/Private/Components/DynamicMeshComponent.cpp
@@ -839,6 +839,8 @@ FPrimitiveSceneProxy* UDynamicMeshComponent::CreateSceneProxy()
 		{
 			NewProxy->Initialize();
 		}
+
+		NewProxy->SetVerifyUsedMaterials(bProxyVerifyUsedMaterials);
 	}
 
 	bProxyValid = true;
@@ -1249,3 +1251,13 @@ void UDynamicMeshComponent::SetDeferredCollisionUpdatesEnabled(bool bEnabled, bo
 		}
 	}
 }
+
+void UDynamicMeshComponent::SetSceneProxyVerifyUsedMaterials(bool bState)
+{
+	bProxyVerifyUsedMaterials = bState;
+	if (FDynamicMeshSceneProxy* Proxy = GetCurrentSceneProxy())
+	{
+		Proxy->SetVerifyUsedMaterials(bState);
+	}
+}
+
diff --git a/Engine/Source/Runtime/GeometryFramework/Public/Components/BaseDynamicMeshSceneProxy.h b/Engine/Source/Runtime/GeometryFramework/Public/Components/BaseDynamicMeshSceneProxy.h
--- a/Engine/Source/Runtime/GeometryFramework/Public/Components/BaseDynamicMeshSceneProxy.h
+++ b/Engine/Source/Runtime/GeometryFramework/Public/Components/BaseDynamicMeshSceneProxy.h
@@ -906,6 +906,14 @@ public:
 		return (Material != nullptr) ? Material : UMaterial::GetDefaultMaterial(MD_Surface);
 	}
 
+	/**
+	 * Set whether or not to validate mesh batch materials against the component materials.
+	 */
+	void SetVerifyUsedMaterials(const bool bState)
+	{
+		bVerifyUsedMaterials = bState;
+	}
+
 
 	/**
 	 * This needs to be called if the set of active materials changes, otherwise
diff --git a/Engine/Source/Runtime/GeometryFramework/Public/Components/DynamicMeshComponent.h b/Engine/Source/Runtime/GeometryFramework/Public/Components/DynamicMeshComponent.h
--- a/Engine/Source/Runtime/GeometryFramework/Public/Components/DynamicMeshComponent.h
+++ b/Engine/Source/Runtime/GeometryFramework/Public/Components/DynamicMeshComponent.h
@@ -546,6 +546,15 @@ protected:
 	 */
 	bool bProxyValid = false;
 
+	/**
+	 * If true, the render proxy will verify that the mesh batch materials match the contents from
+	 * the component GetUsedMaterials(). Used material verification is prone to races when changing
+	 * materials on this component in quick succession (for example, SetOverrideRenderMaterial). This
+	 * parameter is provided to allow clients to opt out of used material verification for these
+	 * use cases.
+	 */
+	bool bProxyVerifyUsedMaterials = true;
+
 	virtual FBaseDynamicMeshSceneProxy* GetBaseSceneProxy() override { return (FBaseDynamicMeshSceneProxy*)GetCurrentSceneProxy(); }
 
 	/** 
@@ -573,6 +582,9 @@ protected:
 	void PostEditChangeProperty(FPropertyChangedEvent& PropertyChangedEvent) override;
 #endif
 
+public:
+	/** Set whether or not to validate mesh batch materials against the component materials. */
+	void SetSceneProxyVerifyUsedMaterials(bool bState);
 
 
 };
