ModelingTools: Fix crash on exit when BakeAll is active by clearing FMeshSceneAdapter on tool shutdown.

#rnx

diff --git a/Engine/Plugins/Experimental/MeshModelingToolsetExp/Source/MeshModelingToolsExp/Private/BakeMultiMeshAttributeMapsTool.cpp b/Engine/Plugins/Experimental/MeshModelingToolsetExp/Source/MeshModelingToolsExp/Private/BakeMultiMeshAttributeMapsTool.cpp
--- a/Engine/Plugins/Experimental/MeshModelingToolsetExp/Source/MeshModelingToolsExp/Private/BakeMultiMeshAttributeMapsTool.cpp
+++ b/Engine/Plugins/Experimental/MeshModelingToolsetExp/Source/MeshModelingToolsExp/Private/BakeMultiMeshAttributeMapsTool.cpp
@@ -325,7 +325,7 @@ class FMultiMeshMapBakerOp : public TGenericDataOperator<FMeshMapBaker>
 {
 public:
 	// General bake settings
-	FMeshSceneAdapter* DetailMeshScene = nullptr;
+	TSharedPtr<FMeshSceneAdapter, ESPMode::ThreadSafe> DetailMeshScene;
 	UE::Geometry::FDynamicMesh3* BaseMesh = nullptr;
 	TSharedPtr<UE::Geometry::FMeshTangentsd, ESPMode::ThreadSafe> BaseMeshTangents;
 	TUniquePtr<UE::Geometry::FMeshMapBaker> Baker;
@@ -351,7 +351,7 @@ public:
 		Baker->SetTargetMeshTangents(BaseMeshTangents);
 		Baker->SetTargetMeshUVCharts(BaseMeshUVCharts.Get());
 		
-		FMeshBakerMeshSceneSampler DetailSampler(DetailMeshScene);
+		FMeshBakerMeshSceneSampler DetailSampler(DetailMeshScene.Get());
 		Baker->SetDetailSampler(&DetailSampler);
 
 		for (const EBakeMapType MapType : ENUM_EBAKEMAPTYPE_ALL)
@@ -415,9 +415,10 @@ void UBakeMultiMeshAttributeMapsTool::Setup()
 			DetailComponents.Add(Component);
 		}
 	}
-	DetailMeshScene.AddComponents(DetailComponents);
-	DetailMeshScene.Build(FMeshSceneAdapterBuildOptions());
-	DetailMeshScene.BuildSpatialEvaluationCache();
+	DetailMeshScene = MakeShared<FMeshSceneAdapter, ESPMode::ThreadSafe>();
+	DetailMeshScene->AddComponents(DetailComponents);
+	DetailMeshScene->Build(FMeshSceneAdapterBuildOptions());
+	DetailMeshScene->BuildSpatialEvaluationCache();
 
 	UToolTarget* Target = Targets[0];
 
@@ -514,7 +515,7 @@ bool UBakeMultiMeshAttributeMapsTool::CanAccept() const
 TUniquePtr<UE::Geometry::TGenericDataOperator<FMeshMapBaker>> UBakeMultiMeshAttributeMapsTool::MakeNewOperator()
 {
 	TUniquePtr<FMultiMeshMapBakerOp> Op = MakeUnique<FMultiMeshMapBakerOp>();
-	Op->DetailMeshScene = &DetailMeshScene;
+	Op->DetailMeshScene = DetailMeshScene;
 	Op->BaseMesh = &TargetMesh;
 	Op->BaseMeshUVCharts = TargetMeshUVCharts;
 	Op->BakeSettings = CachedBakeSettings;
@@ -547,6 +548,7 @@ void UBakeMultiMeshAttributeMapsTool::OnShutdown(EToolShutdownType ShutdownType)
 	{
 		Compute->Shutdown();
 	}
+	DetailMeshScene = nullptr;
 
 	// Restore visibility of detail targets
 	const int NumTargets = Targets.Num();
@@ -672,7 +674,7 @@ EBakeOpState UBakeMultiMeshAttributeMapsTool::UpdateResult_DetailMeshes()
 			}
 		}
 	};
-	DetailMeshScene.ProcessActorChildMeshes(BuildMeshToDataMaps);
+	DetailMeshScene->ProcessActorChildMeshes(BuildMeshToDataMaps);
 
 	// This method will always force a re-evaluation.
 	return EBakeOpState::Evaluate;
@@ -723,7 +725,7 @@ void UBakeMultiMeshAttributeMapsTool::GatherAnalytics(FBakeAnalytics::FMeshSetti
 		Data.NumTargetMeshTris = TargetMesh.TriangleCount();
 		Data.NumDetailMesh = 0;
 		Data.NumDetailMeshTris = 0;
-		DetailMeshScene.ProcessActorChildMeshes([&Data](const FActorAdapter* ActorAdapter, const FActorChildMesh* ChildMesh)
+		DetailMeshScene->ProcessActorChildMeshes([&Data](const FActorAdapter* ActorAdapter, const FActorChildMesh* ChildMesh)
 		{
 			if (ChildMesh)
 			{
diff --git a/Engine/Plugins/Experimental/MeshModelingToolsetExp/Source/MeshModelingToolsExp/Public/BakeMultiMeshAttributeMapsTool.h b/Engine/Plugins/Experimental/MeshModelingToolsetExp/Source/MeshModelingToolsExp/Public/BakeMultiMeshAttributeMapsTool.h
--- a/Engine/Plugins/Experimental/MeshModelingToolsetExp/Source/MeshModelingToolsExp/Public/BakeMultiMeshAttributeMapsTool.h
+++ b/Engine/Plugins/Experimental/MeshModelingToolsetExp/Source/MeshModelingToolsExp/Public/BakeMultiMeshAttributeMapsTool.h
@@ -196,7 +196,7 @@ protected:
 protected:
 	friend class FMultiMeshMapBakerOp;
 
-	UE::Geometry::FMeshSceneAdapter DetailMeshScene;
+	TSharedPtr<UE::Geometry::FMeshSceneAdapter, ESPMode::ThreadSafe> DetailMeshScene;
 
 	void UpdateOnModeChange();
 
