GeometryCore: add min-triangle-count threshold to TMeshSimplification::FastCollapsePass()

ApproximateActors: add FMeshApproximationSettings::bEnableSimplifyPrePass option, default true, and implement in ApproximateActors impl. This does a fast-collapse down to a multiple of the target triangle count (or 1m tris in tolerance-mode). This reduces simplification time in HLOD builds for small city from 4-5m to 1m in some of the worst blocks.

ApproximateActors: add FMeshApproximationSettings::bEnableParallelBaking, default true, and implement in the ApproximateActors Impl. This allows the photocapture step to proceed at the same time as the mesh generation, which shaves 30-40s off every HLOD build, but increases memory requirements.

#rnx

diff --git a/Engine/Plugins/Experimental/MeshModelingToolsetExp/Source/GeometryProcessingAdapters/Private/ApproximateActorsImpl.cpp b/Engine/Plugins/Experimental/MeshModelingToolsetExp/Source/GeometryProcessingAdapters/Private/ApproximateActorsImpl.cpp
--- a/Engine/Plugins/Experimental/MeshModelingToolsetExp/Source/GeometryProcessingAdapters/Private/ApproximateActorsImpl.cpp
+++ b/Engine/Plugins/Experimental/MeshModelingToolsetExp/Source/GeometryProcessingAdapters/Private/ApproximateActorsImpl.cpp
@@ -510,20 +510,31 @@ static TSharedPtr<FApproximationMeshData> GenerateApproximationMesh(
 			FVector2d VolArea = TMeshQueries<FDynamicMesh3>::GetVolumeArea(*CurResultMesh);
 			double MeshAreaMeterSqr = VolArea.Y * 0.0001;
 			int32 AreaBaseTargetTriCount = MeshAreaMeterSqr * Options.SimplificationTargetMetric;
-			Simplifier.SimplifyToTriangleCount(AreaBaseTargetTriCount);
+
+			// do initial fast-collapse pass if enabled
+			if (Options.bEnableFastSimplifyPrePass && CurResultMesh->TriangleCount() > 10 * AreaBaseTargetTriCount)
+			{
+				TRACE_CPUPROFILER_EVENT_SCOPE(ApproximateActorsImpl_Generate_Simplification_PrePass);
+				Simplifier.FastCollapsePass(ApproxAccuracy, 10, false, 5 * AreaBaseTargetTriCount);
+			}
+
+			{
+				TRACE_CPUPROFILER_EVENT_SCOPE(ApproximateActorsImpl_Generate_Simplification_Pass1);
+				Simplifier.SimplifyToTriangleCount(AreaBaseTargetTriCount);
+			}
 		}
 		else if (Options.MeshSimplificationPolicy == IGeometryProcessing_ApproximateActors::ESimplificationPolicy::GeometricTolerance)
 		{
 			double UseTargetTolerance = Options.SimplificationTargetMetric * 100.0;		// convert to cm (UE Units)
 
-			// first do fast collapse
-			// (this does not seem to help perf and probably makes the results slightly worse)
-			//{
-			//	TRACE_CPUPROFILER_EVENT_SCOPE(ApproximateActorsImpl_Generate_Simplification_PrePass);
-			//	Simplifier.FastCollapsePass(ApproxAccuracy, 5);
-			//}
+			// do initial fast-collapse pass if enabled
+			if (Options.bEnableFastSimplifyPrePass && CurResultMesh->TriangleCount() > 1000000)
+			{
+				TRACE_CPUPROFILER_EVENT_SCOPE(ApproximateActorsImpl_Generate_Simplification_PrePass);
+				Simplifier.FastCollapsePass(ApproxAccuracy, 10, false, 1000000);
+			}
 
-			// now simplify down to a reasonable tri count, as geometric metric is (relatively) expensive
+			// simplify down to a reasonable tri count, as geometric metric is (relatively) expensive
 			// (still, this is all incredibly cheap compared to the cost of the rest of this method in practice)
 			{
 				TRACE_CPUPROFILER_EVENT_SCOPE(ApproximateActorsImpl_Generate_Simplification_Pass1);
@@ -545,6 +556,14 @@ static TSharedPtr<FApproximationMeshData> GenerateApproximationMesh(
 		}
 		else
 		{
+			// do initial fast-collapse pass if enabled
+			int32 FastCollapseThresh = FMath::Max(1000000, 2 * BaseTargeTriCount);
+			if (Options.bEnableFastSimplifyPrePass && CurResultMesh->TriangleCount() > FastCollapseThresh)
+			{
+				TRACE_CPUPROFILER_EVENT_SCOPE(ApproximateActorsImpl_Generate_Simplification_PrePass);
+				Simplifier.FastCollapsePass(ApproxAccuracy, 10, false, FastCollapseThresh);
+			}
+
 			Simplifier.SimplifyToTriangleCount(BaseTargeTriCount);
 		}
 
@@ -931,6 +950,7 @@ IGeometryProcessing_ApproximateActors::FOptions FApproximateActorsImpl::Construc
 	{
 		Options.MeshSimplificationPolicy = IGeometryProcessing_ApproximateActors::ESimplificationPolicy::FixedTriangleCount;
 	}
+	Options.bEnableFastSimplifyPrePass = UseSettings.bEnableSimplifyPrePass;
 
 	Options.UVPolicy = IGeometryProcessing_ApproximateActors::EUVGenerationPolicy::PreferXAtlas;
 	if (UseSettings.UVGenerationMethod == EMeshApproximationUVGenerationPolicy::PreferUVAtlas)
@@ -957,6 +977,8 @@ IGeometryProcessing_ApproximateActors::FOptions FApproximateActorsImpl::Construc
 	Options.FieldOfViewDegrees = UseSettings.CaptureFieldOfView;
 	Options.NearPlaneDist = UseSettings.NearPlaneDist;
 
+	Options.bMaximizeBakeParallelism = UseSettings.bEnableParallelBaking;
+
 	Options.bVerbose = UseSettings.bPrintDebugMessages;
 	Options.bWriteDebugMesh = UseSettings.bEmitFullDebugMesh;
 
@@ -1082,28 +1104,41 @@ void FApproximateActorsImpl::GenerateApproximationForActorSet(const TArray<AActo
 
 	// Pass ownership of the Scene to GenerateApproximationMesh() so that it can delete it as soon
 	// as possible (because it is often very large, memory-wise)
-	TSharedPtr<FApproximationMeshData> ApproximationMeshData = GenerateApproximationMesh( MoveTemp(Scene), Options, ApproxAccuracy);
+	TFuture<TSharedPtr<FApproximationMeshData>> GenerateMeshFuture = Async(EAsyncExecution::ThreadPool, [&Scene, Options, ApproxAccuracy]()
+	{
+		return GenerateApproximationMesh( MoveTemp(Scene), Options, ApproxAccuracy);
+	});
+	
+	FDynamicMesh3 FinalMesh;
+	FMeshTangentsd FinalMeshTangents;
+	auto WaitForMeshAvailable = [&]()
+	{
+		TSharedPtr<FApproximationMeshData> ApproximationMeshData = GenerateMeshFuture.Get();
+		ResultsOut.ResultCode = ApproximationMeshData->ResultCode;
+		FinalMesh = MoveTemp(ApproximationMeshData->Mesh);
+		FinalMeshTangents = MoveTemp(ApproximationMeshData->Tangents);
+	};
 
 	// if we are only generating collision mesh, we are done now
 	if (Options.BasePolicy == IGeometryProcessing_ApproximateActors::EApproximationPolicy::CollisionMesh)
 	{
-		ResultsOut.ResultCode = ApproximationMeshData->ResultCode;
+		WaitForMeshAvailable();
 		if (ResultsOut.ResultCode == EResultCode::Success)
 		{
-			EmitGeneratedMeshAsset(Actors, Options, ResultsOut, &ApproximationMeshData->Mesh, nullptr, WriteDebugMesh);
+			EmitGeneratedMeshAsset(Actors, Options, ResultsOut, &FinalMesh, nullptr, WriteDebugMesh);
 		}
 		return;
 	}
 
-
-	if (ApproximationMeshData->ResultCode != EResultCode::Success)
+	// if parallel capture is not allowed, force mesh computation to finish now
+	if (Options.bMaximizeBakeParallelism == false)
 	{
-		ResultsOut.ResultCode = ApproximationMeshData->ResultCode;
-		return;
+		WaitForMeshAvailable();
+		if (ResultsOut.ResultCode != EResultCode::Success)
+		{
+			return;
+		}
 	}
-	FDynamicMesh3 FinalMesh = MoveTemp(ApproximationMeshData->Mesh);
-	FMeshTangentsd FinalMeshTangents = MoveTemp(ApproximationMeshData->Tangents);
-
 
 	//
 	// create a set of spatially located render captures of the scene ("photo set").
@@ -1113,6 +1148,15 @@ void FApproximateActorsImpl::GenerateApproximationForActorSet(const TArray<AActo
 
 	TUniquePtr<FSceneCapturePhotoSet> SceneCapture = CapturePhotoSet(Actors, Options);
 
+	// if parallel capture was allowed, need to force the mesh compute to finish now to be able to proceed
+	if (Options.bMaximizeBakeParallelism == true)
+	{
+		WaitForMeshAvailable();
+		if (ResultsOut.ResultCode != EResultCode::Success)
+		{
+			return;
+		}
+	}
 
 	//
 	// bake textures onto the generated approximation mesh by projecting/sampling
diff --git a/Engine/Plugins/Runtime/GeometryProcessing/Source/DynamicMesh/Private/MeshSimplification.cpp b/Engine/Plugins/Runtime/GeometryProcessing/Source/DynamicMesh/Private/MeshSimplification.cpp
--- a/Engine/Plugins/Runtime/GeometryProcessing/Source/DynamicMesh/Private/MeshSimplification.cpp
+++ b/Engine/Plugins/Runtime/GeometryProcessing/Source/DynamicMesh/Private/MeshSimplification.cpp
@@ -1044,9 +1044,9 @@ void TMeshSimplification<QuadricErrorType>::SimplifyToMinimalPlanar(
 
 
 template <typename QuadricErrorType>
-void TMeshSimplification<QuadricErrorType>::FastCollapsePass(double fMinEdgeLength, int nRounds, bool MeshIsClosedHint)
+void TMeshSimplification<QuadricErrorType>::FastCollapsePass(double fMinEdgeLength, int nRounds, bool MeshIsClosedHint, uint32 MinTriangleCount)
 {
-	if (Mesh->TriangleCount() == 0)    // badness if we don't catch this...
+	if ((uint32)Mesh->TriangleCount() <= MinTriangleCount)    // badness if we don't catch this...
 	{
 		return;
 	}
@@ -1084,6 +1084,10 @@ void TMeshSimplification<QuadricErrorType>::FastCollapsePass(double fMinEdgeLeng
 			{
 				continue;
 			}
+			if ((uint32)Mesh->TriangleCount() <= MinTriangleCount)
+			{
+				break;
+			}
 			if (Cancelled())
 			{
 				return;
@@ -1106,7 +1110,7 @@ void TMeshSimplification<QuadricErrorType>::FastCollapsePass(double fMinEdgeLeng
 			}
 		}
 
-		if (num_last_pass == 0)     // converged
+		if (num_last_pass == 0 || (uint32)Mesh->TriangleCount() <= MinTriangleCount)     // converged
 		{
 			break;
 		}
diff --git a/Engine/Plugins/Runtime/GeometryProcessing/Source/DynamicMesh/Public/MeshSimplification.h b/Engine/Plugins/Runtime/GeometryProcessing/Source/DynamicMesh/Public/MeshSimplification.h
--- a/Engine/Plugins/Runtime/GeometryProcessing/Source/DynamicMesh/Public/MeshSimplification.h
+++ b/Engine/Plugins/Runtime/GeometryProcessing/Source/DynamicMesh/Public/MeshSimplification.h
@@ -149,8 +149,9 @@ public:
 	 * @param MinEdgeLength collapse any edge longer than this
 	 * @param Rounds number of collapse rounds
 	 * @param MeshIsClosedHint if you know the mesh is closed, this pass this true to avoid some precomputes
+	 * @param MinTriangleCount halt fast collapse if mesh falls below this triangle count
 	 */
-	virtual void FastCollapsePass(double MinEdgeLength, int Rounds = 1, bool bMeshIsClosedHint = false);
+	virtual void FastCollapsePass(double MinEdgeLength, int Rounds = 1, bool bMeshIsClosedHint = false, uint32 MinTriangleCount = 0);
 
 	
 
diff --git a/Engine/Source/Developer/GeometryProcessingInterfaces/Public/GeometryProcessingInterfaces/ApproximateActors.h b/Engine/Source/Developer/GeometryProcessingInterfaces/Public/GeometryProcessingInterfaces/ApproximateActors.h
--- a/Engine/Source/Developer/GeometryProcessingInterfaces/Public/GeometryProcessingInterfaces/ApproximateActors.h
+++ b/Engine/Source/Developer/GeometryProcessingInterfaces/Public/GeometryProcessingInterfaces/ApproximateActors.h
@@ -142,6 +142,7 @@ public:
 		ESimplificationPolicy MeshSimplificationPolicy = ESimplificationPolicy::FixedTriangleCount;
 		int32 FixedTriangleCount = 5000;
 		double SimplificationTargetMetric = 0.1;		// interpretation varies depending on MeshSimplificationPolicy
+		bool bEnableFastSimplifyPrePass = false;		// enable pre-simplify "fast-collapse" pass, using a relatively aggressive tolerance based on WorldSpaceApproximationAccuracyMeters
 
 		EGroundPlanePolicy GroundPlanePolicy = EGroundPlanePolicy::NoGroundPlane;
 		double GroundPlaneZHeight = 0.0;
@@ -167,6 +168,7 @@ public:
 		double FieldOfViewDegrees = 45.0;
 		double NearPlaneDist = 1.0;
 
+		bool bMaximizeBakeParallelism = true;		// if true, photocapture is computed at the same time as mesh generation. Faster but requires more memory.
 
 		//
 		// Material output settings
diff --git a/Engine/Source/Runtime/Engine/Classes/Engine/MeshMerging.h b/Engine/Source/Runtime/Engine/Classes/Engine/MeshMerging.h
--- a/Engine/Source/Runtime/Engine/Classes/Engine/MeshMerging.h
+++ b/Engine/Source/Runtime/Engine/Classes/Engine/MeshMerging.h
@@ -908,7 +908,13 @@ struct FMeshApproximationSettings
 	UPROPERTY(BlueprintReadWrite, EditAnywhere, Category = PerformanceSettings)
 	bool bUseRenderLODMeshes = false;
 
+	/** If true, a faster mesh simplfication strategy will be used. This can significantly reduce computation time and memory usage, but potentially at the cost of lower quality output. */
+	UPROPERTY(BlueprintReadWrite, EditAnywhere, Category = PerformanceSettings)
+	bool bEnableSimplifyPrePass = true;
 
+	/** If false, texture capture and baking will be done serially after mesh generation, rather than in parallel when possible. This will reduce the maximum memory requirements of the process.  */
+	UPROPERTY(BlueprintReadWrite, EditAnywhere, Category = PerformanceSettings)
+	bool bEnableParallelBaking = true;
 
 	//
 	// Debug Output Settings
