MeshModelingTools: Fix infinite loop crash in CubeGrid tool when cell size is set to 0.

#rnx

diff --git a/Engine/Plugins/Experimental/MeshModelingToolsetExp/Source/MeshModelingToolsExp/Private/CubeGridTool.cpp b/Engine/Plugins/Experimental/MeshModelingToolsetExp/Source/MeshModelingToolsExp/Private/CubeGridTool.cpp
--- a/Engine/Plugins/Experimental/MeshModelingToolsetExp/Source/MeshModelingToolsExp/Private/CubeGridTool.cpp
+++ b/Engine/Plugins/Experimental/MeshModelingToolsetExp/Source/MeshModelingToolsExp/Private/CubeGridTool.cpp
@@ -321,7 +321,7 @@ namespace CubeGridToolLocals
 		double StepSize = Grid.GetCurrentGridCellSize();
 
 		// Draw the inside only if there aren't too many lines to draw (approximate)
-		if (BoxDimensions[Dim1] / StepSize + BoxDimensions[Dim2] / StepSize > MAX_NUM_INTERIOR_GRID_LINES)
+		if (StepSize <= 0 || BoxDimensions[Dim1] / StepSize + BoxDimensions[Dim2] / StepSize > MAX_NUM_INTERIOR_GRID_LINES)
 		{
 			return;
 		}
diff --git a/Engine/Plugins/Experimental/MeshModelingToolsetExp/Source/MeshModelingToolsExp/Public/CubeGridTool.h b/Engine/Plugins/Experimental/MeshModelingToolsetExp/Source/MeshModelingToolsExp/Public/CubeGridTool.h
--- a/Engine/Plugins/Experimental/MeshModelingToolsetExp/Source/MeshModelingToolsExp/Public/CubeGridTool.h
+++ b/Engine/Plugins/Experimental/MeshModelingToolsetExp/Source/MeshModelingToolsExp/Public/CubeGridTool.h
@@ -102,10 +102,10 @@ public:
 	const uint8 MaxPowerOfTwo = 31;
 
 	/** Smallest block size to use in the grid. For instance, 3.125 results in
-	 blocks that are 100 sized at power of 2 since 3.125 * 2^5 = 100. */
+	 blocks that are 100 sized at 5 power of two since 3.125 * 2^5 = 100. */
 	UPROPERTY(EditAnywhere, Category = Options, meta = (
 		EditCondition = "bAllowedToEditGrid", HideEditConditionToggle,
-		UIMin = "0", UIMax = "10", ClampMin = "0", ClampMax = "1000"))
+		UIMin = "0.1", UIMax = "10", ClampMin = "0.001", ClampMax = "1000"))
 	double BlockBaseSize = 3.125;
 
 	/** When pushing/pulling in a way where the diagonal matters, setting this to true
diff --git a/Engine/Plugins/Runtime/MeshModelingToolset/Source/ModelingComponents/Public/Mechanics/CubeGrid.h b/Engine/Plugins/Runtime/MeshModelingToolset/Source/ModelingComponents/Public/Mechanics/CubeGrid.h
--- a/Engine/Plugins/Runtime/MeshModelingToolset/Source/ModelingComponents/Public/Mechanics/CubeGrid.h
+++ b/Engine/Plugins/Runtime/MeshModelingToolset/Source/ModelingComponents/Public/Mechanics/CubeGrid.h
@@ -104,8 +104,8 @@ public:
 	};
 
 	FCubeGrid()
-		: BaseGridCellSize(10)
-		, CurrentGridPower(3)
+		: BaseGridCellSize(3.125)
+		, CurrentGridPower(5)
 	{
 		UpdateCellSize();
 	}
@@ -130,6 +130,10 @@ public:
 	double GetBaseGridCellSize() const { return BaseGridCellSize; }
 	void SetBaseGridCellSize(double SizeIn)
 	{
+		if (!ensure(SizeIn > 0))
+		{
+			SizeIn = KINDA_SMALL_NUMBER;
+		}
 		BaseGridCellSize = SizeIn;
 		UpdateCellSize();
 	}
