MeshModelingTools: CubeGrid improvements.
- Fix bugs with how the tool deals with materials.
- Make UV's always stay aligned in the same orientation.
- Make UV's be contiguous when performing multiple steps with the same selection.
- Make sides maintain same groups with repeated steps of the same selection.
- Add support for scaling UV's.
- Add support for focus hotkey.
Also change how frame is chosen when subtracting to avoid some weird flipping logic.

diff --git a/Engine/Plugins/Experimental/MeshModelingToolsetExp/Source/MeshModelingToolsExp/Private/CubeGridTool.cpp b/Engine/Plugins/Experimental/MeshModelingToolsetExp/Source/MeshModelingToolsExp/Private/CubeGridTool.cpp
--- a/Engine/Plugins/Experimental/MeshModelingToolsetExp/Source/MeshModelingToolsExp/Private/CubeGridTool.cpp
+++ b/Engine/Plugins/Experimental/MeshModelingToolsetExp/Source/MeshModelingToolsExp/Private/CubeGridTool.cpp
@@ -29,8 +29,9 @@
 #include "ToolContextInterfaces.h"
 #include "ToolTargetManager.h"
 #include "ToolTargets/ToolTarget.h"
-#include "TargetInterfaces/MeshDescriptionCommitter.h"
-#include "TargetInterfaces/MeshDescriptionProvider.h"
+#include "TargetInterfaces/MaterialProvider.h"
+#include "TargetInterfaces/DynamicMeshCommitter.h"
+#include "TargetInterfaces/DynamicMeshProvider.h"
 #include "TargetInterfaces/PrimitiveComponentBackedTarget.h"
 #include "ToolSceneQueriesUtil.h"
 #include "ToolSetupUtil.h"
@@ -55,6 +56,8 @@ namespace CubeGridToolLocals
 	const FText SelectionChangeTransactionName = LOCTEXT("SelectionChangeTransaction", "Cube Grid Selection Change");
 	const FText ModeChangeTransactionName = LOCTEXT("ModeChangeTransaction", "Cube Grid Mode Change");
 
+	const FString PropertyCacheIdentifier = TEXT("CubeGridTool");
+
 	const FString& HoverLineSetID(TEXT("HoverLines"));
 	const FString& GridLineSetID(TEXT("GridLines"));
 	const FString& SelectionLineSetID(TEXT("SelectionLines"));
@@ -192,6 +195,9 @@ namespace CubeGridToolLocals
 	};
 	
 	// Attach a frame to the box such that Z points along the given direction.
+	// If we change the choices we make here, we may need to adjust how we pick the welded vertices in
+	// subtract mode (i.e., which way the values are mirrored when Z is flipped) and the output of
+	// GetFaceUVOrientations.
 	FOrientedBox3d ConvertToOrientedBox(const FAxisAlignedBox3d& Box, FCubeGrid::EFaceDirection Direction)
 	{
 		int FlatDim = FCubeGrid::DirToFlatDim(Direction);
@@ -201,7 +207,6 @@ namespace CubeGridToolLocals
 		FVector3d GridSpaceY = GridSpaceZ.Cross(GridSpaceX);
 
 		FVector3d BoxExtents = Box.Extents();
-		FVector3d Extents = BoxExtents; // The case if selection axis is Z
 
 		if (FlatDim == 0)
 		{
@@ -217,6 +222,36 @@ namespace CubeGridToolLocals
 		return FOrientedBox3d(FFrame3d(Box.Center(), GridSpaceX, GridSpaceY, GridSpaceZ), BoxExtents);
 	}
 
+	
+	// Based on direction of the operation, figures out how the faces of the opmesh need to be rotated
+	// relative to the UV plane to make the UV's always keep the same orientation relative to the grid.
+	// This depends on the choices we make in ConvertToOrientedBox, and is hard to reason through without
+	// trial and error... though there may have been a smarter way to pick our frame from direction that
+	// would have made this table easier to produce.
+	TArray<int32, TFixedAllocator<6>> GetFaceUVOrientations(FCubeGrid::EFaceDirection Direction)
+	{
+		switch (Direction)
+		{
+			case FCubeGrid::EFaceDirection::NegativeZ:
+				return TArray<int32, TFixedAllocator<6>>{ 2, 2, 2, 2, 2, 2 };
+				break;
+			case FCubeGrid::EFaceDirection::PositiveX:
+				return TArray<int32, TFixedAllocator<6>>{ 0, 0, 3, 1, 3, 1 };
+				break;
+			case FCubeGrid::EFaceDirection::NegativeX:
+				return TArray<int32, TFixedAllocator<6>>{ 2, 2, 3, 1, 1, 3 };
+				break;
+			case FCubeGrid::EFaceDirection::PositiveY:
+				return TArray<int32, TFixedAllocator<6>>{ 2, 2, 0, 0, 1, 3 };
+				break;
+			case FCubeGrid::EFaceDirection::NegativeY:
+				return TArray<int32, TFixedAllocator<6>>{ 0, 0, 2, 2, 3, 1 };
+				break;
+		}
+		// FCubeGrid::EFaceDirection::PositiveZ:
+		return TArray<int32, TFixedAllocator<6>>{ 0, 0, 0, 0, 0, 0 };
+	}
+
 	void GetNewSelectionFaceInBox(const FCubeGrid& Grid, const FAxisAlignedBox3d& Box, 
 		const FCubeGrid::FCubeFace& FaceIn, FCubeGrid::FCubeFace& FaceOut)
 	{
@@ -442,8 +477,8 @@ const FToolTargetTypeRequirements& UCubeGridToolBuilder::GetTargetRequirements()
 {
 	static FToolTargetTypeRequirements TypeRequirements({
 		UMaterialProvider::StaticClass(),
-		UMeshDescriptionCommitter::StaticClass(),
-		UMeshDescriptionProvider::StaticClass(),
+		UDynamicMeshCommitter::StaticClass(),
+		UDynamicMeshProvider::StaticClass(),
 		UPrimitiveComponentBackedTarget::StaticClass()
 		});
 	return TypeRequirements;
@@ -527,7 +562,9 @@ TUniquePtr<FDynamicMeshOperator> UCubeGridTool::MakeNewOperator()
 	// is no longer on grid because we changed grid power.
 	double FrameSpaceExtrudeAmount = GetFrameSpaceExtrudeDist(*CubeGrid, Selection.Box.Min, CurrentExtrudeAmount, Selection.Direction);
 
-	FOrientedBox3d FrameSpaceBox = ConvertToOrientedBox(Selection.Box, Selection.Direction);
+	FCubeGrid::EFaceDirection ZDirectionInGrid = CurrentExtrudeAmount < 0 ? FCubeGrid::FlipDir(Selection.Direction) : Selection.Direction;
+
+	FOrientedBox3d FrameSpaceBox = ConvertToOrientedBox(Selection.Box, ZDirectionInGrid);
 
 	// Give the selection box depth
 	FrameSpaceBox.Frame.Origin += (FrameSpaceExtrudeAmount / 2.0f) * FCubeGrid::DirToNormal(Selection.Direction);
@@ -545,6 +582,23 @@ TUniquePtr<FDynamicMeshOperator> UCubeGridTool::MakeNewOperator()
 	Op->WorldBox = WorldBox;
 	Op->bSubtract = CurrentExtrudeAmount < 0;
 	Op->bTrackChangedTids = true;
+	Op->OpMeshMaterialID = OpMeshMaterialID;
+	Op->OpMeshHeightUVOffset = OpMeshHeightUVOffset;
+	if (Settings->bKeepSideGroups)
+	{
+		if (Op->bSubtract && !OpMeshSubtractSideGroups.IsEmpty())
+		{
+			Op->OpMeshSideGroups = OpMeshSubtractSideGroups;
+		}
+		else if (!Op->bSubtract && !OpMeshAddSideGroups.IsEmpty())
+		{
+			Op->OpMeshSideGroups = OpMeshAddSideGroups;
+		}
+	}
+
+	Op->FaceUVOrientations = GetFaceUVOrientations(ZDirectionInGrid);
+	Op->bWorldSpaceUVs = MaterialProperties->bWorldSpaceUVScale;
+	Op->UVScale = MaterialProperties->UVScale;
 
 	if (Mode == EMode::Corner)
 	{
@@ -554,6 +608,14 @@ TUniquePtr<FDynamicMeshOperator> UCubeGridTool::MakeNewOperator()
 			Op->CornerInfo->WeldedAtBase[i] = !CornerSelectedFlags[i];
 		}
 
+		if (Op->bSubtract)
+		{
+			// If we're flipping the Z direction of our frame, we currently effectively do so by rotating around 
+			// the X axis, which determines which way we need to flip these corner labels.
+			Swap(Op->CornerInfo->WeldedAtBase[0], Op->CornerInfo->WeldedAtBase[3]);
+			Swap(Op->CornerInfo->WeldedAtBase[1], Op->CornerInfo->WeldedAtBase[2]);
+		}
+
 		Op->bCrosswiseDiagonal = Settings->bCrosswiseDiagonal;
 	}
 
@@ -575,7 +637,7 @@ void UCubeGridTool::SlideSelection(int32 Amount, bool bEmitChange)
 		Selection.Box.Min + FrameSpaceDisplacement,
 		Selection.Box.Max + FrameSpaceDisplacement);
 
-	SetSelection(NewSelection, true);
+	SetSelection(NewSelection, bEmitChange);
 }
 
 void UCubeGridTool::SetSelection(const FSelection& NewSelection, bool bEmitChange)
@@ -649,16 +711,21 @@ void UCubeGridTool::Setup()
 	AddToolPropertySource(OutputTypeProperties);
 
 	MaterialProperties = NewObject<UNewMeshMaterialProperties>(this);
-	MaterialProperties->RestoreProperties(this);
-	MaterialProperties->bShowExtendedOptions = false;
+	// Change the default by setting it before restoring
+	MaterialProperties->bWorldSpaceUVScale = true;
+	MaterialProperties->RestoreProperties(this, PropertyCacheIdentifier);
 	AddToolPropertySource(MaterialProperties);
 
 	CurrentMesh = MakeShared<FDynamicMesh3, ESPMode::ThreadSafe>();
 	CurrentMesh->EnableAttributes();
+	CurrentMeshMaterials.Empty();
 	if (Target)
 	{
 		*CurrentMesh = UE::ToolTarget::GetDynamicMeshCopy(Target);
 		UE::ToolTarget::SetSourceObjectVisible(Target, false);
+
+		CurrentMeshTransform = UE::ToolTarget::GetLocalToWorldTransform(Target);
+		CurrentMeshMaterials = UE::ToolTarget::GetMaterialSet(Target).Materials;
 	}
 
 	MeshSpatial = MakeShared<FDynamicMeshAABBTree3, ESPMode::ThreadSafe>();
@@ -671,20 +738,25 @@ void UCubeGridTool::Setup()
 	if (Target)
 	{
 		Preview->PreviewMesh->UpdatePreview(CurrentMesh.Get());
-		CurrentMeshTransform = UE::ToolTarget::GetLocalToWorldTransform(Target);
-
-		// Set materials
-		FComponentMaterialSet MaterialSet = UE::ToolTarget::GetMaterialSet(Target);
-		Preview->ConfigureMaterials(MaterialSet.Materials, ToolSetupUtil::GetDefaultWorkingMaterial(GetToolManager()));
 	}
 	Preview->PreviewMesh->SetTransform((FTransform)CurrentMeshTransform);
 	Preview->OnOpCompleted.AddWeakLambda(this, [this](const FDynamicMeshOperator* UncastOp) {
 		const FCubeGridBooleanOp* Op = static_cast<const FCubeGridBooleanOp*>(UncastOp);
 		if (Op->InputMesh == ComputeStartMesh)
 		{
-			LastOpChangedTids = static_cast<const FCubeGridBooleanOp*>(Op)->ChangedTids;
+			LastOpChangedTids = Op->ChangedTids;
+			if (Op->bSubtract)
+			{
+				OpMeshSubtractSideGroups = Op->OpMeshSideGroups;
+			}
+			else
+			{
+				OpMeshAddSideGroups = Op->OpMeshSideGroups;
+			}
+			
 		}
 	});
+	UpdateOpMaterials();
 
 	CubeGrid = MakeShared<FCubeGrid>();
 	CubeGrid->SetGridFrame(FFrame3d(Settings->GridFrameOrigin, Settings->GridFrameOrientation.Quaternion()));
@@ -807,11 +879,33 @@ void UCubeGridTool::Setup()
 		[this](bool bOn) {
 			InvalidatePreview();
 		});
+	Settings->WatchProperty(Settings->bKeepSideGroups,
+		[this](bool bOn) {
+			InvalidatePreview(false);
+		});
 	Settings->WatchProperty(Settings->PlaneTolerance,
 		[this](double Tolerance) {
 			InvalidatePreview(false);
 		});
 
+	MaterialProperties->WatchProperty(MaterialProperties->Material,
+		[this](TWeakObjectPtr<UMaterialInterface> Material) {
+			UpdateOpMaterials();
+			InvalidatePreview(false);
+		});
+	MaterialProperties->WatchProperty(MaterialProperties->UVScale,
+		[this](float UVScale) {
+			InvalidatePreview(false);
+		});
+	MaterialProperties->WatchProperty(MaterialProperties->bShowWireframe,
+		[this](float UVScale) {
+			Preview->PreviewMesh->EnableWireframe(MaterialProperties->bShowWireframe);
+		});
+	MaterialProperties->WatchProperty(MaterialProperties->bWorldSpaceUVScale,
+		[this](float UVScale) {
+			InvalidatePreview(false);
+		});
+
 	auto UpdateFromDetailsPanelTransformChange = [this]()
 	{
 		CubeGrid->SetGridFrame(FFrame3d(Settings->GridFrameOrigin, Settings->GridFrameOrientation.Quaternion()));
@@ -848,6 +942,18 @@ void UCubeGridTool::Setup()
 	}
 }
 
+void UCubeGridTool::UpdateOpMaterials()
+{
+	TArray<UMaterialInterface*> PreviewMaterials = CurrentMeshMaterials;
+	OpMeshMaterialID = PreviewMaterials.Find(MaterialProperties->Material.Get());
+	if (OpMeshMaterialID == INDEX_NONE)
+	{
+		OpMeshMaterialID = PreviewMaterials.Add(MaterialProperties->Material.Get());
+	}
+
+	Preview->ConfigureMaterials(PreviewMaterials, ToolSetupUtil::GetDefaultWorkingMaterial(GetToolManager()));
+}
+
 void UCubeGridTool::UpdateComputeInputs()
 {
 	ComputeStartMesh = MakeShared<FDynamicMesh3>(*CurrentMesh);
@@ -855,8 +961,11 @@ void UCubeGridTool::UpdateComputeInputs()
 
 void UCubeGridTool::Shutdown(EToolShutdownType ShutdownType)
 {
+	using namespace CubeGridToolLocals;
+
 	Settings->SaveProperties(this);
-	MaterialProperties->SaveProperties(this);
+	MaterialProperties->SaveProperties(this, PropertyCacheIdentifier);
+	OutputTypeProperties->SaveProperties(this);
 
 	if (Mode == EMode::Corner)
 	{
@@ -870,7 +979,11 @@ void UCubeGridTool::Shutdown(EToolShutdownType ShutdownType)
 		if (bChangesMade)
 		{
 			GetToolManager()->BeginUndoTransaction(LOCTEXT("CubeGridToolEditTransactionName", "Block Tool Edit"));
-			UE::ToolTarget::CommitDynamicMeshUpdate(Target, *CurrentMesh, true);
+			FComponentMaterialSet OutputMaterialSet;
+			OutputMaterialSet.Materials = CurrentMeshMaterials;
+
+			UE::ToolTarget::CommitDynamicMeshUpdate(Target, *CurrentMesh, true,
+				FConversionToMeshDescriptionOptions(), &OutputMaterialSet);
 			GetToolManager()->EndUndoTransaction();
 		}
 	}
@@ -882,7 +995,7 @@ void UCubeGridTool::Shutdown(EToolShutdownType ShutdownType)
 		NewMeshObjectParams.TargetWorld = TargetWorld;
 		NewMeshObjectParams.Transform = (FTransform)CurrentMeshTransform;
 		NewMeshObjectParams.BaseName = TEXT("CubeGridToolOutput");
-		NewMeshObjectParams.Materials.Add(MaterialProperties->Material.Get());
+		NewMeshObjectParams.Materials = CurrentMeshMaterials;
 		NewMeshObjectParams.SetMesh(CurrentMesh.Get());
 		OutputTypeProperties->ConfigureCreateMeshObjectParams(NewMeshObjectParams);
 		FCreateMeshObjectResult Result = UE::Modeling::CreateMeshObject(GetToolManager(), MoveTemp(NewMeshObjectParams));
@@ -936,11 +1049,22 @@ void UCubeGridTool::ApplyPreview()
 	ChangeTracker.SaveTriangles(*LastOpChangedTids, true /*bSaveVertices*/);
 
 	// Update current mesh
+	bool bWasValid = Preview->ProcessCurrentMesh([this](const FDynamicMesh3& Mesh)
+	{
+		CurrentMesh->Copy(Mesh);
+	});
+	if (!ensure(bWasValid))
+	{
+		return;
+	}
+
 	bChangesMade = true; // TODO: make this undoable
 
-	CurrentMesh->Copy(*Preview->PreviewMesh->GetMesh());
 	MeshSpatial->Build();
 
+	CurrentMeshMaterials.Reset();
+	Preview->PreviewMesh->GetMaterials(CurrentMeshMaterials);
+
 	UpdateComputeInputs();
 
 	const FText TransactionText = LOCTEXT("CubeGridToolTransactionName", "Block Tool Change");
@@ -952,11 +1076,15 @@ void UCubeGridTool::ApplyPreview()
 
 	if (bAdjustSelectionOnPreviewUpdate)
 	{
-		// Change the selection to the new location. Note that this should to happen
+		// Save UV offset before sliding selection
+		OpMeshHeightUVOffset += GetFrameSpaceExtrudeDist(*CubeGrid, Selection.Box.Min, CurrentExtrudeAmount, Selection.Direction);
+
+		// Change the selection to the new location. Note that this should happen only
 		// after resetting bPreviewMayDiffer to avoid an extra preview reset when selection
 		// changes.
 		SlideSelection(CurrentExtrudeAmount, true);
 	}
+
 	CurrentExtrudeAmount = 0;
 
 	GetToolManager()->EmitObjectChange(this, MakeUnique<FCubeGridToolMeshChange>(ChangeTracker.EndChange()), TransactionText);
@@ -1207,6 +1335,8 @@ void UCubeGridTool::EndSelectionChange()
 	if (bPreviousHaveSelection != bHaveSelection
 		|| (bHaveSelection && PreviousSelection != Selection))
 	{
+		ResetMultiStepConsistencyData();
+
 		GetToolManager()->EmitObjectChange(this,
 			MakeUnique<FCubeGridToolSelectionChange>(bPreviousHaveSelection, bHaveSelection,
 				PreviousSelection, Selection),
@@ -1214,6 +1344,13 @@ void UCubeGridTool::EndSelectionChange()
 	}
 }
 
+void UCubeGridTool::ResetMultiStepConsistencyData()
+{
+	OpMeshHeightUVOffset = 0;
+	OpMeshAddSideGroups.Empty();
+	OpMeshSubtractSideGroups.Empty();
+}
+
 void UCubeGridTool::UpdateGizmoVisibility(bool bVisible)
 {
 	GridGizmo->SetVisibility(bVisible);
@@ -1627,22 +1764,22 @@ void UCubeGridTool::UpdateCornerModeLineSet()
 	LineSet->Clear();
 	if (Mode == EMode::Corner && CurrentExtrudeAmount < 0)
 	{
-		FOrientedBox3d FrameSpaceBox = ConvertToOrientedBox(Selection.Box, Selection.Direction);
-
+		// Since we're subtracting here, we need to do some flipping.
+		FOrientedBox3d FrameSpaceBox = ConvertToOrientedBox(Selection.Box, FCubeGrid::FlipDir(Selection.Direction));
 		bool CornerWelded[4];
 		for (int i = 0; i < 4; ++i) {
 			CornerWelded[i] = !CornerSelectedFlags[i];
 		}
-
-		// The choice of diagonal here lines up with the generator in CubeGridBooleanOp. The
-		// indices look a little different because we're accounting for mirroring that happens in
-		// subtract mode, though we would have actually gotten the same results either way in
-		// the cases we care about (the nonplanar ones)
-		int DiagStartIdx = 1;
-		if (CornerWelded[0] != CornerWelded[2] ||
-			(!CornerWelded[0] && CornerWelded[1] && CornerWelded[3]))
+		Swap(CornerWelded[0], CornerWelded[3]);
+		Swap(CornerWelded[1], CornerWelded[2]);
+
+		// The choice of diagonal here lines up with the generator in CubeGridBooleanOp for
+		// the bottom face.
+		int DiagStartIdx = 0;
+		if (CornerWelded[1] != CornerWelded[3] ||
+			(!CornerWelded[1] && CornerWelded[0] && CornerWelded[2]))
 		{
-			DiagStartIdx = 0;
+			DiagStartIdx = 1;
 		}
 		DiagStartIdx = Settings->bCrosswiseDiagonal ? 1 - DiagStartIdx : DiagStartIdx;
 		
@@ -1707,6 +1844,7 @@ void UCubeGridTool::ApplyFlipSelection()
 	FSelection NewSelection = Selection;
 	NewSelection.Direction = FCubeGrid::FlipDir(Selection.Direction);
 	SetSelection(NewSelection, true);
+	ResetMultiStepConsistencyData();
 
 	GetToolManager()->EndUndoTransaction();
 
@@ -1728,6 +1866,8 @@ void UCubeGridTool::ApplySlide(int32 NumBlocks)
 
 	GetToolManager()->BeginUndoTransaction(LOCTEXT("SlideTransactionName", "Slide Selection"));
 	SlideSelection(NumBlocks, true);
+	ResetMultiStepConsistencyData();
+
 	GetToolManager()->EndUndoTransaction();
 }
 
@@ -1859,6 +1999,8 @@ void UCubeGridTool::ApplyAction(ECubeGridToolAction ActionType)
 
 void UCubeGridTool::RegisterActions(FInteractiveToolActionSet& ActionSet)
 {
+	using namespace CubeGridToolLocals;
+
 	int32 ActionID = (int32)EStandardToolActions::BaseClientDefinedActionID + 1;
 	ActionSet.RegisterAction(this, ActionID++,
 		TEXT("PullBlock"),
@@ -1948,6 +2090,28 @@ void UCubeGridTool::RegisterActions(FInteractiveToolActionSet& ActionSet)
 		});
 }
 
+FBox UCubeGridTool::GetWorldSpaceFocusBox()
+{
+	using namespace CubeGridToolLocals;
+
+	FAxisAlignedBox3d Bounds = FAxisAlignedBox3d::Empty();
+	if (bHaveSelection)
+	{
+		FOrientedBox3d FrameSpaceBox = ConvertToOrientedBox(Selection.Box, Selection.Direction);
+
+		// The resulting oriented box is flat in the z frame direction (which is aligned to Selection.Direction).
+		// So we only need to contain the four corners of a z face instead of doing 8 corners.
+		const FFrame3d& GridFrame = CubeGrid->GetFrame();
+		int ZFaceIndex = 0;
+		for (int CornerIndex = 0; CornerIndex < 4; ++CornerIndex)
+		{
+			Bounds.Contain(GridFrame.FromFramePoint(FrameSpaceBox.GetCorner(IndexUtil::BoxFaces[ZFaceIndex][CornerIndex])));
+		}
+	}
+
+	return (FBox)Bounds;
+}
+
 void UCubeGridTool::StartCornerMode()
 {
 	using namespace CubeGridToolLocals;
diff --git a/Engine/Plugins/Experimental/MeshModelingToolsetExp/Source/MeshModelingToolsExp/Public/CubeGridTool.h b/Engine/Plugins/Experimental/MeshModelingToolsetExp/Source/MeshModelingToolsExp/Public/CubeGridTool.h
--- a/Engine/Plugins/Experimental/MeshModelingToolsetExp/Source/MeshModelingToolsExp/Public/CubeGridTool.h
+++ b/Engine/Plugins/Experimental/MeshModelingToolsetExp/Source/MeshModelingToolsExp/Public/CubeGridTool.h
@@ -11,7 +11,7 @@
 #include "GeometryBase.h"
 #include "InteractiveTool.h"
 #include "InteractiveToolBuilder.h"
-#include "InteractiveToolQueryInterfaces.h" // IInteractiveToolNestedAcceptCancelAPI
+#include "InteractiveToolQueryInterfaces.h" // IInteractiveToolNestedAcceptCancelAPI, IInteractiveToolCameraFocusAPI
 #include "Mechanics/CubeGrid.h"
 #include "ModelingOperators.h"
 #include "OrientedBoxTypes.h"
@@ -61,9 +61,13 @@ protected:
 UENUM()
 enum class ECubeGridToolFaceSelectionMode
 {
+	/** Use hit normal to pick the outer face of the containing cell. */
 	OutsideBasedOnNormal,
+	/** Use hit normal to pierce backward through the geometry to pick an inside face of the containing cell. */
 	InsideBasedOnNormal,
+	/** Use view ray to pick the outer face of the containing cell. */
 	OutsideBasedOnViewRay,
+	/** Use view ray to pierce backward through the geometry to pick an inside face of the containing cell. */
 	InsideBasedOnViewRay
 };
 
@@ -115,6 +119,12 @@ public:
 		EditCondition = "bInCornerMode", HideEditConditionToggle))
 	bool bCrosswiseDiagonal = false;
 
+	/** When performing multiple push/pulls with the same selection, attempt to keep the
+	 same group IDs on the sides of the new geometry (ie multiple E/Q presses will not
+	 result in different group topology around the sides compared to a single Ctrl+drag). */
+	UPROPERTY(EditAnywhere, Category = Options, AdvancedDisplay)
+	bool bKeepSideGroups = true;
+
 	/** When performing selection, the tolerance to use when determining
 	 whether things lie in the same plane as a cube face. */
 	UPROPERTY(EditAnywhere, Category = BlockSelection, AdvancedDisplay, meta = (
@@ -252,7 +262,8 @@ UCLASS()
 class MESHMODELINGTOOLSEXP_API UCubeGridTool : public UInteractiveTool,
 	public IClickDragBehaviorTarget, public IHoverBehaviorTarget,
 	public UE::Geometry::IDynamicMeshOperatorFactory,
-	public IInteractiveToolNestedAcceptCancelAPI
+	public IInteractiveToolNestedAcceptCancelAPI,
+	public IInteractiveToolCameraFocusAPI
 {
 	GENERATED_BODY()
 protected:
@@ -349,6 +360,10 @@ public:
 	virtual bool CanCurrentlyNestedAccept() override;
 	virtual bool ExecuteNestedAcceptCommand() override;
 
+	// IInteractiveToolCameraFocusAPI
+	virtual bool SupportsWorldSpaceFocusBox() { return bHaveSelection; }
+	virtual FBox GetWorldSpaceFocusBox() override;
+
 protected:
 
 	UPROPERTY()
@@ -464,6 +479,18 @@ protected:
 	UE::Geometry::FTransformSRT3d CurrentMeshTransform = UE::Geometry::FTransformSRT3d::Identity();
 	TSharedPtr<TArray<int32>, ESPMode::ThreadSafe> LastOpChangedTids;
 
+	TArray<UMaterialInterface*> CurrentMeshMaterials;
+	int32 OpMeshMaterialID = 0;
+	void UpdateOpMaterials();
+
+	// These are used to keep UV's and side groups consistent across multiple E/Q (push/pull) presses. 
+	// This data should be reset whenever the selection changes in a way that is not a byproduct of
+	// a push/pull.
+	double OpMeshHeightUVOffset = 0;
+	TArray<int32, TFixedAllocator<4>> OpMeshAddSideGroups;
+	TArray<int32, TFixedAllocator<4>> OpMeshSubtractSideGroups;
+	void ResetMultiStepConsistencyData();
+
 	// Safe inputs for the background compute to use, untouched by undo/redo/other CurrentMesh updates.
 	TSharedPtr<const UE::Geometry::FDynamicMesh3, ESPMode::ThreadSafe> ComputeStartMesh;
 	void UpdateComputeInputs();
diff --git a/Engine/Plugins/Runtime/MeshModelingToolset/Source/MeshModelingTools/Public/Properties/MeshMaterialProperties.h b/Engine/Plugins/Runtime/MeshModelingToolset/Source/MeshModelingTools/Public/Properties/MeshMaterialProperties.h
--- a/Engine/Plugins/Runtime/MeshModelingToolset/Source/MeshModelingTools/Public/Properties/MeshMaterialProperties.h
+++ b/Engine/Plugins/Runtime/MeshModelingToolset/Source/MeshModelingTools/Public/Properties/MeshMaterialProperties.h
@@ -38,7 +38,7 @@ public:
 	bool bShowWireframe = false;
 
 	/** If true, extended options are available */ 
-	UPROPERTY()
+	UPROPERTY(meta = (TransientToolProperty))
 	bool bShowExtendedOptions = true;
 };
 
diff --git a/Engine/Plugins/Runtime/MeshModelingToolset/Source/ModelingComponents/Public/Mechanics/CubeGrid.h b/Engine/Plugins/Runtime/MeshModelingToolset/Source/ModelingComponents/Public/Mechanics/CubeGrid.h
--- a/Engine/Plugins/Runtime/MeshModelingToolset/Source/ModelingComponents/Public/Mechanics/CubeGrid.h
+++ b/Engine/Plugins/Runtime/MeshModelingToolset/Source/ModelingComponents/Public/Mechanics/CubeGrid.h
@@ -37,16 +37,16 @@ public:
 		PositiveZ = 3
 	};
 
-	/** Convert face direction to its normal. */
+	/**
+	 * Convert face direction to the index of the nonzero normal dimension, or equivalently,
+	 * the dimension in which that face is flat.
+	 */
 	static int8 DirToFlatDim(EFaceDirection DirectionIn)
 	{
 		return FMath::Abs(static_cast<int8>(DirectionIn)) - 1;
 	}
 
-	/** 
-	 * Convert face direction to the index of the nonzero normal dimension, or equivalently,
-	 * the dimension in which that face is flat.
-	 */
+	 /** Convert face direction to its normal. */
 	static FVector3d DirToNormal(EFaceDirection DirectionIn)
 	{
 		FVector3d Normal = FVector3d::Zero();
diff --git a/Engine/Plugins/Runtime/MeshModelingToolset/Source/ModelingOperators/Private/CompositionOps/CubeGridBooleanOp.cpp b/Engine/Plugins/Runtime/MeshModelingToolset/Source/ModelingOperators/Private/CompositionOps/CubeGridBooleanOp.cpp
--- a/Engine/Plugins/Runtime/MeshModelingToolset/Source/ModelingOperators/Private/CompositionOps/CubeGridBooleanOp.cpp
+++ b/Engine/Plugins/Runtime/MeshModelingToolset/Source/ModelingOperators/Private/CompositionOps/CubeGridBooleanOp.cpp
@@ -18,6 +18,8 @@ namespace CubeGridBooleanOpLocals
  * Generator for creating ramps/pyramids/pyramid cutaways in the bounds of a box. The "welded"
  * refers to welding vertices along the Z axis to create those shapes. See the comment in
  * FCubeGridBooleanOp::FCornerInfo for more information.
+ * 
+ * The generator assigns the group IDs 0-3 to the sides (which become 1-4 in the FDynamicMesh3).
  */
 class FWeldedMinimalBoxMeshGenerator : public FMeshShapeGenerator
 {
@@ -25,9 +27,26 @@ public:
 	FOrientedBox3d Box;
 	FCubeGridBooleanOp::FCornerInfo CornerInfo;
 	bool bCrosswiseDiagonal = false;
+
+	// See comments in CubeGridBooleanOp.h about how these work
+	TArray<int32, TFixedAllocator<6>> FaceUVOrientations;
+	double HeightUVOffset = 0;
 		
 	virtual FMeshShapeGenerator& Generate() override
 	{
+		// Fill any unset orientations with 0.
+		FaceUVOrientations.SetNumZeroed(6);
+
+		// Make sure all orientations are in [0,3]
+		for (int& Orientation : FaceUVOrientations)
+		{
+			if (Orientation < 0)
+			{
+				Orientation = 4 - ((-Orientation) % 4);
+			}
+			Orientation %= 4;
+		}
+
 		// There is a particular odd case that we have to guard against. If one set of diagonal
 		// corners is welded and the second set is unwelded, and if the chosen triangulation
 		// diagonal connects the welded vertices, then the edge between these verts would normally
@@ -80,7 +99,7 @@ protected:
 
 		int GroupID = 0;
 
-		// These actually get reinitialized at during each face traversal, and are out here as globals
+		// These actually get reinitialized during each face traversal, and are out here as globals
 		// to the lambdas we use (to avoid having to pass tons of args). They are indexed by 0-3 
 		// counterclockwise from the bottom left corner, like IndexUtil::BoxFacesUV
 		int UVIndices[4];
@@ -90,11 +109,42 @@ protected:
 		double FaceHeight;
 
 		auto MakeUVsAndNormals = [this, &UVIndices, &NormalIndices, &FaceCornerToVert, &FaceWidth, &FaceHeight]
-		(int FaceCornerIndex, const FVector3f& Normal, int ParentVert) {
-			UVIndices[FaceCornerIndex] = AppendUV(FVector2f(
-				FaceWidth * IndexUtil::BoxFacesUV[FaceCornerIndex].X,
-				FaceHeight * IndexUtil::BoxFacesUV[FaceCornerIndex].Y), ParentVert);
+		(int FaceCornerIndex, const FVector3f& Normal, int ParentVert, float HeightOffset, int FaceOrientation) {
+			
+			/*
+			 * FaceOrientation determines how the face is rotated in UV space. The assignment table ends up:
+			 * 
+			 *    FaceOrientation
+			 * corner  0    1    2    3
+			 *   0    0,0  h,0  w,h  0,w
+			 *   1    w,0  h,w  0,h  0,0
+			 *   2    w,h  0,w  0,0  h,0
+			 *   3    0,h  0,0  w,0  h,w
+			 * 
+			 * Where HeightOffset needs to be added to the column where h is used (V for rotations 0 and 2, U for others)
+			 * 
+			 * From the table, we can see that we can index into IndexUtil::BoxFacesUV with an index that is rotated
+			 * backwards from FaceCornerIndex FaceOrientation number of times to get multiplier that we can then multiply
+			 * by w and h appropriately.
+			 */
+
+			int RotatedCornerIndex = (FaceCornerIndex + 4 - FaceOrientation) % 4;
+			FVector2i UVBase = IndexUtil::BoxFacesUV[RotatedCornerIndex];
+
+			FVector2f UV;
+			if (FaceOrientation % 2 == 0)
+			{
+				UV.X = UVBase.X * FaceWidth;
+				UV.Y = UVBase.Y * FaceHeight + HeightOffset;
+			}
+			else
+			{
+				UV.X = UVBase.X * FaceHeight + HeightOffset;
+				UV.Y = UVBase.Y * FaceWidth;
+			}
 
+			UVIndices[FaceCornerIndex] = AppendUV(UV, ParentVert);
+			
 			NormalIndices[FaceCornerIndex] = AppendNormal(Normal, ParentVert);
 			FaceCornerToVert[FaceCornerIndex] = ParentVert;
 		};
@@ -116,9 +166,17 @@ protected:
 			(FVector3f)Box.AxisY(),
 		};
 
+		// Maps each left corner of the bottom face to the IndexUtil::BoxFaces index of the face that we'll construct.
+		int BottomCornerToSideFaceIndex[4] = { 4, 3, 5, 2 };
+
 		// Set up the sides
 		for (int LeftCorner = 0; LeftCorner < 4; ++LeftCorner)
 		{
+			GroupID = LeftCorner;
+
+			int FaceIndex = BottomCornerToSideFaceIndex[LeftCorner];
+			int FaceUVOrientation = FaceUVOrientations[FaceIndex];
+
 			// We're iterating over the bottom corners of the oriented box (see TOrientedBox3::GetCorner()),
 			// but due to the way these are numbered, if we are looking at the side, the bottom two verts
 			// actually decrease in corner index (wrapping around).
@@ -155,28 +213,29 @@ protected:
 			FaceHeight = Distance(Box.GetCorner(LeftCorner), Box.GetCorner(LeftUp));
 
 			// The bottom two vertices will definitely need uv/normal elements
-			MakeUVsAndNormals(0, SideFaceBoxSpaceNormals[LeftCorner], CornerToVert[LeftCorner]);
-			MakeUVsAndNormals(1, SideFaceBoxSpaceNormals[LeftCorner], CornerToVert[RightCorner]);
+			MakeUVsAndNormals(0, SideFaceBoxSpaceNormals[LeftCorner], CornerToVert[LeftCorner], HeightUVOffset, FaceUVOrientation);
+			MakeUVsAndNormals(1, SideFaceBoxSpaceNormals[LeftCorner], CornerToVert[RightCorner], HeightUVOffset, FaceUVOrientation);
 
 			if (!CornerInfo.WeldedAtBase[LeftCorner])
 			{
-				MakeUVsAndNormals(3, SideFaceBoxSpaceNormals[LeftCorner], CornerToVert[LeftUp]);
+				MakeUVsAndNormals(3, SideFaceBoxSpaceNormals[LeftCorner], CornerToVert[LeftUp], HeightUVOffset, FaceUVOrientation);
 				MakeTriangle(0, 1, 3);
 				if (!CornerInfo.WeldedAtBase[RightCorner])
 				{
-					MakeUVsAndNormals(2, SideFaceBoxSpaceNormals[LeftCorner], CornerToVert[RightUp]);
+					MakeUVsAndNormals(2, SideFaceBoxSpaceNormals[LeftCorner], CornerToVert[RightUp], HeightUVOffset, FaceUVOrientation);
 					MakeTriangle(3, 1, 2);
 				}
 			}
 			else
 			{
-				MakeUVsAndNormals(2, SideFaceBoxSpaceNormals[LeftCorner], CornerToVert[RightUp]);
+				MakeUVsAndNormals(2, SideFaceBoxSpaceNormals[LeftCorner], CornerToVert[RightUp], HeightUVOffset, FaceUVOrientation);
 				MakeTriangle(0, 1, 2);
 			}
-
-			++GroupID;
 		}// end making side faces
 
+		// 0-3 are reserved group id's for sides even if a side doesn't exist.
+		GroupID = 4;
+
 		// Figure out which direction we're going to triangulate the top and bottom. 
 		// In the normal case, we place the diagonal such that a single raised/lowered
 		// corner results in a pyramid with a square (rather than triangular) base, and
@@ -190,8 +249,12 @@ protected:
 		}
 		DiagFaceIdx1 = bCrosswiseDiagonal ? 1 - DiagFaceIdx1 : DiagFaceIdx1;
 
+		// DiagFaceIdx1 is either 0 or 1 here and static analysis will complain if we do
+		// mod 4 for some of the later statements. Throw in a check in case we ever change
+		// the above.
+		checkSlow(DiagFaceIdx1 == 0 || DiagFaceIdx1 == 1);
 		int DiagFaceIdx2 = DiagFaceIdx1 + 2;
-		int CCWFromDiag1 = (DiagFaceIdx1 + 1) % 4;
+		int CCWFromDiag1 = DiagFaceIdx1 + 1;
 		int CCWFromDiag2 = (DiagFaceIdx2 + 1) % 4;
 
 		// We've iterated across the sides but it's still actually possible to end up not
@@ -209,19 +272,20 @@ protected:
 		// Prep for UVs
 		FaceWidth = Distance(Box.GetCorner(0), Box.GetCorner(1));
 		FaceHeight = Distance(Box.GetCorner(0), Box.GetCorner(3));
+		int FaceUVOrientation = FaceUVOrientations[0];
 
-		MakeUVsAndNormals(DiagFaceIdx1, (FVector3f)-Box.Frame.Z(), CornerToVert[DiagFaceIdx1]);
-		MakeUVsAndNormals(DiagFaceIdx2, (FVector3f)-Box.Frame.Z(), CornerToVert[DiagFaceIdx2]);
+		MakeUVsAndNormals(DiagFaceIdx1, (FVector3f)-Box.Frame.Z(), CornerToVert[DiagFaceIdx1], 0, FaceUVOrientation);
+		MakeUVsAndNormals(DiagFaceIdx2, (FVector3f)-Box.Frame.Z(), CornerToVert[DiagFaceIdx2], 0, FaceUVOrientation);
 
 		bool bDiagonalWelded = CornerInfo.WeldedAtBase[DiagFaceIdx1] && CornerInfo.WeldedAtBase[DiagFaceIdx2];
 		if (!(bDiagonalWelded && CornerInfo.WeldedAtBase[CCWFromDiag1]))
 		{
-			MakeUVsAndNormals(CCWFromDiag1, (FVector3f)-Box.Frame.Z(), CornerToVert[CCWFromDiag1]);
+			MakeUVsAndNormals(CCWFromDiag1, (FVector3f)-Box.Frame.Z(), CornerToVert[CCWFromDiag1], 0, FaceUVOrientation);
 			MakeTriangle(DiagFaceIdx1, CCWFromDiag1, DiagFaceIdx2);
 		}
 		if (!(bDiagonalWelded && CornerInfo.WeldedAtBase[CCWFromDiag2]))
 		{
-			MakeUVsAndNormals(CCWFromDiag2, (FVector3f)-Box.Frame.Z(), CornerToVert[CCWFromDiag2]);
+			MakeUVsAndNormals(CCWFromDiag2, (FVector3f)-Box.Frame.Z(), CornerToVert[CCWFromDiag2], 0, FaceUVOrientation);
 			MakeTriangle(DiagFaceIdx1, DiagFaceIdx2, CCWFromDiag2);
 		}
 		++GroupID;
@@ -263,6 +327,7 @@ protected:
 		FaceHeight = FMath::Max(
 			FVector3d::Distance(CornerPositions[0], CornerPositions[3]),
 			FVector3d::Distance(CornerPositions[1], CornerPositions[2]));
+		FaceUVOrientation = FaceUVOrientations[1];
 
 		// We're looking at things from the opposite direction, so our diagonal needs to
 		// connect the opposite face indices.
@@ -286,9 +351,9 @@ protected:
 			FVector3f Normal = AToC.Cross(AToB);
 			Normal.Normalize();
 				
-			MakeUVsAndNormals(DiagFaceIdx1, Normal, FaceCornerToVert[DiagFaceIdx1]);
-			MakeUVsAndNormals(CCWFromDiag1, Normal, FaceCornerToVert[CCWFromDiag1]);
-			MakeUVsAndNormals(DiagFaceIdx2, Normal, FaceCornerToVert[DiagFaceIdx2]);
+			MakeUVsAndNormals(DiagFaceIdx1, Normal, FaceCornerToVert[DiagFaceIdx1], 0, FaceUVOrientation);
+			MakeUVsAndNormals(CCWFromDiag1, Normal, FaceCornerToVert[CCWFromDiag1], 0, FaceUVOrientation);
+			MakeUVsAndNormals(DiagFaceIdx2, Normal, FaceCornerToVert[DiagFaceIdx2], 0, FaceUVOrientation);
 
 			MakeTriangle(DiagFaceIdx1, CCWFromDiag1, DiagFaceIdx2);
 
@@ -304,7 +369,7 @@ protected:
 			{
 				// Have to make this copy because we're not allowed to pass a reference to an element
 				// in the container we're modifying.
-				MakeUVsAndNormals(CCWFromDiag2, FVector3f(Normals[NormalIndices[DiagFaceIdx1]]), FaceCornerToVert[CCWFromDiag2]);
+				MakeUVsAndNormals(CCWFromDiag2, FVector3f(Normals[NormalIndices[DiagFaceIdx1]]), FaceCornerToVert[CCWFromDiag2], 0, FaceUVOrientation);
 			}
 			else
 			{
@@ -320,10 +385,10 @@ protected:
 				}
 				else
 				{
-					MakeUVsAndNormals(DiagFaceIdx1, (FVector3f)Normal, FaceCornerToVert[DiagFaceIdx1]);
-					MakeUVsAndNormals(DiagFaceIdx2, (FVector3f)Normal, FaceCornerToVert[DiagFaceIdx2]);
+					MakeUVsAndNormals(DiagFaceIdx1, (FVector3f)Normal, FaceCornerToVert[DiagFaceIdx1], 0, FaceUVOrientation);
+					MakeUVsAndNormals(DiagFaceIdx2, (FVector3f)Normal, FaceCornerToVert[DiagFaceIdx2], 0, FaceUVOrientation);
 				}
-				MakeUVsAndNormals(CCWFromDiag2, (FVector3f)Normal, FaceCornerToVert[CCWFromDiag2]);
+				MakeUVsAndNormals(CCWFromDiag2, (FVector3f)Normal, FaceCornerToVert[CCWFromDiag2], 0, FaceUVOrientation);
 			}
 			MakeTriangle(DiagFaceIdx1, DiagFaceIdx2, CCWFromDiag2);
 		}// end triangulating second tri
@@ -384,48 +449,55 @@ void FCubeGridBooleanOp::CalculateResult(FProgressCancel* Progress)
 	}
 
 	TUniquePtr<FDynamicMesh3> OpMesh;
+
+	FWeldedMinimalBoxMeshGenerator Generator;
+	Generator.Box = WorldBox;
 	if (CornerInfo)
 	{
-		if (bSubtract)
+		Generator.CornerInfo = *CornerInfo;
+	}
+	Generator.bCrosswiseDiagonal = bCrosswiseDiagonal;
+	Generator.HeightUVOffset = OpMeshHeightUVOffset;
+	Generator.FaceUVOrientations = FaceUVOrientations;
+
+	OpMesh = MakeUnique<FDynamicMesh3>(&Generator.Generate());
+
+	if (ensure(OpMesh->HasAttributes()))
+	{
+		// Apply the given material ID to the operator mesh if needed.
+		OpMesh->Attributes()->EnableMaterialID();
+		FDynamicMeshMaterialAttribute* MaterialIDs = OpMesh->Attributes()->GetMaterialID();
+		for (int32 Tid : OpMesh->TriangleIndicesItr())
 		{
-			// Recall that for subtracting a welded box mesh, we pick a different 
-			// portion of the box volume as our operator.
-			// However, to be able to use the same code, we can mirror the box in
-			// such a way that the 0-3 indices are at the selection plane, and the
-			// z axis points in the direction of the push (because pushing a corner
-			// down is equivalent to raising it up and then mirroring it down).
-			// We mirror twice to keep the frame handedness, then adjust the
-			// welded flags to point to the correct verts.
-			WorldBox.Frame = FFrame3d(WorldBox.Center(),
-				-WorldBox.AxisX(),
-				WorldBox.AxisY(),
-				-WorldBox.AxisZ());
-			Swap(CornerInfo->WeldedAtBase[0], CornerInfo->WeldedAtBase[1]);
-			Swap(CornerInfo->WeldedAtBase[2], CornerInfo->WeldedAtBase[3]);
+			MaterialIDs->SetValue(Tid, OpMeshMaterialID);
 		}
 
-		FWeldedMinimalBoxMeshGenerator Generator;
-		Generator.Box = WorldBox;
-		Generator.CornerInfo = *CornerInfo;
-		Generator.bCrosswiseDiagonal = bCrosswiseDiagonal;
+		// Rescale UV's
+		FDynamicMeshEditor Editor(OpMesh.Get());
+		float WorldUnitsInMetersFactor = .01;
+		float UVScaleToUse = bWorldSpaceUVs ? WorldUnitsInMetersFactor * UVScale : UVScale;
+		Editor.RescaleAttributeUVs(UVScaleToUse, bWorldSpaceUVs);
 
-		OpMesh = MakeUnique<FDynamicMesh3>(&Generator.Generate());
+		// Unreal stores V's as 1-V, so we probably want to apply this here so it aligns with everything
+		// else. However this should happen after the scaling above.
+		FDynamicMeshUVOverlay* Overlay = OpMesh->Attributes()->PrimaryUV();
+		for (int32 ElementID : Overlay->ElementIndicesItr())
+		{
+			FVector2f UV = Overlay->GetElement(ElementID);
+			Overlay->SetElement(ElementID, FVector2f(UV.X, 1 - UV.Y));
+		}
 	}
-	else
-	{
-		// Create the box that will be added/subtracted
-		FGridBoxMeshGenerator BoxMeshGenerator;
-		BoxMeshGenerator.Box = WorldBox;
-		BoxMeshGenerator.EdgeVertices = FIndex3i(2, 2, 2);
-		BoxMeshGenerator.bPolygroupPerQuad = true;
 
-		OpMesh = MakeUnique<FDynamicMesh3>(&BoxMeshGenerator.Generate());
+	if (!ResultMesh->HasAttributes())
+	{
+		ResultMesh->EnableAttributes();
+	}
+	if (!ResultMesh->Attributes()->HasMaterialID())
+	{
+		// This usually happens when we start with an empty mesh. Materials need to
+		// be enabled or else we won't add in the OpMesh one during the boolean operation.
+		ResultMesh->Attributes()->EnableMaterialID();
 	}
-
-	// Rescale UV's to world space
-	FDynamicMeshEditor Editor(OpMesh.Get());
-	float WorldUnitsInMetersFactor = .01;
-	Editor.RescaleAttributeUVs(WorldUnitsInMetersFactor, true);
 
 	if (Progress && Progress->Cancelled())
 	{
@@ -439,6 +511,7 @@ void FCubeGridBooleanOp::CalculateResult(FProgressCancel* Progress)
 
 	MeshBoolean.bPutResultInInputSpace = true;
 	MeshBoolean.bSimplifyAlongNewEdges = true;
+	MeshBoolean.bPopulateSecondMeshGroupMap = true;
 	MeshBoolean.Compute();
 
 	if (Progress && Progress->Cancelled())
@@ -466,6 +539,47 @@ void FCubeGridBooleanOp::CalculateResult(FProgressCancel* Progress)
 		return;
 	}
 
+	// Prep to remap the side groups if necessary (or output the new side groups)
+	TMap<int32, int32> GroupRemappingToApply;
+	if (ensure(OpMeshSideGroups.Num() == 4))
+	{
+		for (int SideGroupIndex = 0; SideGroupIndex < 4; ++SideGroupIndex)
+		{
+			// We assign the sides groups 0-3 in the generator, but these get turned into 1-4
+			// in the conversion to DynamicMesh.
+			int32 OpMeshGroupID = SideGroupIndex + 1;
+			const int* ResultGroup = MeshBoolean.SecondMeshGroupMap.FindTo(OpMeshGroupID);
+			if (ResultGroup)
+			{
+				if (OpMeshSideGroups[SideGroupIndex] != IndexConstants::InvalidID)
+				{
+					GroupRemappingToApply.Add(*ResultGroup, OpMeshSideGroups[SideGroupIndex]);
+				}
+				else
+				{
+					OpMeshSideGroups[SideGroupIndex] = *ResultGroup;
+				}
+			}
+		}
+	}
+	// Do the actual group remapping
+	if (!GroupRemappingToApply.IsEmpty())
+	{
+		for (int32 Tid : ResultMesh->TriangleIndicesItr())
+		{
+			const int32* RemapGroup = GroupRemappingToApply.Find(ResultMesh->GetTriangleGroup(Tid));
+			if (RemapGroup)
+			{
+				ResultMesh->SetTriangleGroup(Tid, *RemapGroup);;
+			}
+		}
+
+		if (Progress && Progress->Cancelled())
+		{
+			return;
+		}
+	}
+
 	if (bTrackChangedTids)
 	{
 		ChangedTids = MakeShared<TArray<int32>, ESPMode::ThreadSafe>();
diff --git a/Engine/Plugins/Runtime/MeshModelingToolset/Source/ModelingOperators/Public/CompositionOps/CubeGridBooleanOp.h b/Engine/Plugins/Runtime/MeshModelingToolset/Source/ModelingOperators/Public/CompositionOps/CubeGridBooleanOp.h
--- a/Engine/Plugins/Runtime/MeshModelingToolset/Source/ModelingOperators/Public/CompositionOps/CubeGridBooleanOp.h
+++ b/Engine/Plugins/Runtime/MeshModelingToolset/Source/ModelingOperators/Public/CompositionOps/CubeGridBooleanOp.h
@@ -4,6 +4,7 @@
 
 #include "CoreMinimal.h"
 
+#include "IndexTypes.h"
 #include "ModelingOperators.h"
 #include "OrientedBoxTypes.h"
 
@@ -36,19 +37,13 @@ public:
 	struct FCornerInfo
 	{
 		// "Base" refers to the 0-3 indexed corners in FOrientedBox3d, and the flags here are
-		// true if the vertex is welded to its Z axis neighbor. Note that we have a choice of
-		// which of the two resulting sides of the box volume form our operator mesh. In the
-		// case of addition, we choose the lower side on the Z axis (i.e., non-welded verts
-		// stick up), while in subtraction, we choose the higher side on the Z axis (i.e.
-		// welded verts stick down and therefore cut away the input mesh). This lines up with
-		// how a user uses this operation- they select some corners, and either pull them out
-		// or push them in.
+		// true if the vertex is welded to its Z axis neighbor. The lower side on the Z axis 
+		// is the enclosed area for the purposes of an operation. 
 		bool WeldedAtBase[4] = { false, false, false, false };
 	};
 	TSharedPtr<FCornerInfo> CornerInfo;
 
-	// Determines whether the box is added or subtracted from the mesh. Also has an effect on
-	// how we treat operators with welded corners, see above.
+	// Determines whether the box is added or subtracted from the mesh.
 	bool bSubtract = false;
 
 	// When true, mesh is constructed such that ResultTransform is InputTransform. When false,
@@ -60,6 +55,41 @@ public:
 	// corner results in a pyramid with three faces (if true) or four (if false).
 	bool bCrosswiseDiagonal = false;
 
+	// The material ID to give to triangles in the operator mesh.
+	int32 OpMeshMaterialID = 0;
+
+	/** 
+	 * For each face of the op box (indexed as in IndexUtil::BoxFaces), determines the orientation
+	 * in which the UV's should be assigned, which can be imagined as number of counterclockwise rotations 
+	 * of the face in the UV plane while keeping the corner labels the same.
+	 * 
+	 * For instance, the corner UV's for orientation 0 are (0,0), (width, 0), (width, height), (0, height),
+	 * whereas for orientation 1 they are (height, 0), (height, width), (0, width), (0,0)
+	 * 
+	 * This is used to orient UV's in a uniform direction while extruding/pushing in different directions
+	 * in the cube grid.
+	 */
+	TArray<int32, TFixedAllocator<6>> FaceUVOrientations;
+
+	/**
+	 * When setting UV's on the box, the component affected by height is adjusted by this offset. This is used
+	 * to avoid "restarting" UV's when performing multiple operations with the same selection.
+	 */
+	double OpMeshHeightUVOffset = 0;
+
+	float UVScale = 1;
+	bool bWorldSpaceUVs = true;
+
+	// Both Input and Output:
+	// When not set to InvalidID, these will control the setting of the side groups in the operator
+	// mesh. This allows for the sides to keep the same groups with each step when selection is unchanged.
+	TArray<int32, TFixedAllocator<4>> OpMeshSideGroups = {
+		IndexConstants::InvalidID, 
+		IndexConstants::InvalidID, 
+		IndexConstants::InvalidID, 
+		IndexConstants::InvalidID
+	};
+
 	// Outputs:
 
 	// Will be reset to a container storing tids whose positions or connectivity changed. Only filled
