Multiple non-unity build fixes, mostly having to do with PDI.

#rnx

diff --git a/Engine/Plugins/Experimental/MeshModelingToolsetExp/Source/MeshModelingToolsExp/Private/RevolveBoundaryTool.cpp b/Engine/Plugins/Experimental/MeshModelingToolsetExp/Source/MeshModelingToolsExp/Private/RevolveBoundaryTool.cpp
--- a/Engine/Plugins/Experimental/MeshModelingToolsetExp/Source/MeshModelingToolsExp/Private/RevolveBoundaryTool.cpp
+++ b/Engine/Plugins/Experimental/MeshModelingToolsetExp/Source/MeshModelingToolsExp/Private/RevolveBoundaryTool.cpp
@@ -8,6 +8,7 @@
 #include "CompositionOps/CurveSweepOp.h"
 #include "InteractiveToolManager.h"
 #include "Mechanics/ConstructionPlaneMechanic.h"
+#include "SceneManagement.h" // FPrimitiveDrawInterface
 #include "Selection/PolygonSelectionMechanic.h"
 #include "GroupTopology.h"
 #include "ToolBuilderUtil.h"
diff --git a/Engine/Plugins/Runtime/MeshModelingToolset/Source/MeshModelingTools/Private/DrawAndRevolveTool.cpp b/Engine/Plugins/Runtime/MeshModelingToolset/Source/MeshModelingTools/Private/DrawAndRevolveTool.cpp
--- a/Engine/Plugins/Runtime/MeshModelingToolset/Source/MeshModelingTools/Private/DrawAndRevolveTool.cpp
+++ b/Engine/Plugins/Runtime/MeshModelingToolset/Source/MeshModelingTools/Private/DrawAndRevolveTool.cpp
@@ -11,6 +11,7 @@
 #include "Mechanics/CurveControlPointsMechanic.h"
 #include "Properties/MeshMaterialProperties.h"
 #include "ModelingObjectsCreationAPI.h"
+#include "SceneManagement.h" // FPrimitiveDrawInterface
 #include "Selection/ToolSelectionUtil.h"
 #include "ToolSceneQueriesUtil.h"
 #include "ToolSetupUtil.h"
diff --git a/Engine/Plugins/Runtime/MeshModelingToolset/Source/MeshModelingTools/Private/DrawPolygonTool.cpp b/Engine/Plugins/Runtime/MeshModelingToolset/Source/MeshModelingTools/Private/DrawPolygonTool.cpp
--- a/Engine/Plugins/Runtime/MeshModelingToolset/Source/MeshModelingTools/Private/DrawPolygonTool.cpp
+++ b/Engine/Plugins/Runtime/MeshModelingToolset/Source/MeshModelingTools/Private/DrawPolygonTool.cpp
@@ -20,6 +20,7 @@
 #include "Distance/DistLine3Ray3.h"
 #include "Intersection/IntrSegment2Segment2.h"
 #include "ToolSceneQueriesUtil.h"
+#include "SceneManagement.h" // FPrimitiveDrawInterface
 #include "SceneQueries/SceneSnappingManager.h"
 #include "ConstrainedDelaunay2.h"
 #include "Arrangement2d.h"
diff --git a/Engine/Source/Runtime/InteractiveToolsFramework/Private/BaseGizmos/GizmoBoxComponent.cpp b/Engine/Source/Runtime/InteractiveToolsFramework/Private/BaseGizmos/GizmoBoxComponent.cpp
--- a/Engine/Source/Runtime/InteractiveToolsFramework/Private/BaseGizmos/GizmoBoxComponent.cpp
+++ b/Engine/Source/Runtime/InteractiveToolsFramework/Private/BaseGizmos/GizmoBoxComponent.cpp
@@ -4,6 +4,7 @@
 #include "BaseGizmos/GizmoRenderingUtil.h"
 #include "BaseGizmos/GizmoViewContext.h"
 #include "PrimitiveSceneProxy.h"
+#include "SceneManagement.h" // FMeshElementCollector, FPrimitiveDrawInterface
 
 namespace GizmoBoxComponentLocals
 {
diff --git a/Engine/Source/Runtime/InteractiveToolsFramework/Private/BaseGizmos/GizmoCircleComponent.cpp b/Engine/Source/Runtime/InteractiveToolsFramework/Private/BaseGizmos/GizmoCircleComponent.cpp
--- a/Engine/Source/Runtime/InteractiveToolsFramework/Private/BaseGizmos/GizmoCircleComponent.cpp
+++ b/Engine/Source/Runtime/InteractiveToolsFramework/Private/BaseGizmos/GizmoCircleComponent.cpp
@@ -5,6 +5,7 @@
 #include "BaseGizmos/GizmoMath.h"
 #include "BaseGizmos/GizmoViewContext.h"
 #include "PrimitiveSceneProxy.h"
+#include "SceneManagement.h" // FMeshElementCollector, FPrimitiveDrawInterface
 
 namespace GizmoCircleComponentLocals
 {
diff --git a/Engine/Source/Runtime/InteractiveToolsFramework/Private/BaseGizmos/GizmoLineHandleComponent.cpp b/Engine/Source/Runtime/InteractiveToolsFramework/Private/BaseGizmos/GizmoLineHandleComponent.cpp
--- a/Engine/Source/Runtime/InteractiveToolsFramework/Private/BaseGizmos/GizmoLineHandleComponent.cpp
+++ b/Engine/Source/Runtime/InteractiveToolsFramework/Private/BaseGizmos/GizmoLineHandleComponent.cpp
@@ -4,7 +4,7 @@
 #include "BaseGizmos/GizmoRenderingUtil.h"
 #include "BaseGizmos/GizmoMath.h"
 #include "PrimitiveSceneProxy.h"
-
+#include "SceneManagement.h" // FMeshElementCollector, FPrimitiveDrawInterface
 
 
 
diff --git a/Engine/Source/Runtime/InteractiveToolsFramework/Private/BaseGizmos/GizmoRectangleComponent.cpp b/Engine/Source/Runtime/InteractiveToolsFramework/Private/BaseGizmos/GizmoRectangleComponent.cpp
--- a/Engine/Source/Runtime/InteractiveToolsFramework/Private/BaseGizmos/GizmoRectangleComponent.cpp
+++ b/Engine/Source/Runtime/InteractiveToolsFramework/Private/BaseGizmos/GizmoRectangleComponent.cpp
@@ -4,6 +4,7 @@
 #include "BaseGizmos/GizmoRenderingUtil.h"
 #include "BaseGizmos/GizmoViewContext.h"
 #include "PrimitiveSceneProxy.h"
+#include "SceneManagement.h" // FMeshElementCollector, FPrimitiveDrawInterface
 
 namespace GizmoRectangleComponentLocals
 {
diff --git a/Engine/Source/Runtime/Landscape/Classes/LandscapeComponent.h b/Engine/Source/Runtime/Landscape/Classes/LandscapeComponent.h
--- a/Engine/Source/Runtime/Landscape/Classes/LandscapeComponent.h
+++ b/Engine/Source/Runtime/Landscape/Classes/LandscapeComponent.h
@@ -8,6 +8,7 @@
 #include "Engine/TextureStreamingTypes.h"
 #include "Components/PrimitiveComponent.h"
 #include "PerPlatformProperties.h"
+#include "Serialization/BulkData.h"
 #include "LandscapePhysicalMaterial.h"
 #include "LandscapeWeightmapUsage.h"
 #include "Containers/ArrayView.h"
