Initial version of Render Capture Baking tool

#rnx

diff --git a/Engine/Plugins/Editor/ModelingToolsEditorMode/Source/ModelingToolsEditorMode/Private/ModelingToolsEditorMode.cpp b/Engine/Plugins/Editor/ModelingToolsEditorMode/Source/ModelingToolsEditorMode/Private/ModelingToolsEditorMode.cpp
--- a/Engine/Plugins/Editor/ModelingToolsEditorMode/Source/ModelingToolsEditorMode/Private/ModelingToolsEditorMode.cpp
+++ b/Engine/Plugins/Editor/ModelingToolsEditorMode/Source/ModelingToolsEditorMode/Private/ModelingToolsEditorMode.cpp
@@ -79,6 +79,7 @@
 #include "BakeMeshAttributeMapsTool.h"
 #include "BakeMultiMeshAttributeMapsTool.h"
 #include "BakeMeshAttributeVertexTool.h"
+#include "BakeRenderCaptureTool.h"
 #include "MeshAttributePaintTool.h"
 #include "ParameterizeMeshTool.h"
 #include "RecomputeUVsTool.h"
@@ -600,6 +601,9 @@ void UModelingToolsEditorMode::Enter()
 	auto BakeMultiMeshAttributeMapsToolBuilder = NewObject<UBakeMultiMeshAttributeMapsToolBuilder>();
 	RegisterTool(ToolManagerCommands.BeginBakeMultiMeshAttributeMapsTool, TEXT("BeginBakeMultiMeshAttributeMapsTool"), BakeMultiMeshAttributeMapsToolBuilder);
 
+	auto BakeRenderCaptureToolBuilder = NewObject<UBakeRenderCaptureToolBuilder>();
+	RegisterTool(ToolManagerCommands.BeginBakeRenderCaptureTool, TEXT("BeginBakeRenderCaptureTool"), BakeRenderCaptureToolBuilder);
+
 	auto BakeMeshAttributeVertexToolBuilder = NewObject<UBakeMeshAttributeVertexToolBuilder>();
 	RegisterTool(ToolManagerCommands.BeginBakeMeshAttributeVertexTool, TEXT("BeginBakeMeshAttributeVertexTool"), BakeMeshAttributeVertexToolBuilder);
 
diff --git a/Engine/Plugins/Editor/ModelingToolsEditorMode/Source/ModelingToolsEditorMode/Private/ModelingToolsEditorModeStyle.cpp b/Engine/Plugins/Editor/ModelingToolsEditorMode/Source/ModelingToolsEditorMode/Private/ModelingToolsEditorModeStyle.cpp
--- a/Engine/Plugins/Editor/ModelingToolsEditorMode/Source/ModelingToolsEditorMode/Private/ModelingToolsEditorModeStyle.cpp
+++ b/Engine/Plugins/Editor/ModelingToolsEditorMode/Source/ModelingToolsEditorMode/Private/ModelingToolsEditorModeStyle.cpp
@@ -175,6 +175,8 @@ void FModelingToolsEditorModeStyle::Initialize()
 		StyleSet->Set("ModelingToolsManagerCommands.BeginBakeMeshAttributeMapsTool.Small", 		new IMAGE_BRUSH_SVG("Icons/BakeTexture",			Icon20x20));
 		StyleSet->Set("ModelingToolsManagerCommands.BeginBakeMultiMeshAttributeMapsTool", 			new IMAGE_BRUSH_SVG("Icons/BakeAll",			Icon20x20));
 		StyleSet->Set("ModelingToolsManagerCommands.BeginBakeMultiMeshAttributeMapsTool.Small", 		new IMAGE_BRUSH_SVG("Icons/BakeAll",			Icon20x20));
+		StyleSet->Set("ModelingToolsManagerCommands.BeginBakeRenderCaptureTool", 			new IMAGE_BRUSH_SVG("Icons/BakeAll",			Icon20x20));
+		StyleSet->Set("ModelingToolsManagerCommands.BeginBakeRenderCaptureTool.Small", 		new IMAGE_BRUSH_SVG("Icons/BakeAll",			Icon20x20));
 		StyleSet->Set("ModelingToolsManagerCommands.BeginBakeMeshAttributeVertexTool", new IMAGE_BRUSH_SVG("Icons/BakeVertex", Icon20x20));
 		StyleSet->Set("ModelingToolsManagerCommands.BeginBakeMeshAttributeVertexTool.Small", new IMAGE_BRUSH_SVG("Icons/BakeVertex", Icon20x20));
 		//StyleSet->Set("ModelingToolsManagerCommands.BeginRemoveOccludedTrianglesTool", 				new IMAGE_PLUGIN_BRUSH("Icons/Jacket_40x",			Icon20x20));
diff --git a/Engine/Plugins/Editor/ModelingToolsEditorMode/Source/ModelingToolsEditorMode/Private/ModelingToolsEditorModeToolkit.cpp b/Engine/Plugins/Editor/ModelingToolsEditorMode/Source/ModelingToolsEditorMode/Private/ModelingToolsEditorModeToolkit.cpp
--- a/Engine/Plugins/Editor/ModelingToolsEditorMode/Source/ModelingToolsEditorMode/Private/ModelingToolsEditorModeToolkit.cpp
+++ b/Engine/Plugins/Editor/ModelingToolsEditorMode/Source/ModelingToolsEditorMode/Private/ModelingToolsEditorModeToolkit.cpp
@@ -669,6 +669,7 @@ void FModelingToolsEditorModeToolkit::BuildToolPalette(FName PaletteIndex, class
 		ToolbarBuilder.AddToolBarButton(Commands.BeginBakeMeshAttributeMapsTool);
 		ToolbarBuilder.AddToolBarButton(Commands.BeginBakeMultiMeshAttributeMapsTool);
 		ToolbarBuilder.AddToolBarButton(Commands.BeginBakeMeshAttributeVertexTool);
+		ToolbarBuilder.AddToolBarButton(Commands.BeginBakeRenderCaptureTool);
 	}
 	else if (PaletteIndex == UVTabName)
 	{
diff --git a/Engine/Plugins/Editor/ModelingToolsEditorMode/Source/ModelingToolsEditorMode/Private/ModelingToolsManagerActions.cpp b/Engine/Plugins/Editor/ModelingToolsEditorMode/Source/ModelingToolsEditorMode/Private/ModelingToolsManagerActions.cpp
--- a/Engine/Plugins/Editor/ModelingToolsEditorMode/Source/ModelingToolsEditorMode/Private/ModelingToolsManagerActions.cpp
+++ b/Engine/Plugins/Editor/ModelingToolsEditorMode/Source/ModelingToolsEditorMode/Private/ModelingToolsManagerActions.cpp
@@ -135,9 +135,10 @@ void FModelingToolsManagerCommands::RegisterCommands()
 	REGISTER_MODELING_TOOL_COMMAND(BeginUVLayoutTool, "Layout", "Transform and pack existing UVs");
 
 	// Baking
-	REGISTER_MODELING_TOOL_COMMAND(BeginBakeMeshAttributeMapsTool, "BakeTx", "Bake textures for single meshes");
-	REGISTER_MODELING_TOOL_COMMAND(BeginBakeMultiMeshAttributeMapsTool, "BakeAll", "Bake textures for single meshes from multiple source meshes");
-	REGISTER_MODELING_TOOL_COMMAND(BeginBakeMeshAttributeVertexTool, "BakeVtx", "Bake vertex colors for single meshes");
+	REGISTER_MODELING_TOOL_COMMAND(BeginBakeMeshAttributeMapsTool, "BakeTx", "Bake textures for a target mesh");
+	REGISTER_MODELING_TOOL_COMMAND(BeginBakeMultiMeshAttributeMapsTool, "BakeAll", "Bake textures for a target mesh from multiple source meshes");
+	REGISTER_MODELING_TOOL_COMMAND(BeginBakeRenderCaptureTool, "BakeRC", "Bake textures for a target mesh from multiple source meshes via virtual photo/render capture");
+	REGISTER_MODELING_TOOL_COMMAND(BeginBakeMeshAttributeVertexTool, "BakeVtx", "Bake vertex colors for a target mesh");
 
 	// Volumes
 	REGISTER_MODELING_TOOL_COMMAND(BeginVolumeToMeshTool, "Vol2Msh", "Convert volume to new mesh asset");
diff --git a/Engine/Plugins/Editor/ModelingToolsEditorMode/Source/ModelingToolsEditorMode/Public/ModelingToolsManagerActions.h b/Engine/Plugins/Editor/ModelingToolsEditorMode/Source/ModelingToolsEditorMode/Public/ModelingToolsManagerActions.h
--- a/Engine/Plugins/Editor/ModelingToolsEditorMode/Source/ModelingToolsEditorMode/Public/ModelingToolsManagerActions.h
+++ b/Engine/Plugins/Editor/ModelingToolsEditorMode/Source/ModelingToolsEditorMode/Public/ModelingToolsManagerActions.h
@@ -119,6 +119,7 @@ public:
 	TSharedPtr<FUICommandInfo> BeginAttributeEditorTool;
 	TSharedPtr<FUICommandInfo> BeginBakeMeshAttributeMapsTool;
 	TSharedPtr<FUICommandInfo> BeginBakeMultiMeshAttributeMapsTool;
+	TSharedPtr<FUICommandInfo> BeginBakeRenderCaptureTool;
 	TSharedPtr<FUICommandInfo> BeginBakeMeshAttributeVertexTool;
 	TSharedPtr<FUICommandInfo> BeginUVSeamEditTool;
 
diff --git a/Engine/Plugins/Experimental/MeshModelingToolsetExp/Source/MeshModelingToolsExp/Private/BakeRenderCaptureTool.cpp b/Engine/Plugins/Experimental/MeshModelingToolsetExp/Source/MeshModelingToolsExp/Private/BakeRenderCaptureTool.cpp
new file mode 100644
--- /dev/null
+++ b/Engine/Plugins/Experimental/MeshModelingToolsetExp/Source/MeshModelingToolsExp/Private/BakeRenderCaptureTool.cpp
@@ -0,0 +1,1111 @@
+﻿// Copyright Epic Games, Inc. All Rights Reserved.
+
+#include "BakeRenderCaptureTool.h"
+
+#include "TargetInterfaces/MaterialProvider.h"
+#include "TargetInterfaces/MeshDescriptionProvider.h"
+#include "TargetInterfaces/PrimitiveComponentBackedTarget.h"
+#include "TargetInterfaces/StaticMeshBackedTarget.h"
+#include "ToolTargetManager.h"
+
+#include "DynamicMesh/MeshTransforms.h"
+
+#include "ModelingToolTargetUtil.h"
+
+#include "ModelingObjectsCreationAPI.h"
+
+#include "Sampling/MeshGenericWorldPositionBaker.h"
+#include "Image/ImageInfilling.h"
+
+
+using namespace UE::Geometry;
+
+#define LOCTEXT_NAMESPACE "UBakeRenderCaptureTool"
+
+//
+// Implementation details
+//
+
+static FName BaseColorTexParamName = FName("BaseColor");
+static FName RoughnessTexParamName = FName("Roughness");
+static FName MetallicTexParamName = FName("Metallic");
+static FName SpecularTexParamName = FName("Specular");
+static FName EmissiveTexParamName = FName("Emissive");
+static FName NormalTexParamName = FName("NormalMap");
+static FName PackedMRSTexParamName = FName("PackedMRS");
+
+class FBakeRenderCaptureOptions
+{
+public:
+	enum class ETextureSizePolicy : uint8
+	{
+		TextureSize = 0,
+		TexelDensity = 1
+	};
+
+	/**
+	 * Input options to Actor Approximation process
+	 */
+	struct FOptions
+	{
+		// Base path for generated assets. Names will be generated by appending strings to this path.
+		FString BasePackagePath;
+
+		//
+		// Material approximation settings
+		//
+		int32 RenderCaptureImageSize = 1024;
+
+		// render capture parameters
+		double FieldOfViewDegrees = 45.0;
+		double NearPlaneDist = 1.0;
+
+		bool bMaximizeBakeParallelism = true;		// if true, photocapture is computed at the same time as mesh generation. Faster but requires more memory.
+
+		//
+		// Material output settings
+		//
+
+		// A new MIC derived from this material will be created and assigned to the generated mesh
+		UMaterialInterface* BakeMaterial = nullptr;		// if null, will use /MeshModelingToolsetExp/Materials/FullMaterialBakePreviewMaterial_PackedMRS instead
+		bool bBakeBaseColor = true;
+		bool bBakeRoughness = true;
+		bool bBakeMetallic = true;
+		bool bBakeSpecular = true;
+		bool bBakeEmissive = true;
+		bool bBakeNormalMap = true;
+		
+		bool bUsePackedMRS = true;
+
+		// output texture options
+		FBakeRenderCaptureOptions::ETextureSizePolicy TextureSizePolicy = FBakeRenderCaptureOptions::ETextureSizePolicy::TextureSize;
+		int32 TextureImageSize = 1024;
+		float MeshTexelDensity = 0;
+
+		// supersampling parameter
+		int32 AntiAliasMultiSampling = 0;
+
+
+		//
+		// Mesh settings
+		//
+
+		//  Which UV layer of the Target mesh (the one we're baking to) should be used
+		int32 TargetUVLayer = 0;
+
+		//
+		// Debug settings
+		//
+
+		// print useful information to the Output Log
+		bool bVerbose = false;
+
+		// create a flattened (ie non-instanced) mesh and save it with _DEBUG suffix. Warning often absolutely enormous!
+		bool bWriteDebugMesh = false;
+	};
+
+	/**
+	 * Construct an FOptions from the provided FMeshApproximationSettings.
+	 */
+	static FOptions ConstructOptions(
+		const UBakeRenderCaptureToolProperties& UseSettings,
+		const UBakeRenderCaptureInputToolProperties& InputMeshSettings)
+	{
+		//
+		// Construct options for ApproximateActors operation
+		//
+		FOptions Options;
+
+		Options.TargetUVLayer = InputMeshSettings.TargetUVLayerNamesList.IndexOfByKey(InputMeshSettings.TargetUVLayer);
+		
+		// Options.BasePackagePath // @Incomplete
+		Options.RenderCaptureImageSize = (UseSettings.RenderCaptureResolution == 0) ? Options.TextureImageSize : UseSettings.RenderCaptureResolution;
+		Options.FieldOfViewDegrees = UseSettings.CaptureFieldOfView;
+		Options.NearPlaneDist = UseSettings.NearPlaneDist;
+		//Options.bMaximizeBakeParallelism = true; // @Incomplete
+		
+		//Options.bBakeBaseColor // @Incomplete
+		Options.bBakeRoughness = UseSettings.MaterialSettings.bRoughnessMap;
+		Options.bBakeMetallic = UseSettings.MaterialSettings.bMetallicMap;
+		Options.bBakeSpecular = UseSettings.MaterialSettings.bSpecularMap;
+		Options.bBakeEmissive = UseSettings.MaterialSettings.bEmissiveMap;
+		Options.bBakeNormalMap = UseSettings.MaterialSettings.bNormalMap;
+		Options.bUsePackedMRS = UseSettings.MaterialSettings.bPackedMRSMap;
+
+		//Options.TextureSizePolicy // @Incomplete
+		Options.TextureImageSize = UseSettings.MaterialSettings.TextureSize;
+		//Options.MeshTexelDensity // @Incomplete
+		Options.AntiAliasMultiSampling = FMath::Max(1, UseSettings.MultiSamplingAA);
+
+		// @Incomplete
+		//Options.bGenerateNaniteEnabledMesh 
+		//Options.NaniteProxyTrianglePercent
+		//Options.bSupportRayTracing
+		//Options.bAllowDistanceField
+		//Options.bGenerateLightmapUVs
+		//Options.bCreatePhysicsBody
+		//Options.bVerbose
+		//Options.bWriteDebugMesh
+
+		return Options;
+	}
+};
+
+static int32 GetMeshTextureSizeFromTargetTexelDensity(const FDynamicMesh3& Mesh, float TargetTexelDensity)
+{
+	const FDynamicMeshUVOverlay* UVOverlay = Mesh.Attributes()->PrimaryUV();
+	double Mesh3DArea = 0;
+	double MeshUVArea = 0;
+	for (int TriangleID : Mesh.TriangleIndicesItr())
+	{
+		// World space area
+		Mesh3DArea += Mesh.GetTriArea(TriangleID);
+
+		FIndex3i UVVertices = UVOverlay->GetTriangle(TriangleID);
+		FTriangle2d TriangleUV = FTriangle2d(
+			(FVector2d)UVOverlay->GetElement(UVVertices.A),
+			(FVector2d)UVOverlay->GetElement(UVVertices.B),
+			(FVector2d)UVOverlay->GetElement(UVVertices.C));
+
+		// UV space area
+		MeshUVArea += TriangleUV.Area();
+	}
+	double TexelRatio = FMath::Sqrt(MeshUVArea / Mesh3DArea) * 100;
+
+	// Compute the perfect texture size that would get us to our texture density
+	// Also compute the nearest power of two sizes (below and above our target)
+	const int32 SizePerfect = FMath::CeilToInt(TargetTexelDensity / TexelRatio);
+	const int32 SizeHi = FMath::RoundUpToPowerOfTwo(SizePerfect);
+	const int32 SizeLo = SizeHi >> 1;
+
+	// Compute the texel density we achieve with these two texture sizes
+	const double TexelDensityLo = SizeLo * TexelRatio;
+	const double TexelDensityHi = SizeHi * TexelRatio;
+
+	// Select best match between low & high res textures.
+	const double TexelDensityLoDiff = TargetTexelDensity - TexelDensityLo;
+	const double TexelDensityHiDiff = TexelDensityHi - TargetTexelDensity;
+	const int32 BestTextureSize = TexelDensityLoDiff < TexelDensityHiDiff ? SizeLo : SizeHi;
+
+	return BestTextureSize;
+}
+
+
+static TUniquePtr<FSceneCapturePhotoSet> CapturePhotoSet(
+	const TArray<TObjectPtr<AActor>>& Actors,
+	const FBakeRenderCaptureOptions::FOptions& Options
+)
+{
+	TRACE_CPUPROFILER_EVENT_SCOPE(CapturePhotoSet);
+
+	double FieldOfView = Options.FieldOfViewDegrees;
+	double NearPlaneDist = Options.NearPlaneDist;
+
+	FImageDimensions CaptureDimensions(Options.RenderCaptureImageSize, Options.RenderCaptureImageSize);
+
+	TUniquePtr<FSceneCapturePhotoSet> SceneCapture = MakeUnique<FSceneCapturePhotoSet>();
+
+	SceneCapture->SetCaptureTypeEnabled(ERenderCaptureType::BaseColor, Options.bBakeBaseColor);
+	SceneCapture->SetCaptureTypeEnabled(ERenderCaptureType::WorldNormal, Options.bBakeNormalMap);
+	SceneCapture->SetCaptureTypeEnabled(ERenderCaptureType::Emissive, Options.bBakeEmissive);
+
+	bool bMetallic = Options.bBakeMetallic;
+	bool bRoughness = Options.bBakeRoughness;
+	bool bSpecular  = Options.bBakeSpecular;
+	// if (Options.bUsePackedMRS && (bMetallic || bRoughness || bSpecular ) )
+	if (Options.bUsePackedMRS)
+	{
+		SceneCapture->SetCaptureTypeEnabled(ERenderCaptureType::CombinedMRS, true);
+		SceneCapture->SetCaptureTypeEnabled(ERenderCaptureType::Roughness, false);
+		SceneCapture->SetCaptureTypeEnabled(ERenderCaptureType::Metallic, false);
+		SceneCapture->SetCaptureTypeEnabled(ERenderCaptureType::Specular, false);
+	}
+	else
+	{
+		SceneCapture->SetCaptureTypeEnabled(ERenderCaptureType::CombinedMRS, false);
+		SceneCapture->SetCaptureTypeEnabled(ERenderCaptureType::Roughness, bRoughness);
+		SceneCapture->SetCaptureTypeEnabled(ERenderCaptureType::Metallic, bMetallic);
+		SceneCapture->SetCaptureTypeEnabled(ERenderCaptureType::Specular, bSpecular);
+	}
+
+	SceneCapture->SetCaptureSceneActors(Actors[0]->GetWorld(), Actors);
+
+	// SceneCapture->SetEnableWriteDebugImages(true);
+
+	SceneCapture->AddStandardExteriorCapturesFromBoundingBox(
+		CaptureDimensions, FieldOfView, NearPlaneDist,
+		true, true, true);
+	
+	return SceneCapture;
+}
+
+static void ImageBuildersFromPhotoSet(
+	FSceneCapturePhotoSet* SceneCapture,
+	const FBakeRenderCaptureOptions::FOptions& Options, 
+	const FDynamicMesh3* WorldTargetMesh,
+	const FMeshTangentsd* MeshTangents,
+	TUniquePtr<FBakeRenderCaptureResultsBuilder>& Results)
+{
+	TRACE_CPUPROFILER_EVENT_SCOPE(ImageBuildersFromPhotoSet);
+
+	int32 UVLayer = Options.TargetUVLayer;
+	int32 Supersample = FMath::Max(1, Options.AntiAliasMultiSampling);
+	if ( (Options.TextureImageSize * Supersample) > 16384)
+	{
+		UE_LOG(LogGeometry, Warning, TEXT("Ignoring requested supersampling rate %d because it would require image buffers with resolution %d, please try lower value."), Supersample, Options.TextureImageSize * Supersample);
+		Supersample = 1;
+	}
+
+	FImageDimensions OutputDimensions(Options.TextureImageSize*Supersample, Options.TextureImageSize*Supersample);
+
+	//FScopedSlowTask Progress(8.f, LOCTEXT("BakingTextures", "Baking Textures..."));
+	//Progress.MakeDialog(true);
+	//Progress.EnterProgressFrame(1.f, LOCTEXT("BakingSetup", "Setup..."));
+
+	FDynamicMeshAABBTree3 Spatial(WorldTargetMesh, true);
+
+	FMeshImageBakingCache TempBakeCache;
+	{
+		TRACE_CPUPROFILER_EVENT_SCOPE(ImageBuildersFromPhotoSet_Textures_MakeCache);
+		TempBakeCache.SetDetailMesh(WorldTargetMesh, &Spatial);
+		TempBakeCache.SetBakeTargetMesh(WorldTargetMesh);
+		TempBakeCache.SetDimensions(OutputDimensions);
+		TempBakeCache.SetUVLayer(UVLayer);
+		TempBakeCache.SetThickness(0.1);
+		TempBakeCache.SetCorrespondenceStrategy(FMeshImageBakingCache::ECorrespondenceStrategy::Identity);
+		TempBakeCache.ValidateCache();
+	}
+
+	//Progress.EnterProgressFrame(1.f, LOCTEXT("BakingBaseColor", "Baking Base Color..."));
+
+	FAxisAlignedBox3d TargetBounds = WorldTargetMesh->GetBounds();
+	double RayOffsetHackDist = (double)(100.0f * FMathf::ZeroTolerance * TargetBounds.MinDim() );
+
+	auto VisibilityFunction = [&Spatial, RayOffsetHackDist](const FVector3d& SurfPos, const FVector3d& ImagePosWorld)
+	{
+		FVector3d RayDir = ImagePosWorld - SurfPos;
+		double Dist = Normalize(RayDir);
+		FVector3d RayOrigin = SurfPos + RayOffsetHackDist * RayDir;
+		int32 HitTID = Spatial.FindNearestHitTriangle(FRay3d(RayOrigin, RayDir), IMeshSpatial::FQueryOptions(Dist));
+		return (HitTID == IndexConstants::InvalidID);
+	};
+
+	FSceneCapturePhotoSet::FSceneSample DefaultSample;
+	FVector4f InvalidColor(0, -1, 0, 1);
+	DefaultSample.BaseColor = FVector3f(InvalidColor.X, InvalidColor.Y, InvalidColor.Z);
+
+	FMeshGenericWorldPositionColorBaker BaseColorBaker;
+	BaseColorBaker.SetCache(&TempBakeCache);
+	BaseColorBaker.ColorSampleFunction = [&](FVector3d Position, FVector3d Normal) {
+		FSceneCapturePhotoSet::FSceneSample Sample = DefaultSample;
+		SceneCapture->ComputeSample(FRenderCaptureTypeFlags::BaseColor(),
+			Position, Normal, VisibilityFunction, Sample);
+		return Sample.GetValue4f(ERenderCaptureType::BaseColor);
+	};
+	{
+		TRACE_CPUPROFILER_EVENT_SCOPE(ImageBuildersFromPhotoSet_Textures_BakeColor);
+		BaseColorBaker.Bake();
+	}
+
+	// find "hole" pixels
+	TArray<FVector2i> MissingPixels;
+	Results->ColorImage = BaseColorBaker.TakeResult();
+	TMarchingPixelInfill<FVector4f> Infill;
+
+	{
+		TRACE_CPUPROFILER_EVENT_SCOPE(ImageBuildersFromPhotoSet_Textures_ComputeInfill);
+		TempBakeCache.FindSamplingHoles([&](const FVector2i& Coords)
+		{
+			return Results->ColorImage->GetPixel(Coords) == InvalidColor;
+		}, MissingPixels);
+
+		// solve infill for the holes while also caching infill information
+		Infill.ComputeInfill(*Results->ColorImage, MissingPixels, InvalidColor,
+			[](FVector4f SumValue, int32 Count) {
+			float InvSum = (Count == 0) ? 1.0f : (1.0f / Count);
+			return FVector4f(SumValue.X * InvSum, SumValue.Y * InvSum, SumValue.Z * InvSum, 1.0f);
+		});
+	}
+
+	// downsample the image if necessary
+	if (Supersample > 1)
+	{
+		TRACE_CPUPROFILER_EVENT_SCOPE(ImageBuildersFromPhotoSet_Textures_Downsample);
+		TImageBuilder<FVector4f> Downsampled = Results->ColorImage->FastDownsample(Supersample, FVector4f::Zero(), [](FVector4f V, int N) { return V / (float)N; });
+		*Results->ColorImage = MoveTemp(Downsampled);
+	}
+
+	// this lambda is used to process the per-channel images. It does the bake, applies infill, and downsamples if necessary
+	auto ProcessChannelFunc = [&](ERenderCaptureType CaptureType)
+	{
+		FVector4f DefaultValue(0, 0, 0, 0);
+		FMeshGenericWorldPositionColorBaker ChannelBaker;
+		ChannelBaker.SetCache(&TempBakeCache);
+		ChannelBaker.ColorSampleFunction = [&](FVector3d Position, FVector3d Normal) {
+			FSceneCapturePhotoSet::FSceneSample Sample = DefaultSample;
+			SceneCapture->ComputeSample(FRenderCaptureTypeFlags::Single(CaptureType), Position, Normal, VisibilityFunction, Sample);
+			return Sample.GetValue4f(CaptureType);
+		};
+		ChannelBaker.Bake();
+		TUniquePtr<TImageBuilder<FVector4f>> Image = ChannelBaker.TakeResult();
+
+		Infill.ApplyInfill(*Image,
+			[](FVector4f SumValue, int32 Count) {
+			float InvSum = (Count == 0) ? 1.0f : (1.0f / Count);
+			return FVector4f(SumValue.X * InvSum, SumValue.Y * InvSum, SumValue.Z * InvSum, 1.0f);
+		});
+
+		if (Supersample > 1)
+		{
+			TImageBuilder<FVector4f> Downsampled = Image->FastDownsample(Supersample, FVector4f::Zero(), [](FVector4f V, int N) { return V / (float)N; });
+			*Image = MoveTemp(Downsampled);
+		}
+
+		return MoveTemp(Image);
+	};
+
+	bool bMetallic = Options.bBakeMetallic;
+	bool bRoughness = Options.bBakeRoughness;
+	bool bSpecular = Options.bBakeSpecular;
+	{
+		TRACE_CPUPROFILER_EVENT_SCOPE(ImageBuildersFromPhotoSet_Textures_OtherChannels);
+
+		// if (Options.bUsePackedMRS && (bMetallic || bRoughness || bSpecular))
+		if (Options.bUsePackedMRS)
+		{
+			Results->PackedMRSImage = ProcessChannelFunc(ERenderCaptureType::CombinedMRS);
+		}
+		else
+		{
+			if (bRoughness)
+			{
+				Results->RoughnessImage = ProcessChannelFunc(ERenderCaptureType::Roughness);
+			}
+			if (bMetallic)
+			{
+				Results->MetallicImage = ProcessChannelFunc(ERenderCaptureType::Metallic);
+			}
+			if (bSpecular)
+			{
+				Results->SpecularImage = ProcessChannelFunc(ERenderCaptureType::Specular);
+			}
+		}
+
+		if (Options.bBakeEmissive)
+		{
+			Results->EmissiveImage = ProcessChannelFunc(ERenderCaptureType::Emissive);
+		}
+	}
+
+
+	//Progress.EnterProgressFrame(1.f, LOCTEXT("BakingNormals", "Baking Normals..."));
+
+	if (Options.bBakeNormalMap)
+	{
+		TRACE_CPUPROFILER_EVENT_SCOPE(ImageBuildersFromPhotoSet_Textures_NormalMapBake);
+
+		// no infill on normal map for now, doesn't make sense to do after mapping to tangent space!
+		//  (should we build baked normal map in world space, and then resample to tangent space??)
+		FVector4f DefaultNormalValue(0, 0, 1, 1);
+		FMeshGenericWorldPositionNormalBaker NormalMapBaker;
+		NormalMapBaker.SetCache(&TempBakeCache);
+		NormalMapBaker.BaseMeshTangents = MeshTangents;
+		NormalMapBaker.NormalSampleFunction = [&](FVector3d Position, FVector3d Normal) {
+			FSceneCapturePhotoSet::FSceneSample Sample = DefaultSample;
+			SceneCapture->ComputeSample(FRenderCaptureTypeFlags::WorldNormal(),
+				Position, Normal, VisibilityFunction, Sample);
+			FVector3f NormalColor = Sample.WorldNormal;
+			float x = (NormalColor.X - 0.5f) * 2.0f;
+			float y = (NormalColor.Y - 0.5f) * 2.0f;
+			float z = (NormalColor.Z - 0.5f) * 2.0f;
+			return FVector3f(x, y, z);
+		};
+
+		NormalMapBaker.Bake();
+		Results->NormalImage = NormalMapBaker.TakeResult();
+
+		if (Supersample > 1)
+		{
+			TImageBuilder<FVector3f> Downsampled = Results->NormalImage->FastDownsample(Supersample, FVector3f::Zero(), [](FVector3f V, int N) { return V / (float)N; });
+			*Results->NormalImage = MoveTemp(Downsampled);
+		}
+	}
+
+	// build textures
+	//Progress.EnterProgressFrame(1.f, LOCTEXT("BuildingTextures", "Building Textures..."));
+}
+
+//
+// Tool Operators
+//
+
+
+class FRenderCaptureMapBakerOp : public TGenericDataOperator<FBakeRenderCaptureResultsBuilder>
+{
+public:
+	TArray<TObjectPtr<AActor>>* Actors;
+	UE::Geometry::FDynamicMesh3* BaseMesh = nullptr;
+	TSharedPtr<UE::Geometry::FMeshTangentsd, ESPMode::ThreadSafe> BaseMeshTangents;
+	FBakeRenderCaptureOptions::FOptions Options;
+	TObjectPtr<UBakeRenderCaptureInputToolProperties> InputMeshSettings;
+	FSceneCapturePhotoSet* SceneCapture;
+	
+	// Begin TGenericDataOperator interface
+	// @Incomplete Use the Progress thing
+	virtual void CalculateResult(FProgressCancel* Progress) override
+	{
+		check(Actors != nullptr);
+		check(SceneCapture != nullptr);
+
+
+		FBakeRenderCaptureOptions::FOptions OverridenOptions = Options;
+
+		// evaluate required texture size if needed
+		if (Options.TextureSizePolicy == FBakeRenderCaptureOptions::ETextureSizePolicy::TexelDensity)
+		{
+			const int32 MaxTextureSize = 8192;
+			const int32 BestTextureSize = GetMeshTextureSizeFromTargetTexelDensity(*BaseMesh, Options.MeshTexelDensity);
+
+			if (BestTextureSize > MaxTextureSize)
+			{
+				UE_LOG(LogGeometry, Warning, TEXT("Mesh would require %dx%d textures, clamping down to maximum (%dx%d)"), BestTextureSize, BestTextureSize, MaxTextureSize, MaxTextureSize);
+				OverridenOptions.TextureImageSize = MaxTextureSize;
+			}
+			else
+			{
+				OverridenOptions.TextureImageSize = BestTextureSize;
+			}
+		}
+
+		// bake textures onto the target mesh by projecting/sampling the set of captured photos
+		ImageBuildersFromPhotoSet(SceneCapture, OverridenOptions, BaseMesh, BaseMeshTangents.Get(), Result);
+	}
+	// End TGenericDataOperator interface
+};
+
+
+//
+// Tool Builder
+//
+
+
+
+const FToolTargetTypeRequirements& UBakeRenderCaptureToolBuilder::GetTargetRequirements() const
+{
+	static FToolTargetTypeRequirements TypeRequirements({
+		UMeshDescriptionProvider::StaticClass(),
+		UPrimitiveComponentBackedTarget::StaticClass(),
+		UStaticMeshBackedTarget::StaticClass(),			// FMeshSceneAdapter currently only supports StaticMesh targets
+		UMaterialProvider::StaticClass()
+		});
+	return TypeRequirements;
+}
+
+bool UBakeRenderCaptureToolBuilder::CanBuildTool(const FToolBuilderState& SceneState) const
+{
+	const int32 NumTargets = SceneState.TargetManager->CountSelectedAndTargetable(SceneState, GetTargetRequirements());
+	return (NumTargets > 1);
+}
+
+UMultiSelectionMeshEditingTool* UBakeRenderCaptureToolBuilder::CreateNewTool(const FToolBuilderState& SceneState) const
+{
+	return NewObject<UBakeRenderCaptureTool>(SceneState.ToolManager);
+}
+
+
+
+//
+// Tool
+//
+
+
+
+void UBakeRenderCaptureTool::Setup()
+{
+	TRACE_CPUPROFILER_EVENT_SCOPE(UBakeRenderCaptureTool::Setup);
+
+	Super::Setup();
+
+	InitializePreviewMaterials();
+
+	// Initialize base mesh
+	const FTransformSRT3d BaseToWorld = UE::ToolTarget::GetLocalToWorldTransform(Targets[0]);
+	PreviewMesh->ProcessMesh([this, BaseToWorld](const FDynamicMesh3& Mesh)
+	{
+		TargetMesh.Copy(Mesh);
+		TargetMeshTangents = MakeShared<FMeshTangentsd, ESPMode::ThreadSafe>(&TargetMesh);
+		TargetMeshTangents->CopyTriVertexTangents(Mesh);
+		
+		// FMeshSceneAdapter operates in world space, so ensure our mesh transformed to world.
+		MeshTransforms::ApplyTransform(TargetMesh, BaseToWorld);
+		TargetSpatial.SetMesh(&TargetMesh, true);
+	});
+
+	// Initialize actors
+	const int NumTargets = Targets.Num();
+	Actors.Empty(NumTargets - 1);
+	for (int Idx = 1; Idx < NumTargets; ++Idx)
+	{
+		if (AActor* Actor = UE::ToolTarget::GetTargetActor(Targets[Idx]))
+		{
+			Actors.Add(Actor);
+		}
+	}
+
+	UToolTarget* Target = Targets[0];
+
+	// Setup tool property sets
+
+	Settings = NewObject<UBakeRenderCaptureToolProperties>(this);
+	Settings->RestoreProperties(this);
+	AddToolPropertySource(Settings);
+
+	Settings->MapPreview = BaseColorTexParamName.ToString(); // We always bake the base color
+	Settings->WatchProperty(Settings->MapPreview, [this](FString) { UpdateVisualization(); GetToolManager()->PostInvalidation(); });
+	Settings->WatchProperty(Settings->MultiSamplingAA, [this](int32) { OpState |= EBakeOpState::Evaluate; });
+	Settings->WatchProperty(Settings->RenderCaptureResolution, [this](int32) {  OpState |= EBakeOpState::Evaluate; });
+	Settings->WatchProperty(Settings->MaterialSettings, [this](FMaterialProxySettingsRC) { OpState |= EBakeOpState::Evaluate; });
+	Settings->WatchProperty(Settings->CaptureFieldOfView, [this](float) { OpState |= EBakeOpState::Evaluate; });
+	Settings->WatchProperty(Settings->NearPlaneDist, [this](float) { OpState |= EBakeOpState::Evaluate; });
+	Settings->WatchProperty(Settings->bPrintDebugMessages, [this](bool) { OpState |= EBakeOpState::Evaluate; }); // @Incomplete Actually make this do something or remove it
+	
+	InputMeshSettings = NewObject<UBakeRenderCaptureInputToolProperties>(this);
+	InputMeshSettings->RestoreProperties(this);
+	AddToolPropertySource(InputMeshSettings);
+	InputMeshSettings->TargetStaticMesh = GetStaticMeshTarget(Target);
+	UpdateUVLayerNames(InputMeshSettings->TargetUVLayer, InputMeshSettings->TargetUVLayerNamesList, TargetMesh);
+	InputMeshSettings->WatchProperty(InputMeshSettings->TargetUVLayer, [this](FString) { OpState |= EBakeOpState::Evaluate; });
+	
+	{
+		FBakeRenderCaptureOptions::FOptions Options = FBakeRenderCaptureOptions::ConstructOptions(*Settings, *InputMeshSettings);
+		Settings->MapPreviewNamesList.Add(NormalTexParamName.ToString());
+		Settings->MapPreviewNamesList.Add(BaseColorTexParamName.ToString());
+		Settings->MapPreviewNamesList.Add(RoughnessTexParamName.ToString());
+		Settings->MapPreviewNamesList.Add(MetallicTexParamName.ToString());
+		Settings->MapPreviewNamesList.Add(SpecularTexParamName.ToString());
+		Settings->MapPreviewNamesList.Add(EmissiveTexParamName.ToString());
+		Settings->MapPreviewNamesList.Add(PackedMRSTexParamName.ToString());
+	}
+
+	ResultSettings = NewObject<UBakeRenderCaptureResults>(this);
+	ResultSettings->RestoreProperties(this);
+	AddToolPropertySource(ResultSettings);
+	SetToolPropertySourceEnabled(ResultSettings, true);
+
+	// Hide the render capture meshes since this baker operates solely in world space which will occlude the preview of
+	// the target mesh.
+	for (int Idx = 1; Idx < NumTargets; ++Idx)
+	{
+		UE::ToolTarget::HideSourceObject(Targets[Idx]);
+	}
+
+	OpState |= EBakeOpState::Evaluate;
+
+	SetToolDisplayName(LOCTEXT("ToolName", "Bake Render Capture"));
+	GetToolManager()->DisplayMessage(
+		LOCTEXT("OnStartTool", "Bake Render Capture. Select Bake Mesh (LowPoly) first, then select Detail Meshes (HiPoly) to bake. Assets will be created on Accept."),
+		EToolMessageLevel::UserNotification);
+
+	PostSetup();
+}
+
+
+void UBakeRenderCaptureTool::OnShutdown(EToolShutdownType ShutdownType)
+{
+	TRACE_CPUPROFILER_EVENT_SCOPE(UBakeRenderCaptureTool::Shutdown);
+	
+	Super::OnShutdown(ShutdownType);
+	
+	Settings->SaveProperties(this);
+	InputMeshSettings->SaveProperties(this);
+
+	if (ComputeRC)
+	{
+		ComputeRC->Shutdown();
+	}
+
+	// Restore visibility of source meshes
+	const int NumTargets = Targets.Num();
+	for (int Idx = 1; Idx < NumTargets; ++Idx)
+	{
+		UE::ToolTarget::ShowSourceObject(Targets[Idx]);
+	}
+	
+	if (ShutdownType == EToolShutdownType::Accept)
+	{
+		IStaticMeshBackedTarget* StaticMeshTarget = Cast<IStaticMeshBackedTarget>(Targets[0]);
+		UObject* SourceAsset = StaticMeshTarget ? StaticMeshTarget->GetStaticMesh() : nullptr;
+		const UPrimitiveComponent* SourceComponent = UE::ToolTarget::GetTargetComponent(Targets[0]);
+		CreateTextureAssetsRC(SourceComponent->GetWorld(), SourceAsset);
+	}
+
+	// Clear actors on shutdown so that their lifetime is not tied to the lifetime of the tool
+	Actors.Empty();
+}
+
+void UBakeRenderCaptureTool::CreateTextureAssetsRC(UWorld* SourceWorld, UObject* SourceAsset)
+{
+	bool bCreatedAssetOK = true;
+	const FString BaseName = UE::ToolTarget::GetTargetActor(Targets[0])->GetActorNameOrLabel();
+
+	auto CreateTextureAsset = [this, &bCreatedAssetOK, &SourceWorld, &SourceAsset](const FString& TexName, FTexture2DBuilder::ETextureType Type, TObjectPtr<UTexture2D> Tex)
+	{
+		// See :DeferredPopulateSourceData
+		FTexture2DBuilder::CopyPlatformDataToSourceData(Tex, Type);
+
+		// TODO The original implementation in ApproximateActors also did the following, see WriteTextureLambda in ApproximateActorsImpl.cpp
+		//if (Type == FTexture2DBuilder::ETextureType::Roughness
+		//	|| Type == FTexture2DBuilder::ETextureType::Metallic
+		//	|| Type == FTexture2DBuilder::ETextureType::Specular)
+		//{
+		//	UE::AssetUtils::ConvertToSingleChannel(Texture);
+		//}
+
+		bCreatedAssetOK = bCreatedAssetOK &&
+			UE::Modeling::CreateTextureObject(
+				GetToolManager(),
+				FCreateTextureObjectParams{ 0, SourceWorld, SourceAsset, TexName, Tex }).IsOK();
+	};
+
+	if (ResultSettings->BaseColorMap != nullptr)
+	{
+		const FString TexName = FString::Printf(TEXT("%s_%s"), *BaseName, *BaseColorTexParamName.ToString());
+		CreateTextureAsset(TexName, FTexture2DBuilder::ETextureType::Color, ResultSettings->BaseColorMap);
+	}
+
+	if (Settings->MaterialSettings.bNormalMap && ResultSettings->NormalMap != nullptr)
+	{
+		const FString TexName = FString::Printf(TEXT("%s_%s"), *BaseName, *NormalTexParamName.ToString()); 
+		// Use ColorLinear because the baked NormalMap is in world space
+		CreateTextureAsset(TexName, FTexture2DBuilder::ETextureType::ColorLinear, ResultSettings->NormalMap);
+	}
+
+	if (Settings->MaterialSettings.bEmissiveMap && ResultSettings->EmissiveMap != nullptr)
+	{
+		const FString TexName = FString::Printf(TEXT("%s_%s"), *BaseName, *EmissiveTexParamName.ToString());
+		CreateTextureAsset(TexName, FTexture2DBuilder::ETextureType::EmissiveHDR, ResultSettings->EmissiveMap);
+	}
+
+	if (Settings->MaterialSettings.bPackedMRSMap && ResultSettings->PackedMRSMap != nullptr)
+	{
+		const FString TexName = FString::Printf(TEXT("%s_%s"), *BaseName, *PackedMRSTexParamName.ToString());
+		CreateTextureAsset(TexName, FTexture2DBuilder::ETextureType::ColorLinear, ResultSettings->PackedMRSMap);
+	}
+	else
+	{
+		if (Settings->MaterialSettings.bMetallicMap && ResultSettings->MetallicMap != nullptr)
+		{
+			const FString TexName = FString::Printf(TEXT("%s_%s"), *BaseName, *MetallicTexParamName.ToString());
+			CreateTextureAsset(TexName, FTexture2DBuilder::ETextureType::Metallic, ResultSettings->MetallicMap);
+		}
+
+		if (Settings->MaterialSettings.bRoughnessMap && ResultSettings->RoughnessMap != nullptr)
+		{
+			const FString TexName = FString::Printf(TEXT("%s_%s"), *BaseName, *RoughnessTexParamName.ToString());
+			CreateTextureAsset(TexName, FTexture2DBuilder::ETextureType::Roughness, ResultSettings->RoughnessMap);
+		}
+
+		if (Settings->MaterialSettings.bSpecularMap && ResultSettings->SpecularMap != nullptr)
+		{
+			const FString TexName = FString::Printf(TEXT("%s_%s"), *BaseName, *SpecularTexParamName.ToString());
+			CreateTextureAsset(TexName, FTexture2DBuilder::ETextureType::Specular, ResultSettings->SpecularMap);
+		}
+	}
+
+	ensure(bCreatedAssetOK);
+
+	//RecordAnalytics(BakeAnalytics, GetAnalyticsEventName());
+}
+
+
+void UBakeRenderCaptureTool::OnTick(float DeltaTime)
+{
+	if (ComputeRC)
+	{
+		ComputeRC->Tick(DeltaTime);
+
+		const float ElapsedComputeTime = ComputeRC->GetElapsedComputeTime();
+		if (!CanAccept() && ElapsedComputeTime > SecondsBeforeWorkingMaterial)
+		{
+			UMaterialInstanceDynamic* ProgressMaterial =
+				static_cast<bool>(OpState & EBakeOpState::Invalid) ? ErrorPreviewMaterial : WorkingPreviewMaterial;
+			PreviewMesh->SetOverrideRenderMaterial(ProgressMaterial);
+		}
+	}
+}
+
+
+bool UBakeRenderCaptureTool::CanAccept() const
+{
+	if ((OpState & EBakeOpState::Invalid) == EBakeOpState::Invalid)
+	{
+		return false;
+	}
+
+	if (ResultSettings->BaseColorMap == nullptr)
+	{
+		return false;
+	}
+	if (Settings->MaterialSettings.bNormalMap && ResultSettings->NormalMap == nullptr)
+	{
+		return false;
+	}
+	if (Settings->MaterialSettings.bEmissiveMap && ResultSettings->EmissiveMap == nullptr)
+	{
+		return false;
+	}
+	if (Settings->MaterialSettings.bPackedMRSMap)
+	{
+		if (ResultSettings->PackedMRSMap == nullptr)
+		{
+			return false;
+		}
+	}
+	else
+	{
+		if (Settings->MaterialSettings.bMetallicMap && ResultSettings->MetallicMap == nullptr)
+		{
+			return false;
+		}
+		if (Settings->MaterialSettings.bRoughnessMap && ResultSettings->RoughnessMap == nullptr)
+		{
+			return false;
+		}
+		if (Settings->MaterialSettings.bSpecularMap && ResultSettings->SpecularMap == nullptr)
+		{
+			return false;
+		}
+	}
+
+	return true;
+}
+
+
+TUniquePtr<UE::Geometry::TGenericDataOperator<FBakeRenderCaptureResultsBuilder>> UBakeRenderCaptureTool::FComputeFactory::MakeNewOperator()
+{
+	TUniquePtr<FRenderCaptureMapBakerOp> Op = MakeUnique<FRenderCaptureMapBakerOp>();
+	Op->Actors = &Tool->Actors;
+	Op->BaseMesh = &Tool->TargetMesh;
+	Op->BaseMeshTangents = Tool->TargetMeshTangents;
+	Op->Options = FBakeRenderCaptureOptions::ConstructOptions(*Tool->Settings, *Tool->InputMeshSettings);
+	Op->SceneCapture = Tool->SceneCapture.Get();
+	return Op;
+}
+
+
+void UBakeRenderCaptureTool::OnMapsUpdatedRC(const TUniquePtr<FBakeRenderCaptureResultsBuilder>& NewResult)
+{
+	FBakeRenderCaptureOptions::FOptions Options = FBakeRenderCaptureOptions::ConstructOptions(*Settings, *InputMeshSettings);
+	bool bMetallic = Options.bBakeMetallic;
+	bool bRoughness = Options.bBakeRoughness;
+	bool bSpecular = Options.bBakeSpecular;
+
+	check(IsInGameThread());
+
+	// We do this to defer work I guess, it was like this in the original ApproximateActors implementation :DeferredPopulateSourceData 
+	constexpr bool bPopulateSourceData = false;
+
+	{
+		TRACE_CPUPROFILER_EVENT_SCOPE(BakeRenderCaptureTool_Textures_BuildTextures);
+
+		if (Options.bBakeBaseColor && NewResult->ColorImage.IsValid())
+		{
+			ResultSettings->BaseColorMap = FTexture2DBuilder::BuildTextureFromImage(*NewResult->ColorImage, FTexture2DBuilder::ETextureType::Color, true, bPopulateSourceData);
+		}
+		if (Options.bBakeEmissive && NewResult->EmissiveImage.IsValid())
+		{
+			ResultSettings->EmissiveMap = FTexture2DBuilder::BuildTextureFromImage(*NewResult->EmissiveImage, FTexture2DBuilder::ETextureType::EmissiveHDR, false, bPopulateSourceData);
+			ResultSettings->EmissiveMap->CompressionSettings = TC_HDR_Compressed;
+		}
+		if (Options.bBakeNormalMap && NewResult->NormalImage.IsValid())
+		{
+			// Use ColorLinear because the baked NormalMap is in world space
+			ResultSettings->NormalMap = FTexture2DBuilder::BuildTextureFromImage(*NewResult->NormalImage, FTexture2DBuilder::ETextureType::ColorLinear, false, bPopulateSourceData);
+		}
+
+		// if (Options.bUsePackedMRS && (bRoughness || bMetallic || bSpecular) && NewResult->PackedMRSImage.IsValid())
+		if (Options.bUsePackedMRS && NewResult->PackedMRSImage.IsValid())
+		{
+			ResultSettings->PackedMRSMap = FTexture2DBuilder::BuildTextureFromImage(*NewResult->PackedMRSImage, FTexture2DBuilder::ETextureType::ColorLinear, false, bPopulateSourceData);
+		}
+		else
+		{ 
+			if (bRoughness && NewResult->RoughnessImage.IsValid())
+			{
+				ResultSettings->RoughnessMap = FTexture2DBuilder::BuildTextureFromImage(*NewResult->RoughnessImage, FTexture2DBuilder::ETextureType::Roughness, false, bPopulateSourceData);
+			}
+			if (bMetallic && NewResult->MetallicImage.IsValid())
+			{
+				ResultSettings->MetallicMap = FTexture2DBuilder::BuildTextureFromImage(*NewResult->MetallicImage, FTexture2DBuilder::ETextureType::Metallic, false, bPopulateSourceData);
+			}
+			if (bSpecular && NewResult->SpecularImage.IsValid())
+			{
+				ResultSettings->SpecularMap = FTexture2DBuilder::BuildTextureFromImage(*NewResult->SpecularImage, FTexture2DBuilder::ETextureType::Specular, false, bPopulateSourceData);
+			}
+		}
+	}
+
+	//GatherAnalytics(*NewResult, CachedBakeSettings, BakeAnalytics);
+	UpdateVisualization();
+	GetToolManager()->PostInvalidation();
+}
+
+
+void UBakeRenderCaptureTool::InitializePreviewMaterials()
+{
+	// EmptyColorMapWhite, EmptyColorMapBlack and EmptyNormalMap are defined in the base tool
+
+	{
+		FTexture2DBuilder Builder;
+		Builder.Initialize(FTexture2DBuilder::ETextureType::EmissiveHDR, FImageDimensions(16, 16));
+		Builder.Commit(false);
+		EmptyEmissiveMap = Builder.GetTexture2D();
+	}
+
+	{
+		FTexture2DBuilder Builder;
+		Builder.Initialize(FTexture2DBuilder::ETextureType::ColorLinear, FImageDimensions(16, 16));
+		Builder.Clear(FColor(0,0,0));
+		Builder.Commit(false);
+		EmptyPackedMRSMap = Builder.GetTexture2D();
+	}
+
+	{
+		FTexture2DBuilder Builder;
+		Builder.Initialize(FTexture2DBuilder::ETextureType::Roughness, FImageDimensions(16, 16));
+		Builder.Commit(false);
+		EmptyRoughnessMap = Builder.GetTexture2D();
+	}
+
+	{
+		FTexture2DBuilder Builder;
+		Builder.Initialize(FTexture2DBuilder::ETextureType::Metallic, FImageDimensions(16, 16));
+		Builder.Commit(false);
+		EmptyMetallicMap = Builder.GetTexture2D();
+	}
+
+	{
+		FTexture2DBuilder Builder;
+		Builder.Initialize(FTexture2DBuilder::ETextureType::Specular, FImageDimensions(16, 16));
+		Builder.Commit(false);
+		EmptySpecularMap = Builder.GetTexture2D();
+	}
+	
+	{
+		UMaterial* Material = LoadObject<UMaterial>(nullptr, TEXT("/MeshModelingToolsetExp/Materials/BakeRenderCapturePreviewMaterial"));
+		check(Material);
+		if (Material != nullptr)
+		{
+			PreviewMaterialRC = UMaterialInstanceDynamic::Create(Material, GetToolManager());
+			PreviewMaterialRC->SetTextureParameterValue(TEXT("BaseColor"), EmptyColorMapWhite);
+			PreviewMaterialRC->SetTextureParameterValue(TEXT("Roughness"), EmptyRoughnessMap);
+			PreviewMaterialRC->SetTextureParameterValue(TEXT("Metallic"), EmptyMetallicMap);
+			PreviewMaterialRC->SetTextureParameterValue(TEXT("Specular"), EmptySpecularMap);
+			PreviewMaterialRC->SetTextureParameterValue(TEXT("Emissive"), EmptyEmissiveMap);
+			PreviewMaterialRC->SetTextureParameterValue(TEXT("NormalMap"), EmptyNormalMap);
+		}
+	}
+	
+	{
+		UMaterial* Material = LoadObject<UMaterial>(nullptr, TEXT("/MeshModelingToolsetExp/Materials/FullMaterialBakePreviewMaterial_PackedMRS"));
+		check(Material);
+		if (Material != nullptr)
+		{
+			PreviewMaterialPackedRC = UMaterialInstanceDynamic::Create(Material, GetToolManager());
+			PreviewMaterialPackedRC->SetTextureParameterValue(TEXT("BaseColor"), EmptyColorMapWhite);
+			PreviewMaterialPackedRC->SetTextureParameterValue(TEXT("PackedMRS"), EmptyPackedMRSMap);
+			PreviewMaterialPackedRC->SetTextureParameterValue(TEXT("Emissive"), EmptyEmissiveMap);
+			PreviewMaterialPackedRC->SetTextureParameterValue(TEXT("NormalMap"), EmptyNormalMap);
+		}
+	}
+}
+
+
+
+void UBakeRenderCaptureTool::InvalidateComputeRC()
+{
+	if (!ComputeRC)
+	{
+		// Initialize background compute
+		ComputeRC = MakeUnique<TGenericDataBackgroundCompute<FBakeRenderCaptureResultsBuilder>>();
+		ComputeFactory.Tool = this;
+		ComputeRC->Setup(&ComputeFactory);
+		ComputeRC->OnResultUpdated.AddLambda([this](const TUniquePtr<FBakeRenderCaptureResultsBuilder>& NewResult) { OnMapsUpdatedRC(NewResult); });
+	}
+	ComputeRC->InvalidateResult();
+	OpState = EBakeOpState::Clean;
+}
+
+
+void UBakeRenderCaptureTool::UpdateResult()
+{
+	if (OpState == EBakeOpState::Clean)
+	{
+		// Evaluation already launched/complete. Note that the ComputeRC background compute updates ResultSettings when
+		// they are available by calling OnMapsUpdatedRC in its OnResultUpdated delegate.
+		return;
+	}
+
+	//
+	// create a set of spatially located render captures of the scene ("photo set").
+	// @Refactor SceneCapture currently must happen on the game thread, we should try to fix that
+	//
+	FBakeRenderCaptureOptions::FOptions Options = FBakeRenderCaptureOptions::ConstructOptions(*Settings, *InputMeshSettings);
+
+	for (int Idx = 1; Idx < Targets.Num(); ++Idx)
+	{
+		UE::ToolTarget::ShowSourceObject(Targets[Idx]);
+	}
+
+	SceneCapture.Reset();
+	SceneCapture = CapturePhotoSet(Actors, Options);
+
+	for (int Idx = 1; Idx < Targets.Num(); ++Idx)
+	{
+		UE::ToolTarget::HideSourceObject(Targets[Idx]);
+	}
+
+	// @Incomplete Clear our invalid bitflag to check for valid inputs.
+	//OpState &= ~EBakeOpState::Invalid;
+	//OpState |= CheckValidInputs();
+	//if (static_cast<bool>(OpState & EBakeOpState::Invalid))
+	//{
+	//	// Early exit if op input parameters are invalid.
+	//	return;
+	//}
+	
+	InvalidateResults();
+
+	// This should be the only point of compute invalidation to minimize synchronization issues.
+	InvalidateComputeRC();
+}
+
+
+void UBakeRenderCaptureTool::UpdateVisualization()
+{
+	if (Settings->MapPreview.IsEmpty())
+	{
+		return;
+	}
+
+	FBakeRenderCaptureOptions::FOptions Options = FBakeRenderCaptureOptions::ConstructOptions(*Settings, *InputMeshSettings);
+	
+	// bool bPackedMRS = Options.bUsePackedMRS && (Options.bBakeRoughness || Options.bBakeMetallic || Options.bBakeSpecular);
+	bool bPackedMRS = Options.bUsePackedMRS;
+
+	if (bPackedMRS)
+	{
+		TObjectPtr<UMaterialInstanceDynamic> Material = PreviewMaterialPackedRC;
+		PreviewMesh->SetOverrideRenderMaterial(Material);
+		
+		if (VisualizationProps->bPreviewAsMaterial)
+		{
+			// We set all textures which were computed in the corresponding texture channels
+			Material->SetTextureParameterValue(TEXT("BaseColor"), Options.bBakeBaseColor ? ResultSettings->BaseColorMap : EmptyColorMapWhite);
+			Material->SetTextureParameterValue(TEXT("Emissive"),  Options.bBakeEmissive  ? ResultSettings->EmissiveMap  : EmptyEmissiveMap);
+			Material->SetTextureParameterValue(TEXT("NormalMap"), Options.bBakeNormalMap ? ResultSettings->NormalMap    : EmptyNormalMap);
+			Material->SetTextureParameterValue(TEXT("PackedMRS"), ResultSettings->PackedMRSMap);
+		}
+		else
+		{
+			// The BaseColor texture channel will be set according to the selected MapPreview
+			TObjectPtr<UTexture2D> BaseColorMap = EmptyColorMapWhite;
+			if (Options.bBakeBaseColor && Settings->MapPreview == BaseColorTexParamName.ToString())
+			{
+				BaseColorMap = ResultSettings->BaseColorMap;
+			}
+			else if (Options.bBakeEmissive && Settings->MapPreview == EmissiveTexParamName.ToString())
+			{
+				BaseColorMap = ResultSettings->EmissiveMap;
+			}
+			else if (Options.bBakeNormalMap && Settings->MapPreview == NormalTexParamName.ToString())
+			{
+				BaseColorMap = ResultSettings->NormalMap;
+			}
+			else if (Settings->MapPreview == PackedMRSTexParamName.ToString())
+			{
+				BaseColorMap = ResultSettings->PackedMRSMap;
+			}
+			Material->SetTextureParameterValue(TEXT("BaseColor"), BaseColorMap);
+			Material->SetTextureParameterValue(TEXT("Emissive"),  EmptyEmissiveMap);
+			Material->SetTextureParameterValue(TEXT("NormalMap"), EmptyNormalMap);
+			Material->SetTextureParameterValue(TEXT("PackedMRS"), EmptyPackedMRSMap);
+		}
+
+		Material->SetScalarParameterValue(TEXT("UVChannel"), Options.TargetUVLayer);
+	}
+	else
+	{
+		TObjectPtr<UMaterialInstanceDynamic> Material = PreviewMaterialRC;
+		PreviewMesh->SetOverrideRenderMaterial(Material);
+
+		if (VisualizationProps->bPreviewAsMaterial)
+		{
+			// We set all textures which were computed in the corresponding texture channels
+			Material->SetTextureParameterValue(TEXT("BaseColor"), Options.bBakeBaseColor ? ResultSettings->BaseColorMap : EmptyColorMapWhite);
+			Material->SetTextureParameterValue(TEXT("Roughness"), Options.bBakeRoughness ? ResultSettings->RoughnessMap : EmptyRoughnessMap);
+			Material->SetTextureParameterValue(TEXT("Metallic"),  Options.bBakeMetallic  ? ResultSettings->MetallicMap  : EmptyMetallicMap);
+			Material->SetTextureParameterValue(TEXT("Specular"),  Options.bBakeSpecular  ? ResultSettings->SpecularMap  : EmptySpecularMap);
+			Material->SetTextureParameterValue(TEXT("Emissive"),  Options.bBakeEmissive  ? ResultSettings->EmissiveMap  : EmptyEmissiveMap);
+			Material->SetTextureParameterValue(TEXT("NormalMap"), Options.bBakeNormalMap ? ResultSettings->NormalMap    : EmptyNormalMap);
+		}
+		else
+		{
+			// The BaseColor texture channel will be set according to the selected MapPreview
+			TObjectPtr<UTexture2D> BaseColorMap = EmptyColorMapWhite;
+			if (Options.bBakeBaseColor && Settings->MapPreview == BaseColorTexParamName.ToString())
+			{
+				BaseColorMap = ResultSettings->BaseColorMap;
+			}
+			else if (Options.bBakeRoughness && Settings->MapPreview == RoughnessTexParamName.ToString())
+			{
+				BaseColorMap = ResultSettings->RoughnessMap;
+			}
+			else if (Options.bBakeMetallic && Settings->MapPreview == MetallicTexParamName.ToString())
+			{
+				BaseColorMap = ResultSettings->MetallicMap;
+			}
+			else if (Options.bBakeSpecular && Settings->MapPreview == SpecularTexParamName.ToString())
+			{
+				BaseColorMap = ResultSettings->SpecularMap;
+			}
+			else if (Options.bBakeEmissive && Settings->MapPreview == EmissiveTexParamName.ToString())
+			{
+				BaseColorMap = ResultSettings->EmissiveMap;
+			}
+			else if (Options.bBakeNormalMap && Settings->MapPreview == NormalTexParamName.ToString())
+			{
+				BaseColorMap = ResultSettings->NormalMap;
+			}
+			Material->SetTextureParameterValue(TEXT("BaseColor"), BaseColorMap);
+			
+			Material->SetTextureParameterValue(TEXT("Roughness"), EmptyRoughnessMap);
+			Material->SetTextureParameterValue(TEXT("Metallic"),  EmptyMetallicMap);
+			Material->SetTextureParameterValue(TEXT("Specular"),  EmptySpecularMap);
+			Material->SetTextureParameterValue(TEXT("Emissive"),  EmptyEmissiveMap);
+			Material->SetTextureParameterValue(TEXT("NormalMap"), EmptyNormalMap);
+		}
+		
+		Material->SetScalarParameterValue(TEXT("UVChannel"), Options.TargetUVLayer);
+	}
+}
+
+
+void UBakeRenderCaptureTool::GatherAnalytics(FBakeAnalytics::FMeshSettings& Data)
+{
+	
+}
+
+
+void UBakeRenderCaptureTool::InvalidateResults()
+{
+	ResultSettings->BaseColorMap = nullptr;
+	ResultSettings->RoughnessMap = nullptr;
+	ResultSettings->MetallicMap = nullptr;
+	ResultSettings->SpecularMap = nullptr;
+	ResultSettings->PackedMRSMap = nullptr;
+	ResultSettings->EmissiveMap = nullptr;
+	ResultSettings->NormalMap = nullptr;
+}
+
+#undef LOCTEXT_NAMESPACE
diff --git a/Engine/Plugins/Experimental/MeshModelingToolsetExp/Source/MeshModelingToolsExp/Public/BakeMeshAttributeMapsToolBase.h b/Engine/Plugins/Experimental/MeshModelingToolsetExp/Source/MeshModelingToolsExp/Public/BakeMeshAttributeMapsToolBase.h
--- a/Engine/Plugins/Experimental/MeshModelingToolsetExp/Source/MeshModelingToolsExp/Public/BakeMeshAttributeMapsToolBase.h
+++ b/Engine/Plugins/Experimental/MeshModelingToolsetExp/Source/MeshModelingToolsExp/Public/BakeMeshAttributeMapsToolBase.h
@@ -198,7 +198,15 @@ protected:
 	//
 	TUniquePtr<TGenericDataBackgroundCompute<UE::Geometry::FMeshMapBaker>> Compute = nullptr;
 
-	/** Internal cache of bake texture results. */
+	/**
+	 * Internal cache of bake texture results.
+	 * The tool can inject additional bake types that were not requested by the user. This
+	 * can occur in cases where a particular bake type might need another bake type to preview
+	 * such as BentNormal requiring AmbientOcclusion to preview. To avoid writing out assets
+	 * that the user did not request, we introduce CachedMaps as a temporary texture cache
+	 * for the tool preview. The Result array is then updated from CachedMaps to only hold
+	 * user requested textures that are written out on Shutdown.
+	 */
 	UPROPERTY()
 	TMap<EBakeMapType, TObjectPtr<UTexture2D>> CachedMaps;
 
diff --git a/Engine/Plugins/Experimental/MeshModelingToolsetExp/Source/MeshModelingToolsExp/Public/BakeRenderCaptureTool.h b/Engine/Plugins/Experimental/MeshModelingToolsetExp/Source/MeshModelingToolsExp/Public/BakeRenderCaptureTool.h
new file mode 100644
--- /dev/null
+++ b/Engine/Plugins/Experimental/MeshModelingToolsetExp/Source/MeshModelingToolsExp/Public/BakeRenderCaptureTool.h
@@ -0,0 +1,339 @@
+﻿// Copyright Epic Games, Inc. All Rights Reserved.
+
+#pragma once
+
+#include "CoreMinimal.h"
+#include "MultiSelectionTool.h"
+#include "InteractiveToolBuilder.h"
+#include "BaseTools/MultiSelectionMeshEditingTool.h"
+#include "BakeMeshAttributeMapsToolBase.h"
+
+// Render Capture algorithm includes
+#include "Scene/SceneCapturePhotoSet.h"
+
+#include "BakeRenderCaptureTool.generated.h"
+
+class UTexture2D;
+
+//
+// Tool Result
+//
+
+
+UCLASS()
+class MESHMODELINGTOOLSEXP_API UBakeRenderCaptureResults  : public UInteractiveToolPropertySet
+{
+	GENERATED_BODY()
+
+public:
+	UPROPERTY(VisibleAnywhere, Category = Results, meta = (TransientToolProperty))
+	TObjectPtr<UTexture2D> BaseColorMap;
+
+	UPROPERTY(VisibleAnywhere, Category = Results, meta = (TransientToolProperty))
+	TObjectPtr<UTexture2D> RoughnessMap;
+
+	UPROPERTY(VisibleAnywhere, Category = Results, meta = (TransientToolProperty))
+	TObjectPtr<UTexture2D> MetallicMap;
+
+	UPROPERTY(VisibleAnywhere, Category = Results, meta = (TransientToolProperty))
+	TObjectPtr<UTexture2D> SpecularMap;
+
+	/** Packed Metallic/Roughness/Specular Map */
+	UPROPERTY(VisibleAnywhere, Category = Results, meta = (DisplayName = "Packed MRS Map", TransientToolProperty))
+	TObjectPtr<UTexture2D> PackedMRSMap;
+
+	UPROPERTY(VisibleAnywhere, Category = Results, meta = (TransientToolProperty))
+	TObjectPtr<UTexture2D> EmissiveMap;
+
+	/** World space normal map */
+	UPROPERTY(VisibleAnywhere, Category = Results, meta = (TransientToolProperty))
+	TObjectPtr<UTexture2D> NormalMap;
+};
+
+struct FBakeRenderCaptureResultsBuilder
+{
+	TUniquePtr<TImageBuilder<FVector4f>> RoughnessImage;
+	TUniquePtr<TImageBuilder<FVector4f>> MetallicImage; 
+	TUniquePtr<TImageBuilder<FVector4f>> SpecularImage; 
+	TUniquePtr<TImageBuilder<FVector4f>> PackedMRSImage; 
+	TUniquePtr<TImageBuilder<FVector4f>> EmissiveImage;
+	TUniquePtr<TImageBuilder<FVector4f>> ColorImage;
+	TUniquePtr<TImageBuilder<FVector3f>> NormalImage;
+};
+
+//
+// Tool Builder
+//
+
+
+
+UCLASS()
+class MESHMODELINGTOOLSEXP_API UBakeRenderCaptureToolBuilder : public UMultiSelectionMeshEditingToolBuilder
+{
+	GENERATED_BODY()
+
+public:
+	virtual bool CanBuildTool(const FToolBuilderState& SceneState) const override;
+	virtual UMultiSelectionMeshEditingTool* CreateNewTool(const FToolBuilderState& SceneState) const override;
+
+protected:
+	virtual const FToolTargetTypeRequirements& GetTargetRequirements() const override;
+};
+
+
+
+//
+// Tool Properties
+//
+
+
+USTRUCT()
+struct FMaterialProxySettingsRC
+{
+	GENERATED_USTRUCT_BODY()
+
+	// Size of generated BaseColor map
+	UPROPERTY(Category = Material, EditAnywhere, meta=(ClampMin="1", UIMin="1"))
+	int32 TextureSize;
+
+	// Whether to generate a texture for the World Normal property
+	UPROPERTY(Category = Material, EditAnywhere)
+	uint8 bNormalMap:1;
+
+	// Whether to generate a texture for the Metallic property
+	UPROPERTY(Category = Material, EditAnywhere, meta=(EditCondition="bPackedMRSMap == false"))
+	uint8 bMetallicMap:1;
+
+	// Whether to generate a texture for the Roughness property
+	UPROPERTY(Category = Material, EditAnywhere, meta=(EditCondition="bPackedMRSMap == false"))
+	uint8 bRoughnessMap:1;
+
+	// Whether to generate a texture for the Specular property
+	UPROPERTY(Category = Material, EditAnywhere, meta=(EditCondition="bPackedMRSMap == false"))
+	uint8 bSpecularMap:1;
+	
+	// Whether to generate a packed texture for the Metallic/Roughness/Specular properties
+	UPROPERTY(Category = Material, EditAnywhere)
+	uint8 bPackedMRSMap:1;
+
+	// Whether to generate a texture for the Emissive property
+	UPROPERTY(Category = Material, EditAnywhere)
+	uint8 bEmissiveMap:1;
+
+	FMaterialProxySettingsRC()
+		: TextureSize(512)
+		, bNormalMap(true)
+		, bMetallicMap(true)
+		, bRoughnessMap(true)
+		, bSpecularMap(true)
+		, bPackedMRSMap(true)
+		, bEmissiveMap(true)
+	{
+	}
+
+	bool operator == (const FMaterialProxySettingsRC& Other) const
+	{
+		// @Incomplete use the uproperty compare thing
+		return TextureSize == Other.TextureSize
+			&& bNormalMap == Other.bNormalMap
+			&& bMetallicMap == Other.bMetallicMap
+			&& bRoughnessMap == Other.bRoughnessMap
+			&& bSpecularMap == Other.bSpecularMap
+			&& bPackedMRSMap == Other.bPackedMRSMap
+			&& bEmissiveMap == Other.bEmissiveMap;
+	}
+
+	bool operator != (const FMaterialProxySettingsRC& Other) const
+	{
+		return !(*this == Other);
+	}
+};
+
+UCLASS()
+class MESHMODELINGTOOLSEXP_API UBakeRenderCaptureToolProperties : public UInteractiveToolPropertySet
+{
+	GENERATED_BODY()
+
+public:
+	
+	/** The map type to preview */
+	UPROPERTY(EditAnywhere, Category = BakeOutput, meta=(DisplayName="Preview Output Type", TransientToolProperty, GetOptions = GetMapPreviewNamesFunc))
+	FString MapPreview;
+	
+	UFUNCTION()
+	const TArray<FString>& GetMapPreviewNamesFunc()
+	{
+		return MapPreviewNamesList;
+	}
+	UPROPERTY(meta = (TransientToolProperty))
+	TArray<FString> MapPreviewNamesList;
+
+	//
+	// Material Baking Settings
+	//
+
+	/** If Value is > 1, Multisample output baked textures by this amount in each direction (eg 4 == 16x supersampling) */
+	UPROPERTY(EditAnywhere, Category = MaterialSettings, meta = (ClampMin = "0", ClampMax = "8", UIMin = "0", UIMax = "4"))
+	int32 MultiSamplingAA = 0;
+
+	/** If Value is zero, use MaterialSettings resolution, otherwise override the render capture resolution */
+	UPROPERTY(EditAnywhere, Category = MaterialSettings, meta = (ClampMin = "0"))
+	int32 RenderCaptureResolution = 512;
+
+	/** Material generation settings */
+	UPROPERTY(EditAnywhere, Category = MaterialSettings)
+	FMaterialProxySettingsRC MaterialSettings; // TODO make a separate struct and make it work like the baking tool (pick texture widgets etc)
+
+	UPROPERTY(EditAnywhere, AdvancedDisplay, Category = MaterialSettings, meta = (ClampMin = "5.0", ClampMax = "160.0"))
+	float CaptureFieldOfView = 30.0f;
+
+	UPROPERTY(EditAnywhere, AdvancedDisplay, Category = MaterialSettings, meta = (ClampMin = "0.001", ClampMax = "1000.0"))
+	float NearPlaneDist = 1.0f;
+	
+	//
+	// Debug Output Settings
+	//
+
+	/** If true, print out debugging messages */
+	UPROPERTY(EditAnywhere, Category = DebugSettings)
+	bool bPrintDebugMessages = false;
+};
+
+
+UCLASS()
+class MESHMODELINGTOOLSEXP_API UBakeRenderCaptureInputToolProperties : public UInteractiveToolPropertySet
+{
+	GENERATED_BODY()
+
+public:
+	/** Target mesh to sample to */
+	UPROPERTY(VisibleAnywhere, Category = BakeInput, DisplayName = "Target Mesh", meta = (TransientToolProperty))
+	TObjectPtr<UStaticMesh> TargetStaticMesh = nullptr;
+
+	/** UV channel to use for the target mesh */
+	UPROPERTY(EditAnywhere, Category = BakeInput, meta = (DisplayName = "Target Mesh UV Channel",
+		GetOptions = GetTargetUVLayerNamesFunc, TransientToolProperty, NoResetToDefault))
+	FString TargetUVLayer;
+
+	UFUNCTION()
+	const TArray<FString>& GetTargetUVLayerNamesFunc() const
+	{
+		return TargetUVLayerNamesList;
+	}
+
+	UPROPERTY(meta = (TransientToolProperty))
+	TArray<FString> TargetUVLayerNamesList;
+};
+
+
+//
+// Tool
+//
+
+
+
+UCLASS()
+class MESHMODELINGTOOLSEXP_API UBakeRenderCaptureTool : public UBakeMeshAttributeMapsToolBase
+{
+	GENERATED_BODY()
+
+public:
+	UBakeRenderCaptureTool() = default;
+
+	// Begin UInteractiveTool interface
+	virtual void Setup() override;
+	virtual void OnShutdown(EToolShutdownType ShutdownType) override;
+	virtual void OnTick(float DeltaTime) override;
+	virtual bool HasCancel() const override { return true; }
+	virtual bool HasAccept() const override { return true; }
+	virtual bool CanAccept() const override;
+	// End UInteractiveTool interface
+
+
+protected:
+
+	UPROPERTY()
+	TArray<TObjectPtr<AActor>> Actors;
+
+	UPROPERTY()
+	TObjectPtr<UBakeRenderCaptureToolProperties> Settings;
+
+	UPROPERTY()
+	TObjectPtr<UBakeRenderCaptureInputToolProperties> InputMeshSettings;
+
+	// The computed textures are displayed in the details panel and used in the preview material, they are written
+	// out to assest on shutdown.
+	UPROPERTY()
+	TObjectPtr<UBakeRenderCaptureResults> ResultSettings;
+	
+protected:
+
+	// Begin UBakeMeshAttributeMapsToolBase interface
+	virtual void UpdateResult() override;
+	virtual void UpdateVisualization() override;
+	virtual void GatherAnalytics(FBakeAnalytics::FMeshSettings& Data) override;
+	// End UBakeMeshAttributeMapsToolBase interface
+	
+	void InvalidateResults();
+
+	class FComputeFactory : public UE::Geometry::IGenericDataOperatorFactory<FBakeRenderCaptureResultsBuilder>
+	{
+	public:
+		UBakeRenderCaptureTool* Tool;
+
+		// Begin IGenericDataOperatorFactory interface
+		virtual TUniquePtr<UE::Geometry::TGenericDataOperator<FBakeRenderCaptureResultsBuilder>> MakeNewOperator() override;
+		// End IGenericDataOperatorFactory interface
+	};
+
+	// TODO :PortRenderCaptureToolToNewBakingFramework
+	// This tool has been ported to modeling mode after the new baking framework was implemented,
+	// since there is quite a lot of code to move the plan is to get the tool working in modeling
+	// mode using the old framework and then port it to the new one. This means that some aspects
+	// of the UBakeMeshAttributeMapsToolBase cannot yet be used, so we work around these and point
+	// out where we're going to need to do future @Refactor-ing. -- matija.kecman, 4 March 2022
+
+	// :PortRenderCaptureToolToNewBakingFramework
+	// When we use the new baking framework, we should be able to use:
+	// - UBakeMeshAttributeMapsToolBase::Compute
+	// - UBakeMeshAttributeMapsToolBase::InvalidateCompute
+	// - UBakeMeshAttributeMapsToolBase::OnMapsUpdated
+	// - UBakeMeshAttributeMapsToolBase::OnTick
+ 	FComputeFactory ComputeFactory;
+	TUniquePtr<TGenericDataBackgroundCompute<FBakeRenderCaptureResultsBuilder>> ComputeRC = nullptr;
+	void InvalidateComputeRC();
+	void OnMapsUpdatedRC(const TUniquePtr<FBakeRenderCaptureResultsBuilder>& NewResult);
+	
+	/**
+	 * Create texture assets from our result map of Texture2D
+	 * @param SourceWorld the source world to define where the texture assets will be stored.
+	 * @param SourceAsset if not null, result textures will be stored adjacent to this asset.
+	 */
+	void CreateTextureAssetsRC(UWorld* SourceWorld, UObject* SourceAsset);
+
+	// empty maps are shown when nothing is computed
+	UPROPERTY()
+	TObjectPtr<UTexture2D> EmptyEmissiveMap;
+	UPROPERTY()
+	TObjectPtr<UTexture2D> EmptyPackedMRSMap;
+	UPROPERTY()
+	TObjectPtr<UTexture2D> EmptyRoughnessMap;
+	UPROPERTY()
+	TObjectPtr<UTexture2D> EmptyMetallicMap;
+	UPROPERTY()
+	TObjectPtr<UTexture2D> EmptySpecularMap;
+	UPROPERTY()
+	TObjectPtr<UMaterialInstanceDynamic> PreviewMaterialRC;
+	UPROPERTY()
+	TObjectPtr<UMaterialInstanceDynamic> PreviewMaterialPackedRC;
+	void InitializePreviewMaterials();
+
+	// @Refactor TODO We currently need to compute this on the game thread because the implementation has checks for this
+	TUniquePtr<UE::Geometry::FSceneCapturePhotoSet> SceneCapture;
+
+	// Analytics
+	virtual FString GetAnalyticsEventName() const override
+	{
+		return TEXT("BakeRC");
+	}
+};
