Render Capture Baking: Cleanup some unused code

#rnx

diff --git a/Engine/Plugins/Experimental/MeshModelingToolsetExp/Source/MeshModelingToolsExp/Private/BakeRenderCaptureTool.cpp b/Engine/Plugins/Experimental/MeshModelingToolsetExp/Source/MeshModelingToolsExp/Private/BakeRenderCaptureTool.cpp
--- a/Engine/Plugins/Experimental/MeshModelingToolsetExp/Source/MeshModelingToolsExp/Private/BakeRenderCaptureTool.cpp
+++ b/Engine/Plugins/Experimental/MeshModelingToolsetExp/Source/MeshModelingToolsExp/Private/BakeRenderCaptureTool.cpp
@@ -48,9 +48,6 @@ public:
 	 */
 	struct FOptions
 	{
-		// Base path for generated assets. Names will be generated by appending strings to this path.
-		FString BasePackagePath;
-
 		//
 		// Material approximation settings
 		//
@@ -60,8 +57,6 @@ public:
 		double FieldOfViewDegrees = 45.0;
 		double NearPlaneDist = 1.0;
 
-		bool bMaximizeBakeParallelism = true;		// if true, photocapture is computed at the same time as mesh generation. Faster but requires more memory.
-
 		//
 		// Material output settings
 		//
@@ -78,30 +73,17 @@ public:
 		bool bUsePackedMRS = true;
 
 		// output texture options
-		FBakeRenderCaptureOptions::ETextureSizePolicy TextureSizePolicy = FBakeRenderCaptureOptions::ETextureSizePolicy::TextureSize;
 		int32 TextureImageSize = 1024;
-		float MeshTexelDensity = 0;
 
 		// supersampling parameter
 		int32 AntiAliasMultiSampling = 0;
 
-
 		//
 		// Mesh settings
 		//
 
 		//  Which UV layer of the Target mesh (the one we're baking to) should be used
 		int32 TargetUVLayer = 0;
-
-		//
-		// Debug settings
-		//
-
-		// print useful information to the Output Log
-		bool bVerbose = false;
-
-		// create a flattened (ie non-instanced) mesh and save it with _DEBUG suffix. Warning often absolutely enormous!
-		bool bWriteDebugMesh = false;
 	};
 
 	/**
@@ -118,13 +100,11 @@ public:
 
 		Options.TargetUVLayer = InputMeshSettings.TargetUVLayerNamesList.IndexOfByKey(InputMeshSettings.TargetUVLayer);
 		
-		// Options.BasePackagePath // @Incomplete
 		Options.RenderCaptureImageSize = (UseSettings.RenderCaptureResolution == 0) ? Options.TextureImageSize : UseSettings.RenderCaptureResolution;
 		Options.FieldOfViewDegrees = UseSettings.CaptureFieldOfView;
 		Options.NearPlaneDist = UseSettings.NearPlaneDist;
-		//Options.bMaximizeBakeParallelism = true; // @Incomplete
 		
-		//Options.bBakeBaseColor // @Incomplete
+		//Options.bBakeBaseColor // This is always true
 		Options.bBakeRoughness = UseSettings.MaterialSettings.bRoughnessMap;
 		Options.bBakeMetallic = UseSettings.MaterialSettings.bMetallicMap;
 		Options.bBakeSpecular = UseSettings.MaterialSettings.bSpecularMap;
@@ -132,65 +112,13 @@ public:
 		Options.bBakeNormalMap = UseSettings.MaterialSettings.bNormalMap;
 		Options.bUsePackedMRS = UseSettings.MaterialSettings.bPackedMRSMap;
 
-		//Options.TextureSizePolicy // @Incomplete
 		Options.TextureImageSize = UseSettings.MaterialSettings.TextureSize;
-		//Options.MeshTexelDensity // @Incomplete
 		Options.AntiAliasMultiSampling = FMath::Max(1, UseSettings.MultiSamplingAA);
 
-		// @Incomplete
-		//Options.bGenerateNaniteEnabledMesh 
-		//Options.NaniteProxyTrianglePercent
-		//Options.bSupportRayTracing
-		//Options.bAllowDistanceField
-		//Options.bGenerateLightmapUVs
-		//Options.bCreatePhysicsBody
-		//Options.bVerbose
-		//Options.bWriteDebugMesh
-
 		return Options;
 	}
 };
 
-static int32 GetMeshTextureSizeFromTargetTexelDensity(const FDynamicMesh3& Mesh, float TargetTexelDensity)
-{
-	const FDynamicMeshUVOverlay* UVOverlay = Mesh.Attributes()->PrimaryUV();
-	double Mesh3DArea = 0;
-	double MeshUVArea = 0;
-	for (int TriangleID : Mesh.TriangleIndicesItr())
-	{
-		// World space area
-		Mesh3DArea += Mesh.GetTriArea(TriangleID);
-
-		FIndex3i UVVertices = UVOverlay->GetTriangle(TriangleID);
-		FTriangle2d TriangleUV = FTriangle2d(
-			(FVector2d)UVOverlay->GetElement(UVVertices.A),
-			(FVector2d)UVOverlay->GetElement(UVVertices.B),
-			(FVector2d)UVOverlay->GetElement(UVVertices.C));
-
-		// UV space area
-		MeshUVArea += TriangleUV.Area();
-	}
-	double TexelRatio = FMath::Sqrt(MeshUVArea / Mesh3DArea) * 100;
-
-	// Compute the perfect texture size that would get us to our texture density
-	// Also compute the nearest power of two sizes (below and above our target)
-	const int32 SizePerfect = FMath::CeilToInt(TargetTexelDensity / TexelRatio);
-	const int32 SizeHi = FMath::RoundUpToPowerOfTwo(SizePerfect);
-	const int32 SizeLo = SizeHi >> 1;
-
-	// Compute the texel density we achieve with these two texture sizes
-	const double TexelDensityLo = SizeLo * TexelRatio;
-	const double TexelDensityHi = SizeHi * TexelRatio;
-
-	// Select best match between low & high res textures.
-	const double TexelDensityLoDiff = TargetTexelDensity - TexelDensityLo;
-	const double TexelDensityHiDiff = TexelDensityHi - TargetTexelDensity;
-	const int32 BestTextureSize = TexelDensityLoDiff < TexelDensityHiDiff ? SizeLo : SizeHi;
-
-	return BestTextureSize;
-}
-
-
 static TUniquePtr<FSceneCapturePhotoSet> CapturePhotoSet(
 	const TArray<TObjectPtr<AActor>>& Actors,
 	const FBakeRenderCaptureOptions::FOptions& Options
@@ -457,28 +385,8 @@ public:
 		check(Actors != nullptr);
 		check(SceneCapture != nullptr);
 
-
-		FBakeRenderCaptureOptions::FOptions OverridenOptions = Options;
-
-		// evaluate required texture size if needed
-		if (Options.TextureSizePolicy == FBakeRenderCaptureOptions::ETextureSizePolicy::TexelDensity)
-		{
-			const int32 MaxTextureSize = 8192;
-			const int32 BestTextureSize = GetMeshTextureSizeFromTargetTexelDensity(*BaseMesh, Options.MeshTexelDensity);
-
-			if (BestTextureSize > MaxTextureSize)
-			{
-				UE_LOG(LogGeometry, Warning, TEXT("Mesh would require %dx%d textures, clamping down to maximum (%dx%d)"), BestTextureSize, BestTextureSize, MaxTextureSize, MaxTextureSize);
-				OverridenOptions.TextureImageSize = MaxTextureSize;
-			}
-			else
-			{
-				OverridenOptions.TextureImageSize = BestTextureSize;
-			}
-		}
-
 		// bake textures onto the target mesh by projecting/sampling the set of captured photos
-		ImageBuildersFromPhotoSet(SceneCapture, OverridenOptions, BaseMesh, BaseMeshTangents.Get(), Result);
+		ImageBuildersFromPhotoSet(SceneCapture, Options, BaseMesh, BaseMeshTangents.Get(), Result);
 	}
 	// End TGenericDataOperator interface
 };
@@ -567,7 +475,6 @@ void UBakeRenderCaptureTool::Setup()
 	Settings->WatchProperty(Settings->MaterialSettings, [this](FMaterialProxySettingsRC) { OpState |= EBakeOpState::Evaluate; });
 	Settings->WatchProperty(Settings->CaptureFieldOfView, [this](float) { OpState |= EBakeOpState::Evaluate; });
 	Settings->WatchProperty(Settings->NearPlaneDist, [this](float) { OpState |= EBakeOpState::Evaluate; });
-	Settings->WatchProperty(Settings->bPrintDebugMessages, [this](bool) { OpState |= EBakeOpState::Evaluate; }); // @Incomplete Actually make this do something or remove it
 	
 	InputMeshSettings = NewObject<UBakeRenderCaptureInputToolProperties>(this);
 	InputMeshSettings->RestoreProperties(this);
diff --git a/Engine/Plugins/Experimental/MeshModelingToolsetExp/Source/MeshModelingToolsExp/Public/BakeRenderCaptureTool.h b/Engine/Plugins/Experimental/MeshModelingToolsetExp/Source/MeshModelingToolsExp/Public/BakeRenderCaptureTool.h
--- a/Engine/Plugins/Experimental/MeshModelingToolsetExp/Source/MeshModelingToolsExp/Public/BakeRenderCaptureTool.h
+++ b/Engine/Plugins/Experimental/MeshModelingToolsetExp/Source/MeshModelingToolsExp/Public/BakeRenderCaptureTool.h
@@ -189,14 +189,6 @@ public:
 
 	UPROPERTY(EditAnywhere, AdvancedDisplay, Category = MaterialSettings, meta = (ClampMin = "0.001", ClampMax = "1000.0"))
 	float NearPlaneDist = 1.0f;
-	
-	//
-	// Debug Output Settings
-	//
-
-	/** If true, print out debugging messages */
-	UPROPERTY(EditAnywhere, Category = DebugSettings)
-	bool bPrintDebugMessages = false;
 };
 
 
