Add a TargetWorld member to MeshSurfacePointTool, similar to other base tool classes.
Add attribute editor and vertex paint tools to cloth editor.

#rnx

diff --git a/Engine/Plugins/Experimental/MeshModelingToolsetExp/Source/MeshModelingToolsExp/Private/DeformMeshPolygonsTool.cpp b/Engine/Plugins/Experimental/MeshModelingToolsetExp/Source/MeshModelingToolsExp/Private/DeformMeshPolygonsTool.cpp
--- a/Engine/Plugins/Experimental/MeshModelingToolsetExp/Source/MeshModelingToolsExp/Private/DeformMeshPolygonsTool.cpp
+++ b/Engine/Plugins/Experimental/MeshModelingToolsetExp/Source/MeshModelingToolsExp/Private/DeformMeshPolygonsTool.cpp
@@ -833,11 +833,6 @@ UDeformMeshPolygonsTool::UDeformMeshPolygonsTool()
 	UInteractiveTool::SetToolDisplayName(LOCTEXT("DeformPolygroupsToolName", "PolyGroup Deform"));
 }
 
-void UDeformMeshPolygonsTool::SetWorld(UWorld* World)
-{
-	TargetWorld = World;
-}
-
 void UDeformMeshPolygonsTool::Setup()
 {
 	UMeshSurfacePointTool::Setup();
@@ -845,7 +840,7 @@ void UDeformMeshPolygonsTool::Setup()
 	LaplacianDeformer = MakePimpl<FGroupTopologyLaplacianDeformer>();
 
 	// create dynamic mesh component to use for live preview
-	check(TargetWorld);
+	check(TargetWorld.IsValid());
 	FActorSpawnParameters SpawnInfo;
 	PreviewMeshActor = TargetWorld->SpawnActor<AInternalToolFrameworkActor>(FVector::ZeroVector, FRotator::ZeroRotator, SpawnInfo);
 
diff --git a/Engine/Plugins/Experimental/MeshModelingToolsetExp/Source/MeshModelingToolsExp/Private/DynamicMeshBrushTool.cpp b/Engine/Plugins/Experimental/MeshModelingToolsetExp/Source/MeshModelingToolsExp/Private/DynamicMeshBrushTool.cpp
--- a/Engine/Plugins/Experimental/MeshModelingToolsetExp/Source/MeshModelingToolsExp/Private/DynamicMeshBrushTool.cpp
+++ b/Engine/Plugins/Experimental/MeshModelingToolsetExp/Source/MeshModelingToolsExp/Private/DynamicMeshBrushTool.cpp
@@ -24,7 +24,7 @@ void UDynamicMeshBrushTool::Setup()
 {
 	PreviewMesh = NewObject<UPreviewMesh>(this);
 	PreviewMesh->bBuildSpatialDataStructure = true;
-	PreviewMesh->CreateInWorld(UE::ToolTarget::GetTargetActor(Target)->GetWorld(), FTransform::Identity);
+	PreviewMesh->CreateInWorld(GetTargetWorld(), FTransform::Identity);
 	FTransformSRT3d LocalToWorldTransform = UE::ToolTarget::GetLocalToWorldTransform(Target);
 	PreviewMesh->SetTransform((FTransform)LocalToWorldTransform);
 	ToolSetupUtil::ApplyRenderingConfigurationToPreview(PreviewMesh, Target);
diff --git a/Engine/Plugins/Experimental/MeshModelingToolsetExp/Source/MeshModelingToolsExp/Public/DeformMeshPolygonsTool.h b/Engine/Plugins/Experimental/MeshModelingToolsetExp/Source/MeshModelingToolsExp/Public/DeformMeshPolygonsTool.h
--- a/Engine/Plugins/Experimental/MeshModelingToolsetExp/Source/MeshModelingToolsExp/Public/DeformMeshPolygonsTool.h
+++ b/Engine/Plugins/Experimental/MeshModelingToolsetExp/Source/MeshModelingToolsExp/Public/DeformMeshPolygonsTool.h
@@ -154,8 +154,6 @@ public:
 
 	virtual void RegisterActions(FInteractiveToolActionSet& ActionSet) override;
 
-	virtual void SetWorld(UWorld* World);
-
 	virtual void Setup() override;
 	virtual void Shutdown(EToolShutdownType ShutdownType) override;
 
@@ -180,9 +178,6 @@ public:
 
 protected:
 
-	UPROPERTY()
-	TObjectPtr<UWorld> TargetWorld = nullptr;
-
 	UPROPERTY()
 	TObjectPtr<AInternalToolFrameworkActor> PreviewMeshActor = nullptr;
 
diff --git a/Engine/Plugins/Experimental/MeshModelingToolsetExp/Source/MeshModelingToolsExp/Public/EditUVIslandsTool.h b/Engine/Plugins/Experimental/MeshModelingToolsetExp/Source/MeshModelingToolsExp/Public/EditUVIslandsTool.h
--- a/Engine/Plugins/Experimental/MeshModelingToolsetExp/Source/MeshModelingToolsExp/Public/EditUVIslandsTool.h
+++ b/Engine/Plugins/Experimental/MeshModelingToolsetExp/Source/MeshModelingToolsExp/Public/EditUVIslandsTool.h
@@ -66,11 +66,6 @@ public:
 
 	virtual void RegisterActions(FInteractiveToolActionSet& ActionSet) override;
 
-	virtual void SetWorld(UWorld* World)
-	{
-		TargetWorld = World;
-	}
-
 	virtual void Setup() override;
 	virtual void Shutdown(EToolShutdownType ShutdownType) override;
 
@@ -101,9 +96,6 @@ public:
 
 protected:
 
-	UPROPERTY()
-	TObjectPtr<UWorld> TargetWorld = nullptr;
-
 	UPROPERTY()
 	TObjectPtr<AInternalToolFrameworkActor> PreviewMeshActor = nullptr;
 
diff --git a/Engine/Source/Runtime/InteractiveToolsFramework/Private/BaseTools/MeshSurfacePointTool.cpp b/Engine/Source/Runtime/InteractiveToolsFramework/Private/BaseTools/MeshSurfacePointTool.cpp
--- a/Engine/Source/Runtime/InteractiveToolsFramework/Private/BaseTools/MeshSurfacePointTool.cpp
+++ b/Engine/Source/Runtime/InteractiveToolsFramework/Private/BaseTools/MeshSurfacePointTool.cpp
@@ -45,6 +45,7 @@ void UMeshSurfacePointToolBuilder::InitializeNewTool(UMeshSurfacePointTool* NewT
 	check(Target);
 	NewTool->SetTarget(Target);
 	NewTool->SetStylusAPI(this->StylusAPI);
+	NewTool->SetWorld(SceneState.World);
 }
 
 
@@ -202,6 +203,16 @@ float UMeshSurfacePointTool::GetCurrentDevicePressure() const
 	return (StylusAPI != nullptr) ? FMath::Clamp(StylusAPI->GetCurrentPressure(), 0.0f, 1.0f) : 1.0f;
 }
 
+void UMeshSurfacePointTool::SetWorld(UWorld* World)
+{
+	TargetWorld = World;
+}
+
+UWorld* UMeshSurfacePointTool::GetTargetWorld()
+{
+	return TargetWorld.Get();
+}
+
 
 
 #undef LOCTEXT_NAMESPACE
\ No newline at end of file
diff --git a/Engine/Source/Runtime/InteractiveToolsFramework/Public/BaseTools/MeshSurfacePointTool.h b/Engine/Source/Runtime/InteractiveToolsFramework/Public/BaseTools/MeshSurfacePointTool.h
--- a/Engine/Source/Runtime/InteractiveToolsFramework/Public/BaseTools/MeshSurfacePointTool.h
+++ b/Engine/Source/Runtime/InteractiveToolsFramework/Public/BaseTools/MeshSurfacePointTool.h
@@ -134,6 +134,9 @@ public:
 	/** @return current input device pressure in range 0-1 */
 	virtual float GetCurrentDevicePressure() const;
 
+	void SetWorld(UWorld* World);
+	UWorld* GetTargetWorld();
+
 protected:
 	/** Current state of the shift modifier toggle */
 	bool bShiftToggle = false;
@@ -144,5 +147,8 @@ protected:
 	FRay LastWorldRay;
 
 	IToolStylusStateProviderAPI* StylusAPI = nullptr;
+
+	UPROPERTY()
+	TWeakObjectPtr<UWorld> TargetWorld = nullptr;
 };
 
