add convex decomposition options to Msh2Coll tool and geometryscript collision functions

diff --git a/Engine/Plugins/Experimental/GeometryScripting/Source/GeometryScriptingCore/Private/CollisionFunctions.cpp b/Engine/Plugins/Experimental/GeometryScripting/Source/GeometryScriptingCore/Private/CollisionFunctions.cpp
--- a/Engine/Plugins/Experimental/GeometryScripting/Source/GeometryScriptingCore/Private/CollisionFunctions.cpp
+++ b/Engine/Plugins/Experimental/GeometryScripting/Source/GeometryScriptingCore/Private/CollisionFunctions.cpp
@@ -92,7 +92,17 @@ void ComputeCollisionFromMesh(
 	case EGeometryScriptCollisionGenerationMethod::ConvexHulls:
 		ShapeGenerator.bSimplifyHulls = Options.bSimplifyHulls;
 		ShapeGenerator.HullTargetFaceCount = Options.ConvexHullTargetFaceCount;
-		ShapeGenerator.Generate_ConvexHulls(NewCollision.Geometry);
+		if (Options.MaxConvexHullsPerMesh > 1)
+		{
+			ShapeGenerator.ConvexDecompositionMaxPieces = Options.MaxConvexHullsPerMesh;
+			ShapeGenerator.ConvexDecompositionSearchFactor = Options.ConvexDecompositionSearchFactor;
+			ShapeGenerator.ConvexDecompositionErrorTolerance = Options.ConvexDecompositionErrorTolerance;
+			ShapeGenerator.Generate_ConvexHullDecompositions(NewCollision.Geometry);
+		}
+		else
+		{
+			ShapeGenerator.Generate_ConvexHulls(NewCollision.Geometry);
+		}
 		break;
 	case EGeometryScriptCollisionGenerationMethod::SweptHulls:
 		ShapeGenerator.bSimplifyHulls = Options.bSimplifyHulls;
diff --git a/Engine/Plugins/Experimental/GeometryScripting/Source/GeometryScriptingCore/Public/GeometryScript/CollisionFunctions.h b/Engine/Plugins/Experimental/GeometryScripting/Source/GeometryScriptingCore/Public/GeometryScript/CollisionFunctions.h
--- a/Engine/Plugins/Experimental/GeometryScripting/Source/GeometryScriptingCore/Public/GeometryScript/CollisionFunctions.h
+++ b/Engine/Plugins/Experimental/GeometryScripting/Source/GeometryScriptingCore/Public/GeometryScript/CollisionFunctions.h
@@ -66,6 +66,15 @@ public:
 	UPROPERTY(BlueprintReadWrite, Category = Options)
 	int ConvexHullTargetFaceCount = 25;
 
+	UPROPERTY(BlueprintReadWrite, Category = Options)
+	int MaxConvexHullsPerMesh = 1;
+
+	UPROPERTY(BlueprintReadWrite, Category = Options)
+	float ConvexDecompositionSearchFactor = .5f;
+
+	UPROPERTY(BlueprintReadWrite, Category = Options)
+	float ConvexDecompositionErrorTolerance = 0;
+
 	UPROPERTY(BlueprintReadWrite, Category = Options)
 	float SweptHullSimplifyTolerance = 0.1;
 
diff --git a/Engine/Plugins/Experimental/MeshModelingToolsetExp/Source/MeshModelingToolsExp/Private/Physics/SetCollisionGeometryTool.cpp b/Engine/Plugins/Experimental/MeshModelingToolsetExp/Source/MeshModelingToolsExp/Private/Physics/SetCollisionGeometryTool.cpp
--- a/Engine/Plugins/Experimental/MeshModelingToolsetExp/Source/MeshModelingToolsExp/Private/Physics/SetCollisionGeometryTool.cpp
+++ b/Engine/Plugins/Experimental/MeshModelingToolsetExp/Source/MeshModelingToolsExp/Private/Physics/SetCollisionGeometryTool.cpp
@@ -117,6 +117,9 @@ void USetCollisionGeometryTool::Setup()
 	Settings->WatchProperty(Settings->bDetectCapsules, [this](int32) { bResultValid = false; });
 	Settings->WatchProperty(Settings->bSimplifyHulls, [this](bool) { bResultValid = false; });
 	Settings->WatchProperty(Settings->HullTargetFaceCount, [this](int32) { bResultValid = false; });
+	Settings->WatchProperty(Settings->MaxHullsPerMesh, [this](int32) { bResultValid = false; });
+	Settings->WatchProperty(Settings->ConvexDecompositionSearchFactor, [this](int32) { bResultValid = false; });
+	Settings->WatchProperty(Settings->AddHullsErrorTolerance, [this](int32) { bResultValid = false; });
 	Settings->WatchProperty(Settings->bSimplifyPolygons, [this](bool) { bResultValid = false; });
 	Settings->WatchProperty(Settings->HullTolerance, [this](float) { bResultValid = false; });
 	Settings->WatchProperty(Settings->SweepAxis, [this](EProjectedHullAxis) { bResultValid = false; });
@@ -358,7 +361,17 @@ void USetCollisionGeometryTool::UpdateGeneratedCollision()
 	case ECollisionGeometryType::ConvexHulls:
 		UseShapeGenerator->bSimplifyHulls = Settings->bSimplifyHulls;
 		UseShapeGenerator->HullTargetFaceCount = Settings->HullTargetFaceCount;
-		UseShapeGenerator->Generate_ConvexHulls(NewCollision->Geometry);
+		if (Settings->MaxHullsPerMesh > 1)
+		{
+			UseShapeGenerator->ConvexDecompositionMaxPieces = Settings->MaxHullsPerMesh;
+			UseShapeGenerator->ConvexDecompositionSearchFactor = Settings->ConvexDecompositionSearchFactor;
+			UseShapeGenerator->ConvexDecompositionErrorTolerance = Settings->AddHullsErrorTolerance;
+			UseShapeGenerator->Generate_ConvexHullDecompositions(NewCollision->Geometry);
+		}
+		else
+		{
+			UseShapeGenerator->Generate_ConvexHulls(NewCollision->Geometry);
+		}
 		break;
 	case ECollisionGeometryType::SweptHulls:
 		UseShapeGenerator->bSimplifyHulls = Settings->bSimplifyPolygons;
diff --git a/Engine/Plugins/Experimental/MeshModelingToolsetExp/Source/MeshModelingToolsExp/Public/Physics/SetCollisionGeometryTool.h b/Engine/Plugins/Experimental/MeshModelingToolsetExp/Source/MeshModelingToolsExp/Public/Physics/SetCollisionGeometryTool.h
--- a/Engine/Plugins/Experimental/MeshModelingToolsetExp/Source/MeshModelingToolsExp/Public/Physics/SetCollisionGeometryTool.h
+++ b/Engine/Plugins/Experimental/MeshModelingToolsetExp/Source/MeshModelingToolsExp/Public/Physics/SetCollisionGeometryTool.h
@@ -115,6 +115,21 @@ public:
 		EditConditionHides, EditCondition = "GeometryType == ECollisionGeometryType::ConvexHulls"))
 	int32 HullTargetFaceCount = 20;
 
+	/** How many convex hulls can be used to approximate each mesh */
+	UPROPERTY(EditAnywhere, Category = ConvexHulls, meta = (UIMin = "1", UIMax = "100", ClampMin = "1",
+		EditConditionHides, EditCondition = "GeometryType == ECollisionGeometryType::ConvexHulls"))
+	int32 MaxHullsPerMesh = 1;
+
+	/** How much to search the space of possible decompositions beyond MaxHullsPerMesh; for larger values, will do additional work to try to better approximate mesh features (but resulting hulls may overlap more) */
+	UPROPERTY(EditAnywhere, Category = ConvexHulls, meta = (UIMin = "0", UIMax = "2", ClampMin = "0",
+		EditConditionHides, EditCondition = "GeometryType == ECollisionGeometryType::ConvexHulls && MaxHullsPerMesh > 1"))
+	float ConvexDecompositionSearchFactor = .5;
+
+	/** Error tolerance for adding more convex hulls, in cm.  For volumetric errors, the value will be cubed (so a value of 10 indicates a 10x10x10 volume worth of error is acceptable). */
+	UPROPERTY(EditAnywhere, Category = ConvexHulls, meta = (UIMin = "0", UIMax = "1000", ClampMin = "0",
+		EditConditionHides, EditCondition = "GeometryType == ECollisionGeometryType::ConvexHulls && MaxHullsPerMesh > 1"))
+	float AddHullsErrorTolerance = 0;
+
 	UPROPERTY(EditAnywhere, Category = SweptHulls, meta = (EditConditionHides, EditCondition = "GeometryType == ECollisionGeometryType::SweptHulls"))
 	bool bSimplifyPolygons = true;
 
diff --git a/Engine/Plugins/Runtime/GeometryProcessing/Source/DynamicMesh/Private/Operations/MeshConvexHull.cpp b/Engine/Plugins/Runtime/GeometryProcessing/Source/DynamicMesh/Private/Operations/MeshConvexHull.cpp
--- a/Engine/Plugins/Runtime/GeometryProcessing/Source/DynamicMesh/Private/Operations/MeshConvexHull.cpp
+++ b/Engine/Plugins/Runtime/GeometryProcessing/Source/DynamicMesh/Private/Operations/MeshConvexHull.cpp
@@ -28,34 +28,48 @@ bool FMeshConvexHull::Compute(FProgressCancel* Progress)
 
 	if (bPostSimplify)
 	{
-		check(MaxTargetFaceCount > 0);
-		bool bSimplified = false;
-		if (ConvexHull.TriangleCount() > MaxTargetFaceCount)
-		{
-			FVolPresMeshSimplification Simplifier(&ConvexHull);
-			Simplifier.CollapseMode = FVolPresMeshSimplification::ESimplificationCollapseModes::MinimalExistingVertexError;
-			Simplifier.SimplifyToTriangleCount(MaxTargetFaceCount);
-			bSimplified = true;
-		}
+		bOK = SimplifyHull(ConvexHull, MaxTargetFaceCount, Progress);
+	}
 
+	return bOK;
+}
 
-		if (bSimplified)
-		{
-			// recalculate convex hull
-			// TODO: test if simplified mesh is convex first, can just re-use in that case!!
-			FMeshConvexHull SimplifiedHull(&ConvexHull);
-			if (SimplifiedHull.Compute(Progress))
-			{
-				ConvexHull = MoveTemp(SimplifiedHull.ConvexHull);
-			}
-		}
 
+bool FMeshConvexHull::SimplifyHull(FDynamicMesh3& HullMesh, int32 MaxTargetFaceCount, FProgressCancel* Progress)
+{
+	if (Progress && Progress->Cancelled())
+	{
+		return false;
 	}
 
-	return bOK;
-}
+	check(MaxTargetFaceCount > 0);
+	bool bSimplified = false;
+	if (HullMesh.TriangleCount() > MaxTargetFaceCount)
+	{
+		FVolPresMeshSimplification Simplifier(&HullMesh);
+		Simplifier.CollapseMode = FVolPresMeshSimplification::ESimplificationCollapseModes::MinimalExistingVertexError;
+		Simplifier.SimplifyToTriangleCount(MaxTargetFaceCount);
+		bSimplified = true;
+	}
 
+	if (Progress && Progress->Cancelled())
+	{
+		return false;
+	}
 
+	if (bSimplified)
+	{
+		// recalculate convex hull
+		// TODO: test if simplified mesh is convex first, can just re-use in that case!!
+		FMeshConvexHull SimplifiedHull(&HullMesh);
+		if (SimplifiedHull.Compute(Progress))
+		{
+			HullMesh = MoveTemp(SimplifiedHull.ConvexHull);
+		}
+	}
+
+	return true;
+}
 
 
 
diff --git a/Engine/Plugins/Runtime/GeometryProcessing/Source/DynamicMesh/Private/ShapeApproximation/MeshSimpleShapeApproximation.cpp b/Engine/Plugins/Runtime/GeometryProcessing/Source/DynamicMesh/Private/ShapeApproximation/MeshSimpleShapeApproximation.cpp
--- a/Engine/Plugins/Runtime/GeometryProcessing/Source/DynamicMesh/Private/ShapeApproximation/MeshSimpleShapeApproximation.cpp
+++ b/Engine/Plugins/Runtime/GeometryProcessing/Source/DynamicMesh/Private/ShapeApproximation/MeshSimpleShapeApproximation.cpp
@@ -6,6 +6,7 @@
 #include "MinVolumeSphere3.h"
 #include "MinVolumeBox3.h"
 #include "FitCapsule3.h"
+#include "CompGeom/ConvexDecomposition3.h"
 //#include "DynamicMesh/DynamicMeshAABBTree3.h"
 
 #include "ShapeApproximation/ShapeDetection3.h"
@@ -295,10 +296,46 @@ void FMeshSimpleShapeApproximation::Generate_ConvexHulls(FSimpleShapeSet3d& Shap
 
 		if (Hull.Compute(Progress))
 		{
-			FConvexShape3d NewConvex;
-			NewConvex.Mesh = MoveTemp(Hull.ConvexHull);
 			GeometryLock.Lock();
-			ShapeSetOut.Convexes.Add(NewConvex);
+			ShapeSetOut.Convexes.Emplace(MoveTemp(Hull.ConvexHull));
+			GeometryLock.Unlock();
+		}
+	});
+}
+
+
+
+void FMeshSimpleShapeApproximation::Generate_ConvexHullDecompositions(FSimpleShapeSet3d& ShapeSetOut, FProgressCancel* Progress)
+{
+	FCriticalSection GeometryLock;
+	ParallelFor(SourceMeshes.Num(), [&](int32 idx)
+	{
+		if (GetDetectedSimpleShape(SourceMeshCaches[idx], ShapeSetOut, GeometryLock))
+		{
+			return;
+		}
+
+		if (Progress && Progress->Cancelled())
+		{
+			return;
+		}
+
+		const FDynamicMesh3& SourceMesh = *SourceMeshes[idx];
+		// TODO: if (bSimplifyHulls), also consider simplifying the input?
+		FConvexDecomposition3 Decomposition(SourceMesh);
+
+		Decomposition.Compute(ConvexDecompositionMaxPieces, ConvexDecompositionMaxPieces * ConvexDecompositionSearchFactor, ConvexDecompositionErrorTolerance);
+
+		for (int32 HullIdx = 0; HullIdx < Decomposition.NumHulls(); HullIdx++)
+		{
+			FDynamicMesh3 HullMesh = Decomposition.GetHullMesh(HullIdx);
+			if (!FMeshConvexHull::SimplifyHull(HullMesh, HullTargetFaceCount, Progress))
+			{
+				return;
+			}
+
+			GeometryLock.Lock();
+			ShapeSetOut.Convexes.Emplace(MoveTemp(HullMesh));
 			GeometryLock.Unlock();
 		}
 	});
diff --git a/Engine/Plugins/Runtime/GeometryProcessing/Source/DynamicMesh/Public/Operations/MeshConvexHull.h b/Engine/Plugins/Runtime/GeometryProcessing/Source/DynamicMesh/Public/Operations/MeshConvexHull.h
--- a/Engine/Plugins/Runtime/GeometryProcessing/Source/DynamicMesh/Public/Operations/MeshConvexHull.h
+++ b/Engine/Plugins/Runtime/GeometryProcessing/Source/DynamicMesh/Public/Operations/MeshConvexHull.h
@@ -66,7 +66,11 @@ public:
 	 */
 	bool Compute(FProgressCancel* Progress = nullptr);
 
-
+	/**
+	 * Simplify the output convex hull. Assumes ConvexHull is already computed.
+	 * @return true on success
+	 */
+	static bool SimplifyHull(FDynamicMesh3& HullMesh, int32 MaxTargetFaceCount, FProgressCancel* Progress = nullptr);
 
 protected:
 	bool Compute_FullMesh(FProgressCancel* Progress);
diff --git a/Engine/Plugins/Runtime/GeometryProcessing/Source/DynamicMesh/Public/ShapeApproximation/MeshSimpleShapeApproximation.h b/Engine/Plugins/Runtime/GeometryProcessing/Source/DynamicMesh/Public/ShapeApproximation/MeshSimpleShapeApproximation.h
--- a/Engine/Plugins/Runtime/GeometryProcessing/Source/DynamicMesh/Public/ShapeApproximation/MeshSimpleShapeApproximation.h
+++ b/Engine/Plugins/Runtime/GeometryProcessing/Source/DynamicMesh/Public/ShapeApproximation/MeshSimpleShapeApproximation.h
@@ -68,6 +68,13 @@ public:
 	/** simplification tolerance when simplifying 2D convex hulls, eg for swept/projected hulls */
 	double HullSimplifyTolerance = 1.0;
 
+	/** How many convex pieces to target per mesh when creating convex decompositions.  If ErrorTolerance is set, can create fewer pieces. */
+	int32 ConvexDecompositionMaxPieces = 20;
+	/** How much additional decomposition decomposition + merging to do, as a fraction of max pieces.  Larger values can help better-cover small features, while smaller values create a cleaner decomposition with less overlap between hulls. */
+	float ConvexDecompositionSearchFactor = .5;
+	/** Error tolerance to guide convex decomposition (in cm); we stop adding new parts if the volume error is below the threshold.  For volumetric errors, value will be cubed. */
+	double ConvexDecompositionErrorTolerance = 0;
+
 	bool bUseExactComputationForBox = false;
 
 	//
@@ -115,6 +122,12 @@ public:
 	 */
 	void Generate_ConvexHulls(FSimpleShapeSet3d& ShapeSetOut, FProgressCancel* Progress = nullptr);
 
+	/**
+	 * Calculate multiple 3D Convex Hulls for each input mesh and store in ShapeSetOut.
+	 * Each convex hull is stored as a triangle mesh, and optionally simplified if bSimplifyHulls=true
+	 */
+	void Generate_ConvexHullDecompositions(FSimpleShapeSet3d& ShapeSetOut, FProgressCancel* Progress = nullptr);
+
 
 	/** Type/Mode for deciding 3D axis to use in Generate_ProjectedHulls() */
 	enum class EProjectedHullAxisMode
