Render Capture Baking: Use new baking framework

#rnx

diff --git a/Engine/Plugins/Experimental/MeshModelingToolsetExp/Source/MeshModelingToolsExp/Private/BakeRenderCaptureTool.cpp b/Engine/Plugins/Experimental/MeshModelingToolsetExp/Source/MeshModelingToolsExp/Private/BakeRenderCaptureTool.cpp
--- a/Engine/Plugins/Experimental/MeshModelingToolsetExp/Source/MeshModelingToolsExp/Private/BakeRenderCaptureTool.cpp
+++ b/Engine/Plugins/Experimental/MeshModelingToolsetExp/Source/MeshModelingToolsExp/Private/BakeRenderCaptureTool.cpp
@@ -14,7 +14,9 @@
 
 #include "ModelingObjectsCreationAPI.h"
 
-#include "Sampling/MeshGenericWorldPositionBaker.h"
+#include "Sampling/MeshImageBakingCache.h"
+#include "Sampling/MeshMapBaker.h"
+#include "Sampling/MeshGenericWorldPositionEvaluator.h"
 #include "Image/ImageInfilling.h"
 
 
@@ -137,10 +139,6 @@ static TUniquePtr<FSceneCapturePhotoSet> CapturePhotoSet(
 	SceneCapture->SetCaptureTypeEnabled(ERenderCaptureType::WorldNormal, Options.bBakeNormalMap);
 	SceneCapture->SetCaptureTypeEnabled(ERenderCaptureType::Emissive, Options.bBakeEmissive);
 
-	bool bMetallic = Options.bBakeMetallic;
-	bool bRoughness = Options.bBakeRoughness;
-	bool bSpecular  = Options.bBakeSpecular;
-	// if (Options.bUsePackedMRS && (bMetallic || bRoughness || bSpecular ) )
 	if (Options.bUsePackedMRS)
 	{
 		SceneCapture->SetCaptureTypeEnabled(ERenderCaptureType::CombinedMRS, true);
@@ -151,9 +149,9 @@ static TUniquePtr<FSceneCapturePhotoSet> CapturePhotoSet(
 	else
 	{
 		SceneCapture->SetCaptureTypeEnabled(ERenderCaptureType::CombinedMRS, false);
-		SceneCapture->SetCaptureTypeEnabled(ERenderCaptureType::Roughness, bRoughness);
-		SceneCapture->SetCaptureTypeEnabled(ERenderCaptureType::Metallic, bMetallic);
-		SceneCapture->SetCaptureTypeEnabled(ERenderCaptureType::Specular, bSpecular);
+		SceneCapture->SetCaptureTypeEnabled(ERenderCaptureType::Roughness, Options.bBakeRoughness);
+		SceneCapture->SetCaptureTypeEnabled(ERenderCaptureType::Metallic, Options.bBakeMetallic);
+		SceneCapture->SetCaptureTypeEnabled(ERenderCaptureType::Specular, Options.bBakeSpecular);
 	}
 
 	SceneCapture->SetCaptureSceneActors(Actors[0]->GetWorld(), Actors);
@@ -169,46 +167,16 @@ static TUniquePtr<FSceneCapturePhotoSet> CapturePhotoSet(
 
 static void ImageBuildersFromPhotoSet(
 	FSceneCapturePhotoSet* SceneCapture,
-	const FBakeRenderCaptureOptions::FOptions& Options, 
-	const FDynamicMesh3* WorldTargetMesh,
-	const FMeshTangentsd* MeshTangents,
+	const FBakeRenderCaptureOptions::FOptions& Options,
+	const FDynamicMesh3* BaseMesh,
+	const TSharedPtr<UE::Geometry::FMeshTangentsd, ESPMode::ThreadSafe>& BaseMeshTangents,
 	TUniquePtr<FBakeRenderCaptureResultsBuilder>& Results)
 {
 	TRACE_CPUPROFILER_EVENT_SCOPE(ImageBuildersFromPhotoSet);
 
-	int32 UVLayer = Options.TargetUVLayer;
-	int32 Supersample = FMath::Max(1, Options.AntiAliasMultiSampling);
-	if ( (Options.TextureImageSize * Supersample) > 16384)
-	{
-		UE_LOG(LogGeometry, Warning, TEXT("Ignoring requested supersampling rate %d because it would require image buffers with resolution %d, please try lower value."), Supersample, Options.TextureImageSize * Supersample);
-		Supersample = 1;
-	}
-
-	FImageDimensions OutputDimensions(Options.TextureImageSize*Supersample, Options.TextureImageSize*Supersample);
-
-	//FScopedSlowTask Progress(8.f, LOCTEXT("BakingTextures", "Baking Textures..."));
-	//Progress.MakeDialog(true);
-	//Progress.EnterProgressFrame(1.f, LOCTEXT("BakingSetup", "Setup..."));
-
-	FDynamicMeshAABBTree3 Spatial(WorldTargetMesh, true);
-
-	FMeshImageBakingCache TempBakeCache;
-	{
-		TRACE_CPUPROFILER_EVENT_SCOPE(ImageBuildersFromPhotoSet_Textures_MakeCache);
-		TempBakeCache.SetDetailMesh(WorldTargetMesh, &Spatial);
-		TempBakeCache.SetBakeTargetMesh(WorldTargetMesh);
-		TempBakeCache.SetDimensions(OutputDimensions);
-		TempBakeCache.SetUVLayer(UVLayer);
-		TempBakeCache.SetThickness(0.1);
-		TempBakeCache.SetCorrespondenceStrategy(FMeshImageBakingCache::ECorrespondenceStrategy::Identity);
-		TempBakeCache.ValidateCache();
-	}
-
-	//Progress.EnterProgressFrame(1.f, LOCTEXT("BakingBaseColor", "Baking Base Color..."));
-
-	FAxisAlignedBox3d TargetBounds = WorldTargetMesh->GetBounds();
-	double RayOffsetHackDist = (double)(100.0f * FMathf::ZeroTolerance * TargetBounds.MinDim() );
+	FDynamicMeshAABBTree3 Spatial(BaseMesh, true);
 
+	double RayOffsetHackDist = (double)(100.0f * FMathf::ZeroTolerance * BaseMesh->GetBounds().MinDim() );
 	auto VisibilityFunction = [&Spatial, RayOffsetHackDist](const FVector3d& SurfPos, const FVector3d& ImagePosWorld)
 	{
 		FVector3d RayDir = ImagePosWorld - SurfPos;
@@ -218,36 +186,159 @@ static void ImageBuildersFromPhotoSet(
 		return (HitTID == IndexConstants::InvalidID);
 	};
 
-	FSceneCapturePhotoSet::FSceneSample DefaultSample;
+	FMeshMapBaker Baker;
+	FMeshBakerDynamicMeshSampler Sampler(BaseMesh, &Spatial, BaseMeshTangents.Get());
+	Baker.SetDetailSampler(&Sampler);
+	Baker.SetTargetMesh(BaseMesh);
+	Baker.SetTargetMeshTangents(BaseMeshTangents);
+	Baker.SetDimensions(FImageDimensions(Options.TextureImageSize, Options.TextureImageSize));
+	Baker.SetSamplesPerPixel(1); // TODO Support multisampling (must work with infill)
+	Baker.SetFilter(FMeshMapBaker::EBakeFilterType::None); // TODO Support texture filtering (must work with infill)
+	Baker.SetTargetMeshUVLayer(Options.TargetUVLayer);
+	Baker.SetCorrespondenceStrategy(FMeshMapBaker::ECorrespondenceStrategy::Identity);
+
+	TMap<ERenderCaptureType, int32> CaptureTypeToEvaluatorIndexMap;
+
 	FVector4f InvalidColor(0, -1, 0, 1);
-	DefaultSample.BaseColor = FVector3f(InvalidColor.X, InvalidColor.Y, InvalidColor.Z);
-
-	FMeshGenericWorldPositionColorBaker BaseColorBaker;
-	BaseColorBaker.SetCache(&TempBakeCache);
-	BaseColorBaker.ColorSampleFunction = [&](FVector3d Position, FVector3d Normal) {
-		FSceneCapturePhotoSet::FSceneSample Sample = DefaultSample;
-		SceneCapture->ComputeSample(FRenderCaptureTypeFlags::BaseColor(),
-			Position, Normal, VisibilityFunction, Sample);
-		return Sample.GetValue4f(ERenderCaptureType::BaseColor);
+	FVector3f DefaultNormal = FVector3f::UnitZ();
+
+	FSceneCapturePhotoSet::FSceneSample DefaultColorSample;
+	DefaultColorSample.BaseColor = FVector3f(InvalidColor.X, InvalidColor.Y, InvalidColor.Z);
+	DefaultColorSample.Roughness = InvalidColor.X;
+	DefaultColorSample.Specular = InvalidColor.X;
+	DefaultColorSample.Metallic = InvalidColor.X;
+	DefaultColorSample.Emissive = FVector3f(InvalidColor.X, InvalidColor.Y, InvalidColor.Z);
+	DefaultColorSample.WorldNormal = FVector4f((DefaultNormal + FVector3f::One()) * .5f, InvalidColor.W);
+
+	auto AddColorEvaluator = [&](ERenderCaptureType CaptureType)
+	{
+		TSharedPtr<FMeshGenericWorldPositionColorEvaluator> Evaluator = MakeShared<FMeshGenericWorldPositionColorEvaluator>();
+		Evaluator->DefaultColor = DefaultColorSample.GetValue4f(CaptureType);
+		Evaluator->ColorSampleFunction =
+			[&DefaultColorSample, &VisibilityFunction, SceneCapture, CaptureType](FVector3d Position, FVector3d Normal)
+		{
+			FSceneCapturePhotoSet::FSceneSample ColorSample = DefaultColorSample;
+			SceneCapture->ComputeSample(FRenderCaptureTypeFlags::Single(CaptureType), Position, Normal, VisibilityFunction, ColorSample);
+			return ColorSample.GetValue4f(CaptureType);
+		};
+		
+		int32 EvaluatorIndex = Baker.AddEvaluator(Evaluator);
+		
+		CaptureTypeToEvaluatorIndexMap.Emplace(CaptureType, EvaluatorIndex);
 	};
+
+	AddColorEvaluator(ERenderCaptureType::BaseColor);
+
+	if (Options.bUsePackedMRS)
+	{
+		AddColorEvaluator(ERenderCaptureType::CombinedMRS);
+	}
+	else
+	{
+		if (Options.bBakeRoughness)
+		{
+			AddColorEvaluator(ERenderCaptureType::Roughness);
+		}
+		if (Options.bBakeMetallic)
+		{
+			AddColorEvaluator(ERenderCaptureType::Metallic);
+		}
+		if (Options.bBakeSpecular)
+		{
+			AddColorEvaluator(ERenderCaptureType::Specular);
+		}
+	}
+
+	if (Options.bBakeEmissive)
+	{
+		AddColorEvaluator(ERenderCaptureType::Emissive);
+	}
+
+	if (Options.bBakeNormalMap) {
+		TSharedPtr<FMeshGenericWorldPositionNormalEvaluator> Evaluator = MakeShared<FMeshGenericWorldPositionNormalEvaluator>();
+		Evaluator->DefaultUnitWorldNormal = DefaultNormal;
+		Evaluator->UnitWorldNormalSampleFunction =
+			[&DefaultColorSample, &VisibilityFunction, SceneCapture](FVector3d Position, FVector3d Normal)
+		{
+			FSceneCapturePhotoSet::FSceneSample ColorSample = DefaultColorSample;
+			SceneCapture->ComputeSample(FRenderCaptureTypeFlags::WorldNormal(), Position, Normal, VisibilityFunction, ColorSample);
+			FVector3f NormalColor = ColorSample.WorldNormal;
+
+			// Map from color components [0,1] to normal components [-1,1]
+			float x = (NormalColor.X - 0.5f) * 2.0f;
+			float y = (NormalColor.Y - 0.5f) * 2.0f;
+			float z = (NormalColor.Z - 0.5f) * 2.0f;
+			return FVector3f(x, y, z);
+		};
+		
+		int32 EvaluatorIndex = Baker.AddEvaluator(Evaluator);
+		
+		CaptureTypeToEvaluatorIndexMap.Emplace(ERenderCaptureType::WorldNormal, EvaluatorIndex);
+	}
+
+	{
+		TRACE_CPUPROFILER_EVENT_SCOPE(ImageBuildersFromPhotoSet_Bake);
+		Baker.Bake();
+	}
+
+	for (TTuple<ERenderCaptureType, int32>& Item : CaptureTypeToEvaluatorIndexMap)
 	{
-		TRACE_CPUPROFILER_EVENT_SCOPE(ImageBuildersFromPhotoSet_Textures_BakeColor);
-		BaseColorBaker.Bake();
+		if (Item.Key == ERenderCaptureType::BaseColor)
+		{
+			Results->ColorImage = MoveTemp(Baker.GetBakeResults(Item.Value)[0]);
+		}
+		else if (Item.Key == ERenderCaptureType::CombinedMRS)
+		{
+			Results->PackedMRSImage = MoveTemp(Baker.GetBakeResults(Item.Value)[0]);
+		}
+		else if (Item.Key == ERenderCaptureType::Roughness)
+		{
+			Results->RoughnessImage = MoveTemp(Baker.GetBakeResults(Item.Value)[0]);
+		}
+		else if (Item.Key == ERenderCaptureType::Specular)
+		{
+			Results->SpecularImage = MoveTemp(Baker.GetBakeResults(Item.Value)[0]);
+		}
+		else if (Item.Key == ERenderCaptureType::Metallic)
+		{
+			Results->MetallicImage = MoveTemp(Baker.GetBakeResults(Item.Value)[0]);
+		}
+		else if (Item.Key == ERenderCaptureType::Emissive)
+		{
+			Results->EmissiveImage = MoveTemp(Baker.GetBakeResults(Item.Value)[0]);
+		}
+		else if (Item.Key == ERenderCaptureType::WorldNormal)
+		{
+			Results->NormalImage = MoveTemp(Baker.GetBakeResults(Item.Value)[0]);
+		}
+	}
+
+	FMeshImageBakingCache TempBakeCache;
+	{
+		TRACE_CPUPROFILER_EVENT_SCOPE(ImageBuildersFromPhotoSet_MakeMeshImageBakingCache);
+
+		TempBakeCache.SetDetailMesh(BaseMesh, &Spatial);
+		TempBakeCache.SetBakeTargetMesh(BaseMesh);
+		TempBakeCache.SetDimensions(FImageDimensions(Options.TextureImageSize, Options.TextureImageSize));
+		TempBakeCache.SetUVLayer(Options.TargetUVLayer);
+		TempBakeCache.SetThickness(0.1);
+		TempBakeCache.SetCorrespondenceStrategy(FMeshImageBakingCache::ECorrespondenceStrategy::Identity);
+		TempBakeCache.ValidateCache();
 	}
 
-	// find "hole" pixels
-	TArray<FVector2i> MissingPixels;
-	Results->ColorImage = BaseColorBaker.TakeResult();
 	TMarchingPixelInfill<FVector4f> Infill;
 
 	{
-		TRACE_CPUPROFILER_EVENT_SCOPE(ImageBuildersFromPhotoSet_Textures_ComputeInfill);
+		TRACE_CPUPROFILER_EVENT_SCOPE(ImageBuildersFromPhotoSet_FindAndComputeInfill);
+		
+		// Find missing pixels
+		TArray<FVector2i> MissingPixels;
 		TempBakeCache.FindSamplingHoles([&](const FVector2i& Coords)
 		{
 			return Results->ColorImage->GetPixel(Coords) == InvalidColor;
 		}, MissingPixels);
 
-		// solve infill for the holes while also caching infill information
+		// Solve infill for the holes while also caching infill information
 		Infill.ComputeInfill(*Results->ColorImage, MissingPixels, InvalidColor,
 			[](FVector4f SumValue, int32 Count) {
 			float InvSum = (Count == 0) ? 1.0f : (1.0f / Count);
@@ -255,112 +346,49 @@ static void ImageBuildersFromPhotoSet(
 		});
 	}
 
-	// downsample the image if necessary
-	if (Supersample > 1)
+	auto ApplyInfill = [&Infill](TUniquePtr<TImageBuilder<FVector4f>>& Image)
 	{
-		TRACE_CPUPROFILER_EVENT_SCOPE(ImageBuildersFromPhotoSet_Textures_Downsample);
-		TImageBuilder<FVector4f> Downsampled = Results->ColorImage->FastDownsample(Supersample, FVector4f::Zero(), [](FVector4f V, int N) { return V / (float)N; });
-		*Results->ColorImage = MoveTemp(Downsampled);
-	}
-
-	// this lambda is used to process the per-channel images. It does the bake, applies infill, and downsamples if necessary
-	auto ProcessChannelFunc = [&](ERenderCaptureType CaptureType)
-	{
-		FVector4f DefaultValue(0, 0, 0, 0);
-		FMeshGenericWorldPositionColorBaker ChannelBaker;
-		ChannelBaker.SetCache(&TempBakeCache);
-		ChannelBaker.ColorSampleFunction = [&](FVector3d Position, FVector3d Normal) {
-			FSceneCapturePhotoSet::FSceneSample Sample = DefaultSample;
-			SceneCapture->ComputeSample(FRenderCaptureTypeFlags::Single(CaptureType), Position, Normal, VisibilityFunction, Sample);
-			return Sample.GetValue4f(CaptureType);
-		};
-		ChannelBaker.Bake();
-		TUniquePtr<TImageBuilder<FVector4f>> Image = ChannelBaker.TakeResult();
-
 		Infill.ApplyInfill(*Image,
-			[](FVector4f SumValue, int32 Count) {
-			float InvSum = (Count == 0) ? 1.0f : (1.0f / Count);
-			return FVector4f(SumValue.X * InvSum, SumValue.Y * InvSum, SumValue.Z * InvSum, 1.0f);
-		});
-
-		if (Supersample > 1)
-		{
-			TImageBuilder<FVector4f> Downsampled = Image->FastDownsample(Supersample, FVector4f::Zero(), [](FVector4f V, int N) { return V / (float)N; });
-			*Image = MoveTemp(Downsampled);
-		}
-
-		return MoveTemp(Image);
+			[](FVector4f SumValue, int32 Count)
+			{
+				float InvSum = (Count == 0) ? 1.0f : (1.0f / Count);
+				return FVector4f(SumValue.X * InvSum, SumValue.Y * InvSum, SumValue.Z * InvSum, 1.0f);
+			});
 	};
 
-	bool bMetallic = Options.bBakeMetallic;
-	bool bRoughness = Options.bBakeRoughness;
-	bool bSpecular = Options.bBakeSpecular;
 	{
-		TRACE_CPUPROFILER_EVENT_SCOPE(ImageBuildersFromPhotoSet_Textures_OtherChannels);
+		TRACE_CPUPROFILER_EVENT_SCOPE(ImageBuildersFromPhotoSet_ApplyInfill);
+
+		// Note: Infill for Results->ColorImage was applied when the infill was computed
 
-		// if (Options.bUsePackedMRS && (bMetallic || bRoughness || bSpecular))
 		if (Options.bUsePackedMRS)
 		{
-			Results->PackedMRSImage = ProcessChannelFunc(ERenderCaptureType::CombinedMRS);
+			ApplyInfill(Results->PackedMRSImage);
 		}
 		else
 		{
-			if (bRoughness)
+			if (Options.bBakeRoughness)
 			{
-				Results->RoughnessImage = ProcessChannelFunc(ERenderCaptureType::Roughness);
+				ApplyInfill(Results->RoughnessImage);
 			}
-			if (bMetallic)
+			if (Options.bBakeMetallic)
 			{
-				Results->MetallicImage = ProcessChannelFunc(ERenderCaptureType::Metallic);
+				ApplyInfill(Results->MetallicImage);
 			}
-			if (bSpecular)
+			if (Options.bBakeSpecular)
 			{
-				Results->SpecularImage = ProcessChannelFunc(ERenderCaptureType::Specular);
+				ApplyInfill(Results->SpecularImage);
 			}
 		}
 
 		if (Options.bBakeEmissive)
 		{
-			Results->EmissiveImage = ProcessChannelFunc(ERenderCaptureType::Emissive);
+			ApplyInfill(Results->EmissiveImage);
 		}
-	}
 
-
-	//Progress.EnterProgressFrame(1.f, LOCTEXT("BakingNormals", "Baking Normals..."));
-
-	if (Options.bBakeNormalMap)
-	{
-		TRACE_CPUPROFILER_EVENT_SCOPE(ImageBuildersFromPhotoSet_Textures_NormalMapBake);
-
-		// no infill on normal map for now, doesn't make sense to do after mapping to tangent space!
+		// Note: No infill on normal map for now, doesn't make sense to do after mapping to tangent space!
 		//  (should we build baked normal map in world space, and then resample to tangent space??)
-		FVector4f DefaultNormalValue(0, 0, 1, 1);
-		FMeshGenericWorldPositionNormalBaker NormalMapBaker;
-		NormalMapBaker.SetCache(&TempBakeCache);
-		NormalMapBaker.BaseMeshTangents = MeshTangents;
-		NormalMapBaker.NormalSampleFunction = [&](FVector3d Position, FVector3d Normal) {
-			FSceneCapturePhotoSet::FSceneSample Sample = DefaultSample;
-			SceneCapture->ComputeSample(FRenderCaptureTypeFlags::WorldNormal(),
-				Position, Normal, VisibilityFunction, Sample);
-			FVector3f NormalColor = Sample.WorldNormal;
-			float x = (NormalColor.X - 0.5f) * 2.0f;
-			float y = (NormalColor.Y - 0.5f) * 2.0f;
-			float z = (NormalColor.Z - 0.5f) * 2.0f;
-			return FVector3f(x, y, z);
-		};
-
-		NormalMapBaker.Bake();
-		Results->NormalImage = NormalMapBaker.TakeResult();
-
-		if (Supersample > 1)
-		{
-			TImageBuilder<FVector3f> Downsampled = Results->NormalImage->FastDownsample(Supersample, FVector3f::Zero(), [](FVector3f V, int N) { return V / (float)N; });
-			*Results->NormalImage = MoveTemp(Downsampled);
-		}
 	}
-
-	// build textures
-	//Progress.EnterProgressFrame(1.f, LOCTEXT("BuildingTextures", "Building Textures..."));
 }
 
 //
@@ -376,17 +404,19 @@ public:
 	TSharedPtr<UE::Geometry::FMeshTangentsd, ESPMode::ThreadSafe> BaseMeshTangents;
 	FBakeRenderCaptureOptions::FOptions Options;
 	TObjectPtr<UBakeRenderCaptureInputToolProperties> InputMeshSettings;
-	FSceneCapturePhotoSet* SceneCapture;
+	TUniquePtr<FSceneCapturePhotoSet> SceneCapture;
 	
 	// Begin TGenericDataOperator interface
 	// @Incomplete Use the Progress thing
 	virtual void CalculateResult(FProgressCancel* Progress) override
 	{
 		check(Actors != nullptr);
-		check(SceneCapture != nullptr);
+		check(BaseMesh != nullptr);
+		check(BaseMeshTangents.IsValid());
+		check(SceneCapture.IsValid());
 
-		// bake textures onto the target mesh by projecting/sampling the set of captured photos
-		ImageBuildersFromPhotoSet(SceneCapture, Options, BaseMesh, BaseMeshTangents.Get(), Result);
+		// Bake textures onto the base/target mesh by projecting/sampling the set of captured photos
+		ImageBuildersFromPhotoSet(SceneCapture.Get(), Options, BaseMesh, BaseMeshTangents, Result);
 	}
 	// End TGenericDataOperator interface
 };
@@ -582,9 +612,8 @@ void UBakeRenderCaptureTool::CreateTextureAssetsRC(UWorld* SourceWorld, UObject*
 
 	if (Settings->MaterialSettings.bNormalMap && ResultSettings->NormalMap != nullptr)
 	{
-		const FString TexName = FString::Printf(TEXT("%s_%s"), *BaseName, *NormalTexParamName.ToString()); 
-		// Use ColorLinear because the baked NormalMap is in world space
-		CreateTextureAsset(TexName, FTexture2DBuilder::ETextureType::ColorLinear, ResultSettings->NormalMap);
+		const FString TexName = FString::Printf(TEXT("%s_%s"), *BaseName, *NormalTexParamName.ToString());
+		CreateTextureAsset(TexName, FTexture2DBuilder::ETextureType::NormalMap, ResultSettings->NormalMap);
 	}
 
 	if (Settings->MaterialSettings.bEmissiveMap && ResultSettings->EmissiveMap != nullptr)
@@ -695,7 +724,7 @@ TUniquePtr<UE::Geometry::TGenericDataOperator<FBakeRenderCaptureResultsBuilder>>
 	Op->BaseMesh = &Tool->TargetMesh;
 	Op->BaseMeshTangents = Tool->TargetMeshTangents;
 	Op->Options = FBakeRenderCaptureOptions::ConstructOptions(*Tool->Settings, *Tool->InputMeshSettings);
-	Op->SceneCapture = Tool->SceneCapture.Get();
+	Op->SceneCapture = MoveTemp(Tool->SceneCapture);
 	return Op;
 }
 
@@ -703,9 +732,6 @@ TUniquePtr<UE::Geometry::TGenericDataOperator<FBakeRenderCaptureResultsBuilder>>
 void UBakeRenderCaptureTool::OnMapsUpdatedRC(const TUniquePtr<FBakeRenderCaptureResultsBuilder>& NewResult)
 {
 	FBakeRenderCaptureOptions::FOptions Options = FBakeRenderCaptureOptions::ConstructOptions(*Settings, *InputMeshSettings);
-	bool bMetallic = Options.bBakeMetallic;
-	bool bRoughness = Options.bBakeRoughness;
-	bool bSpecular = Options.bBakeSpecular;
 
 	check(IsInGameThread());
 
@@ -726,26 +752,24 @@ void UBakeRenderCaptureTool::OnMapsUpdatedRC(const TUniquePtr<FBakeRenderCapture
 		}
 		if (Options.bBakeNormalMap && NewResult->NormalImage.IsValid())
 		{
-			// Use ColorLinear because the baked NormalMap is in world space
-			ResultSettings->NormalMap = FTexture2DBuilder::BuildTextureFromImage(*NewResult->NormalImage, FTexture2DBuilder::ETextureType::ColorLinear, false, bPopulateSourceData);
+			ResultSettings->NormalMap = FTexture2DBuilder::BuildTextureFromImage(*NewResult->NormalImage, FTexture2DBuilder::ETextureType::NormalMap, false, bPopulateSourceData);
 		}
 
-		// if (Options.bUsePackedMRS && (bRoughness || bMetallic || bSpecular) && NewResult->PackedMRSImage.IsValid())
 		if (Options.bUsePackedMRS && NewResult->PackedMRSImage.IsValid())
 		{
 			ResultSettings->PackedMRSMap = FTexture2DBuilder::BuildTextureFromImage(*NewResult->PackedMRSImage, FTexture2DBuilder::ETextureType::ColorLinear, false, bPopulateSourceData);
 		}
 		else
 		{ 
-			if (bRoughness && NewResult->RoughnessImage.IsValid())
+			if (Options.bBakeRoughness && NewResult->RoughnessImage.IsValid())
 			{
 				ResultSettings->RoughnessMap = FTexture2DBuilder::BuildTextureFromImage(*NewResult->RoughnessImage, FTexture2DBuilder::ETextureType::Roughness, false, bPopulateSourceData);
 			}
-			if (bMetallic && NewResult->MetallicImage.IsValid())
+			if (Options.bBakeMetallic && NewResult->MetallicImage.IsValid())
 			{
 				ResultSettings->MetallicMap = FTexture2DBuilder::BuildTextureFromImage(*NewResult->MetallicImage, FTexture2DBuilder::ETextureType::Metallic, false, bPopulateSourceData);
 			}
-			if (bSpecular && NewResult->SpecularImage.IsValid())
+			if (Options.bBakeSpecular && NewResult->SpecularImage.IsValid())
 			{
 				ResultSettings->SpecularMap = FTexture2DBuilder::BuildTextureFromImage(*NewResult->SpecularImage, FTexture2DBuilder::ETextureType::Specular, false, bPopulateSourceData);
 			}
@@ -897,10 +921,7 @@ void UBakeRenderCaptureTool::UpdateVisualization()
 
 	FBakeRenderCaptureOptions::FOptions Options = FBakeRenderCaptureOptions::ConstructOptions(*Settings, *InputMeshSettings);
 	
-	// bool bPackedMRS = Options.bUsePackedMRS && (Options.bBakeRoughness || Options.bBakeMetallic || Options.bBakeSpecular);
-	bool bPackedMRS = Options.bUsePackedMRS;
-
-	if (bPackedMRS)
+	if (Options.bUsePackedMRS)
 	{
 		TObjectPtr<UMaterialInstanceDynamic> Material = PreviewMaterialPackedRC;
 		PreviewMesh->SetOverrideRenderMaterial(Material);
diff --git a/Engine/Plugins/Experimental/MeshModelingToolsetExp/Source/MeshModelingToolsExp/Public/BakeRenderCaptureTool.h b/Engine/Plugins/Experimental/MeshModelingToolsetExp/Source/MeshModelingToolsExp/Public/BakeRenderCaptureTool.h
--- a/Engine/Plugins/Experimental/MeshModelingToolsetExp/Source/MeshModelingToolsExp/Public/BakeRenderCaptureTool.h
+++ b/Engine/Plugins/Experimental/MeshModelingToolsetExp/Source/MeshModelingToolsExp/Public/BakeRenderCaptureTool.h
@@ -58,7 +58,7 @@ struct FBakeRenderCaptureResultsBuilder
 	TUniquePtr<TImageBuilder<FVector4f>> PackedMRSImage; 
 	TUniquePtr<TImageBuilder<FVector4f>> EmissiveImage;
 	TUniquePtr<TImageBuilder<FVector4f>> ColorImage;
-	TUniquePtr<TImageBuilder<FVector3f>> NormalImage;
+	TUniquePtr<TImageBuilder<FVector4f>> NormalImage; // Tangent-space normal map
 };
 
 //
@@ -112,7 +112,7 @@ struct FMaterialProxySettingsRC
 	UPROPERTY(Category = Material, EditAnywhere, meta=(EditCondition="bPackedMRSMap == false"))
 	uint8 bSpecularMap:1;
 	
-	// Whether to generate a packed texture for the Metallic/Roughness/Specular properties
+	// Whether to generate a packed texture with Metallic, Roughness and Specular properties
 	UPROPERTY(Category = Material, EditAnywhere)
 	uint8 bPackedMRSMap:1;
 
@@ -172,8 +172,10 @@ public:
 	// Material Baking Settings
 	//
 
-	/** If Value is > 1, Multisample output baked textures by this amount in each direction (eg 4 == 16x supersampling) */
-	UPROPERTY(EditAnywhere, Category = MaterialSettings, meta = (ClampMin = "0", ClampMax = "8", UIMin = "0", UIMax = "4"))
+	/** If Value is > 1, Multisample output baked textures by this amount in each direction (eg 4 == 16x supersampling).
+	 *  WIP feature currently disabled
+	 */
+	UPROPERTY(EditAnywhere, Category = MaterialSettings, meta = (ClampMin = "0", ClampMax = "8", UIMin = "0", UIMax = "4", EditCondition="false"))
 	int32 MultiSamplingAA = 0;
 
 	/** If Value is zero, use MaterialSettings resolution, otherwise override the render capture resolution */
diff --git a/Engine/Plugins/Runtime/GeometryProcessing/Source/DynamicMesh/Private/Sampling/MeshGenericWorldPositionEvaluator.cpp b/Engine/Plugins/Runtime/GeometryProcessing/Source/DynamicMesh/Private/Sampling/MeshGenericWorldPositionEvaluator.cpp
new file mode 100644
--- /dev/null
+++ b/Engine/Plugins/Runtime/GeometryProcessing/Source/DynamicMesh/Private/Sampling/MeshGenericWorldPositionEvaluator.cpp
@@ -0,0 +1,111 @@
+﻿// Copyright Epic Games, Inc. All Rights Reserved.
+
+#include "Sampling/MeshGenericWorldPositionEvaluator.h"
+#include "Sampling/MeshBaseBaker.h"
+
+#include "VectorTypes.h"
+
+using namespace UE::Geometry;
+
+void FMeshGenericWorldPositionColorEvaluator::Setup(const FMeshBaseBaker& Baker, FEvaluationContext& Context)
+{
+	Context.Evaluate = &EvaluateSample;
+	Context.EvaluateDefault = &EvaluateDefault;
+	Context.EvaluateColor = &EvaluateColor;
+	Context.EvalData = this;
+	Context.AccumulateMode = EAccumulateMode::Add;
+	Context.DataLayout = { EComponents::Float4 };
+}
+
+void FMeshGenericWorldPositionColorEvaluator::EvaluateSample(float*& Out, const FCorrespondenceSample& Sample, void* EvalData)
+{
+	FMeshGenericWorldPositionColorEvaluator* Eval = static_cast<FMeshGenericWorldPositionColorEvaluator*>(EvalData);
+	
+	const FVector3d Position = Sample.BaseSample.SurfacePoint;
+	const FVector3d Normal = Sample.BaseNormal;
+	const FVector4f SampleResult = Eval->ColorSampleFunction(Position, Normal);
+	
+	WriteToBuffer(Out, SampleResult);
+}
+
+void FMeshGenericWorldPositionColorEvaluator::EvaluateDefault(float*& Out, void* EvalData)
+{
+	FMeshGenericWorldPositionColorEvaluator* Eval = static_cast<FMeshGenericWorldPositionColorEvaluator*>(EvalData);
+	WriteToBuffer(Out, Eval->DefaultColor);
+}
+
+void FMeshGenericWorldPositionColorEvaluator::EvaluateColor(const int DataIdx, float*& In, FVector4f& Out, void* EvalData)
+{
+	Out = FVector4f(In[0], In[1], In[2], In[3]);
+	In += 4;
+}
+
+
+
+
+
+
+
+
+
+
+void FMeshGenericWorldPositionNormalEvaluator::Setup(const FMeshBaseBaker& Baker, FEvaluationContext& Context)
+{
+	BaseMeshTangents = Baker.GetTargetMeshTangents();
+	
+	// Make sure we have per-triangle tangents computed so that GetInterpolatedTriangleTangent() below works
+	check(BaseMeshTangents != nullptr);
+	check(BaseMeshTangents->GetTangents().Num() >= 3 * Baker.GetTargetMesh()->MaxTriangleID());
+
+	Context.Evaluate = &EvaluateSample;
+	Context.EvaluateDefault = &EvaluateDefault;
+	Context.EvaluateColor = &EvaluateColor;
+	Context.EvalData = this;
+	Context.AccumulateMode = EAccumulateMode::Add;
+	Context.DataLayout = { EComponents::Float3 };
+}
+
+void FMeshGenericWorldPositionNormalEvaluator::EvaluateSample(float*& Out, const FCorrespondenceSample& Sample, void* EvalData)
+{
+	const FMeshGenericWorldPositionNormalEvaluator* Eval = static_cast<FMeshGenericWorldPositionNormalEvaluator*>(EvalData);
+	const FVector3f SampleResult = Eval->EvaluateSampleImpl(Sample);
+	WriteToBuffer(Out, SampleResult);
+}
+
+void FMeshGenericWorldPositionNormalEvaluator::EvaluateDefault(float*& Out, void* EvalData)
+{
+	FMeshGenericWorldPositionNormalEvaluator* Eval = static_cast<FMeshGenericWorldPositionNormalEvaluator*>(EvalData);
+	WriteToBuffer(Out, Eval->DefaultUnitWorldNormal);
+}
+
+void FMeshGenericWorldPositionNormalEvaluator::EvaluateColor(const int DataIdx, float*& In, FVector4f& Out, void* EvalData)
+{
+	// Map normal space [-1,1] to color space [0,1]
+	const FVector3f Normal(In[0], In[1], In[2]);
+	const FVector3f Color = (Normal + FVector3f::One()) * 0.5f;
+	Out = FVector4f(Color.X, Color.Y, Color.Z, 1.0f);
+	In += 3;
+}
+
+FVector3f FMeshGenericWorldPositionNormalEvaluator::EvaluateSampleImpl(const FCorrespondenceSample& SampleData) const
+{
+	FVector3d Position = SampleData.BaseSample.SurfacePoint;
+	FVector3d Normal = SampleData.BaseNormal;
+	int32 TriangleID = SampleData.BaseSample.TriangleIndex;
+	FVector3d BaryCoords = SampleData.BaseSample.BaryCoords;
+
+	// Get tangents on base mesh
+	FVector3d BaseTangentX, BaseTangentY;
+	BaseMeshTangents->GetInterpolatedTriangleTangent(TriangleID, BaryCoords, BaseTangentX, BaseTangentY);
+
+	// Sample world normal at world position
+	FVector3f WorldSpaceNormal = UnitWorldNormalSampleFunction(Position, Normal);
+
+	// Compute normal in tangent space
+	FVector3f TangentSpaceNormal(
+			(float)WorldSpaceNormal.Dot(FVector3f(BaseTangentX)),
+			(float)WorldSpaceNormal.Dot(FVector3f(BaseTangentY)),
+			(float)WorldSpaceNormal.Dot(FVector3f(SampleData.BaseNormal)));
+
+	return TangentSpaceNormal;
+}
diff --git a/Engine/Plugins/Runtime/GeometryProcessing/Source/DynamicMesh/Private/Sampling/MeshOcclusionMapEvaluator.cpp b/Engine/Plugins/Runtime/GeometryProcessing/Source/DynamicMesh/Private/Sampling/MeshOcclusionMapEvaluator.cpp
--- a/Engine/Plugins/Runtime/GeometryProcessing/Source/DynamicMesh/Private/Sampling/MeshOcclusionMapEvaluator.cpp
+++ b/Engine/Plugins/Runtime/GeometryProcessing/Source/DynamicMesh/Private/Sampling/MeshOcclusionMapEvaluator.cpp
@@ -278,3 +278,6 @@ FMeshOcclusionMapEvaluator::FOcclusionTuple FMeshOcclusionMapEvaluator::SampleFu
 	return FOcclusionTuple(AccumOcclusion, AccumNormal);
 }
 
+#undef WANT_AMBIENT_OCCLUSION
+#undef WANT_BENT_NORMAL
+#undef WANT_ALL
diff --git a/Engine/Plugins/Runtime/GeometryProcessing/Source/DynamicMesh/Public/Sampling/MeshGenericWorldPositionEvaluator.h b/Engine/Plugins/Runtime/GeometryProcessing/Source/DynamicMesh/Public/Sampling/MeshGenericWorldPositionEvaluator.h
new file mode 100644
--- /dev/null
+++ b/Engine/Plugins/Runtime/GeometryProcessing/Source/DynamicMesh/Public/Sampling/MeshGenericWorldPositionEvaluator.h
@@ -0,0 +1,79 @@
+﻿// Copyright Epic Games, Inc. All Rights Reserved.
+
+#pragma once
+
+#include "Sampling/MeshMapEvaluator.h"
+#include "DynamicMesh/MeshTangents.h"
+#include "Image/ImageBuilder.h"
+
+namespace UE
+{
+namespace Geometry
+{
+
+/**
+* Bake Colors based on arbitrary Position/Normal sampling of the base/target mesh.
+* Useful for (eg) baking from an arbitrary function defined over space
+*/
+class DYNAMICMESH_API FMeshGenericWorldPositionColorEvaluator : public FMeshMapEvaluator
+{
+public:
+
+	FVector4f DefaultColor = FVector4f(0, 0, 0, 1);
+	TFunction<FVector4f(FVector3d, FVector3d)> ColorSampleFunction 
+		= [this](FVector3d BaseSurfacePoint, FVector3d BaseSurfaceNormal) { return DefaultColor; };
+
+public:
+	virtual ~FMeshGenericWorldPositionColorEvaluator() {}
+
+	// Begin FMeshMapEvaluator interface
+	virtual void Setup(const FMeshBaseBaker& Baker, FEvaluationContext& Context) override;
+
+	virtual EMeshMapEvaluatorType Type() const override { return EMeshMapEvaluatorType::GenericWorldPositionColor; }
+	// End FMeshMapEvaluator interface
+
+	static void EvaluateSample(float*& Out, const FCorrespondenceSample& Sample, void* EvalData);
+
+	static void EvaluateDefault(float*& Out, void* EvalData);
+
+	static void EvaluateColor(const int DataIdx, float*& In, FVector4f& Out, void* EvalData);
+};
+
+
+
+/**
+ * Bake Tangent-Space Normals based on arbitrary Position/Normal sampling of the base/target mesh.
+ * Useful for (eg) baking from an arbitrary function defined over space.
+ */
+class DYNAMICMESH_API FMeshGenericWorldPositionNormalEvaluator : public FMeshMapEvaluator
+{
+public:
+
+	FVector3f DefaultUnitWorldNormal = FVector3f::UnitZ();
+	TFunction<FVector3f(FVector3d, FVector3d)> UnitWorldNormalSampleFunction 
+		= [this](FVector3d Position, FVector3d Normal) { return DefaultUnitWorldNormal; };
+
+public:
+	virtual ~FMeshGenericWorldPositionNormalEvaluator() {}
+
+	// Begin FMeshMapEvaluator interface
+	virtual void Setup(const FMeshBaseBaker& Baker, FEvaluationContext& Context) override;
+
+	virtual EMeshMapEvaluatorType Type() const override { return EMeshMapEvaluatorType::GenericWorldPositionNormal; }
+	// End FMeshMapEvaluator interface
+
+	static void EvaluateSample(float*& Out, const FCorrespondenceSample& Sample, void* EvalData);
+
+	static void EvaluateDefault(float*& Out, void* EvalData);
+
+	static void EvaluateColor(const int DataIdx, float*& In, FVector4f& Out, void* EvalData);
+
+private:
+
+	FVector3f EvaluateSampleImpl(const FCorrespondenceSample& SampleData) const;
+	const FMeshTangentsd* BaseMeshTangents = nullptr;
+};
+
+
+} // end namespace UE::Geometry
+} // end namespace UE
\ No newline at end of file
diff --git a/Engine/Plugins/Runtime/GeometryProcessing/Source/DynamicMesh/Public/Sampling/MeshMapEvaluator.h b/Engine/Plugins/Runtime/GeometryProcessing/Source/DynamicMesh/Public/Sampling/MeshMapEvaluator.h
--- a/Engine/Plugins/Runtime/GeometryProcessing/Source/DynamicMesh/Public/Sampling/MeshMapEvaluator.h
+++ b/Engine/Plugins/Runtime/GeometryProcessing/Source/DynamicMesh/Public/Sampling/MeshMapEvaluator.h
@@ -19,7 +19,9 @@ enum class EMeshMapEvaluatorType
 	Curvature,
 	Property,
 	ResampleImage,
-	MultiResampleImage
+	MultiResampleImage,
+	GenericWorldPositionColor,
+	GenericWorldPositionNormal
 };
 
 class FMeshMapEvaluator
