ModelingTools: Add SampleFilterMask support for BakeTexture/BakeAll.

diff --git a/Engine/Plugins/Experimental/GeometryScripting/Source/GeometryScriptingCore/Private/MeshBakeFunctions.cpp b/Engine/Plugins/Experimental/GeometryScripting/Source/GeometryScriptingCore/Private/MeshBakeFunctions.cpp
--- a/Engine/Plugins/Experimental/GeometryScripting/Source/GeometryScriptingCore/Private/MeshBakeFunctions.cpp
+++ b/Engine/Plugins/Experimental/GeometryScripting/Source/GeometryScriptingCore/Private/MeshBakeFunctions.cpp
@@ -228,6 +228,24 @@ namespace GeometryScriptBakeLocals
 		return true;
 	};
 
+	TUniquePtr<TImageBuilder<FVector4f>> GetSampleFilterMask(
+		const FGeometryScriptBakeTextureOptions& Options,
+		const FText& DebugPrefix,
+		TArray<FGeometryScriptDebugMessage>* Debug)
+	{
+		if (Options.SampleFilterMask)
+		{
+			TUniquePtr<TImageBuilder<FVector4f>> Result = MakeUnique<TImageBuilder<FVector4f>>();
+			if (!UE::AssetUtils::ReadTexture(Options.SampleFilterMask, *Result, true))
+			{
+				UE::Geometry::AppendError(Debug, EGeometryScriptErrorType::InvalidInputs, FText::Format(LOCTEXT("Bake_InvalidSampleFilterMask", "{0}: Failed to read SampleFilterMask"), DebugPrefix));
+				return nullptr;
+			}
+			return Result;
+		}
+		return nullptr;
+	}
+
 	struct FEvaluatorState
 	{
 		const FDynamicMesh3* TargetMesh = nullptr;
@@ -495,6 +513,15 @@ namespace GeometryScriptBakeLocals
 		Baker.SetDimensions(BakeDimensions);
 		Baker.SetProjectionDistance(BakeOptions.ProjectionDistance);
 		Baker.SetSamplesPerPixel(GetSamplesPerPixel(BakeOptions.SamplesPerPixel));
+		TUniquePtr<TImageBuilder<FVector4f>> SampleFilterMask = GetSampleFilterMask(BakeOptions, BakeTexturePrefix, Debug);
+		if (SampleFilterMask)
+		{
+			Baker.SampleFilterF = [&SampleFilterMask](const FVector2i& ImageCoords, const FVector2d& UV, int32 TriID)
+			{
+				const FVector4f Color = SampleFilterMask->BilinearSampleUV<float>(UV, FVector4f(0, 0, 0, 1));
+				return Color.IsNearlyZero3();
+			};
+		}
 		if (bIsBakeToSelf)
 		{
 			Baker.SetCorrespondenceStrategy(FMeshBaseBaker::ECorrespondenceStrategy::Identity);
diff --git a/Engine/Plugins/Experimental/GeometryScripting/Source/GeometryScriptingCore/Public/GeometryScript/MeshBakeFunctions.h b/Engine/Plugins/Experimental/GeometryScripting/Source/GeometryScriptingCore/Public/GeometryScript/MeshBakeFunctions.h
--- a/Engine/Plugins/Experimental/GeometryScripting/Source/GeometryScriptingCore/Public/GeometryScript/MeshBakeFunctions.h
+++ b/Engine/Plugins/Experimental/GeometryScripting/Source/GeometryScriptingCore/Public/GeometryScript/MeshBakeFunctions.h
@@ -210,6 +210,10 @@ struct GEOMETRYSCRIPTINGCORE_API FGeometryScriptBakeTextureOptions
 	UPROPERTY(BlueprintReadWrite, Category = Options)
 	EGeometryScriptBakeSamplesPerPixel SamplesPerPixel = EGeometryScriptBakeSamplesPerPixel::Sample1;
 
+	/** Mask texture for filtering out samples/pixels from the output texture */
+	UPROPERTY(BlueprintReadWrite, Category = Options)
+	TObjectPtr<UTexture2D> SampleFilterMask = nullptr;
+
 	/** Maximum allowed distance for the projection from target mesh to source mesh for the sample to be considered valid.
 	 * This is only relevant if a separate source mesh is provided. */
 	UPROPERTY(BlueprintReadWrite, Category = Options)
diff --git a/Engine/Plugins/Experimental/MeshModelingToolsetExp/Source/MeshModelingToolsExp/Private/BakeMeshAttributeMapsTool.cpp b/Engine/Plugins/Experimental/MeshModelingToolsetExp/Source/MeshModelingToolsExp/Private/BakeMeshAttributeMapsTool.cpp
--- a/Engine/Plugins/Experimental/MeshModelingToolsetExp/Source/MeshModelingToolsExp/Private/BakeMeshAttributeMapsTool.cpp
+++ b/Engine/Plugins/Experimental/MeshModelingToolsetExp/Source/MeshModelingToolsExp/Private/BakeMeshAttributeMapsTool.cpp
@@ -104,6 +104,7 @@ public:
 	TSharedPtr<UE::Geometry::TMeshTangents<double>, ESPMode::ThreadSafe> BaseMeshTangents;
 	TSharedPtr<TArray<int32>, ESPMode::ThreadSafe> BaseMeshUVCharts;
 	bool bIsBakeToSelf = false;
+	ImagePtr SampleFilterMask;
 
 	// Map Type settings
 	EBakeMapType Maps;
@@ -141,6 +142,14 @@ public:
 		{
 			Baker->SetCorrespondenceStrategy(FMeshBaseBaker::ECorrespondenceStrategy::Identity);
 		}
+		if (SampleFilterMask)
+		{
+			Baker->SampleFilterF = [this](const FVector2i& ImageCoords, const FVector2d& UV, int32 TriID)
+			{
+				const FVector4f Color = SampleFilterMask->BilinearSampleUV<float>(UV, FVector4f(0, 0, 0, 1));
+				return Color.IsNearlyZero3();
+			};
+		}
 		
 		FMeshBakerDynamicMeshSampler DetailSampler(DetailMesh.Get(), DetailSpatial.Get(), DetailMeshTangents.Get());
 		Baker->SetDetailSampler(&DetailSampler);
@@ -293,6 +302,7 @@ void UBakeMeshAttributeMapsTool::Setup()
 	Settings->WatchProperty(Settings->Resolution, [this](EBakeTextureResolution) { OpState |= EBakeOpState::Evaluate; });
 	Settings->WatchProperty(Settings->BitDepth, [this](EBakeTextureBitDepth) { OpState |= EBakeOpState::Evaluate; });
 	Settings->WatchProperty(Settings->SamplesPerPixel, [this](EBakeTextureSamplesPerPixel) { OpState |= EBakeOpState::Evaluate; });
+	Settings->WatchProperty(Settings->SampleFilterMask, [this](UTexture2D*){ OpState |= EBakeOpState::Evaluate; });
 	
 
 	InputMeshSettings = NewObject<UBakeInputMeshProperties>(this);
@@ -419,6 +429,7 @@ TUniquePtr<UE::Geometry::TGenericDataOperator<FMeshMapBaker>> UBakeMeshAttribute
 	Op->BakeSettings = CachedBakeSettings;
 	Op->BaseMeshUVCharts = TargetMeshUVCharts;
 	Op->bIsBakeToSelf = bIsBakeToSelf;
+	Op->SampleFilterMask = CachedSampleFilterMask;
 
 	constexpr EBakeMapType RequiresTangents = EBakeMapType::TangentSpaceNormal | EBakeMapType::BentNormal;
 	if ((bool)(CachedBakeSettings.BakeMapTypes & RequiresTangents))
@@ -620,6 +631,8 @@ void UBakeMeshAttributeMapsTool::UpdateResult()
 
 	OpState |= UpdateResult_DetailNormalMap();
 
+	OpState |= UpdateResult_SampleFilterMask(Settings->SampleFilterMask);
+
 	// Update map type settings
 	if ((bool)(CachedBakeSettings.BakeMapTypes & EBakeMapType::TangentSpaceNormal))
 	{
diff --git a/Engine/Plugins/Experimental/MeshModelingToolsetExp/Source/MeshModelingToolsExp/Private/BakeMeshAttributeMapsToolBase.cpp b/Engine/Plugins/Experimental/MeshModelingToolsetExp/Source/MeshModelingToolsExp/Private/BakeMeshAttributeMapsToolBase.cpp
--- a/Engine/Plugins/Experimental/MeshModelingToolsetExp/Source/MeshModelingToolsExp/Private/BakeMeshAttributeMapsToolBase.cpp
+++ b/Engine/Plugins/Experimental/MeshModelingToolsetExp/Source/MeshModelingToolsExp/Private/BakeMeshAttributeMapsToolBase.cpp
@@ -141,6 +141,29 @@ void UBakeMeshAttributeMapsToolBase::UpdateResult()
 }
 
 
+EBakeOpState UBakeMeshAttributeMapsToolBase::UpdateResult_SampleFilterMask(UTexture2D* SampleFilterMask)
+{
+	EBakeOpState ResultState = EBakeOpState::Clean;
+	if (SampleFilterMask)
+	{
+		CachedSampleFilterMask = MakeShared<UE::Geometry::TImageBuilder<FVector4f>, ESPMode::ThreadSafe>();
+		if (!UE::AssetUtils::ReadTexture(SampleFilterMask, *CachedSampleFilterMask, bPreferPlatformData))
+		{
+			GetToolManager()->DisplayMessage(LOCTEXT("CannotReadTextureWarning", "Cannot read from the sample filter mask"), EToolMessageLevel::UserWarning);
+			return EBakeOpState::Invalid;
+		}
+		ResultState = EBakeOpState::Evaluate;
+	}
+	else if (CachedSampleFilterMask)
+	{
+		// Clear CachedSampleFilterMask if SampleFilterMask is null and re-evaluate.
+		CachedSampleFilterMask.Reset();
+		ResultState = EBakeOpState::Evaluate;
+	}
+	return ResultState;
+}
+
+
 void UBakeMeshAttributeMapsToolBase::UpdateVisualization()
 {
 }
diff --git a/Engine/Plugins/Experimental/MeshModelingToolsetExp/Source/MeshModelingToolsExp/Private/BakeMultiMeshAttributeMapsTool.cpp b/Engine/Plugins/Experimental/MeshModelingToolsetExp/Source/MeshModelingToolsExp/Private/BakeMultiMeshAttributeMapsTool.cpp
--- a/Engine/Plugins/Experimental/MeshModelingToolsetExp/Source/MeshModelingToolsExp/Private/BakeMultiMeshAttributeMapsTool.cpp
+++ b/Engine/Plugins/Experimental/MeshModelingToolsetExp/Source/MeshModelingToolsExp/Private/BakeMultiMeshAttributeMapsTool.cpp
@@ -331,6 +331,7 @@ public:
 	TUniquePtr<UE::Geometry::FMeshMapBaker> Baker;
 	UBakeMultiMeshAttributeMapsTool::FBakeSettings BakeSettings;
 	TSharedPtr<TArray<int32>, ESPMode::ThreadSafe> BaseMeshUVCharts;
+	TSharedPtr<UE::Geometry::TImageBuilder<FVector4f>> SampleFilterMask;
 
 	// Detail bake data
 	TArray<TSharedPtr<UE::Geometry::TImageBuilder<FVector4f>>> CachedColorImages;
@@ -350,6 +351,14 @@ public:
 		Baker->SetSamplesPerPixel(BakeSettings.SamplesPerPixel);
 		Baker->SetTargetMeshTangents(BaseMeshTangents);
 		Baker->SetTargetMeshUVCharts(BaseMeshUVCharts.Get());
+		if (SampleFilterMask)
+		{
+			Baker->SampleFilterF = [this](const FVector2i& ImageCoords, const FVector2d& UV, int32 TriID)
+			{
+				const FVector4f Color = SampleFilterMask->BilinearSampleUV<float>(UV, FVector4f(0, 0, 0, 1));
+				return Color.IsNearlyZero3();
+			};
+		}
 		
 		FMeshBakerMeshSceneSampler DetailSampler(DetailMeshScene.Get());
 		Baker->SetDetailSampler(&DetailSampler);
@@ -433,6 +442,7 @@ void UBakeMultiMeshAttributeMapsTool::Setup()
 	Settings->WatchProperty(Settings->Resolution, [this](EBakeTextureResolution) { OpState |= EBakeOpState::Evaluate; });
 	Settings->WatchProperty(Settings->BitDepth, [this](EBakeTextureBitDepth) { OpState |= EBakeOpState::Evaluate; });
 	Settings->WatchProperty(Settings->SamplesPerPixel, [this](EBakeTextureSamplesPerPixel) { OpState |= EBakeOpState::Evaluate; });
+	Settings->WatchProperty(Settings->SampleFilterMask, [this](UTexture2D*){ OpState |= EBakeOpState::Evaluate; });
 	
 	InputMeshSettings = NewObject<UBakeMultiMeshInputToolProperties>(this);
 	InputMeshSettings->RestoreProperties(this);
@@ -519,6 +529,7 @@ TUniquePtr<UE::Geometry::TGenericDataOperator<FMeshMapBaker>> UBakeMultiMeshAttr
 	Op->BaseMesh = &TargetMesh;
 	Op->BaseMeshUVCharts = TargetMeshUVCharts;
 	Op->BakeSettings = CachedBakeSettings;
+	Op->SampleFilterMask = CachedSampleFilterMask;
 
 	constexpr EBakeMapType RequiresTangents = EBakeMapType::TangentSpaceNormal | EBakeMapType::BentNormal;
 	if (static_cast<bool>(CachedBakeSettings.BakeMapTypes & RequiresTangents))
@@ -605,6 +616,8 @@ void UBakeMultiMeshAttributeMapsTool::UpdateResult()
 
 	OpState |= UpdateResult_TargetMeshTangents(CachedBakeSettings.BakeMapTypes);
 
+	OpState |= UpdateResult_SampleFilterMask(Settings->SampleFilterMask);
+
 	// Update map type settings
 	OpState |= UpdateResult_DetailMeshes();
 
diff --git a/Engine/Plugins/Experimental/MeshModelingToolsetExp/Source/MeshModelingToolsExp/Public/BakeMeshAttributeMapsTool.h b/Engine/Plugins/Experimental/MeshModelingToolsetExp/Source/MeshModelingToolsExp/Public/BakeMeshAttributeMapsTool.h
--- a/Engine/Plugins/Experimental/MeshModelingToolsetExp/Source/MeshModelingToolsExp/Public/BakeMeshAttributeMapsTool.h
+++ b/Engine/Plugins/Experimental/MeshModelingToolsetExp/Source/MeshModelingToolsExp/Public/BakeMeshAttributeMapsTool.h
@@ -73,6 +73,10 @@ public:
 	UPROPERTY(EditAnywhere, Category = Textures)
 	EBakeTextureSamplesPerPixel SamplesPerPixel = EBakeTextureSamplesPerPixel::Sample1;
 
+	/** Mask texture for filtering out samples/pixels from the output texture */
+	UPROPERTY(EditAnywhere, Category = Textures, AdvancedDisplay)
+	TObjectPtr<UTexture2D> SampleFilterMask = nullptr;
+
 	UFUNCTION()
 	const TArray<FString>& GetMapPreviewNamesFunc();
 	UPROPERTY(meta = (TransientToolProperty))
diff --git a/Engine/Plugins/Experimental/MeshModelingToolsetExp/Source/MeshModelingToolsExp/Public/BakeMeshAttributeMapsToolBase.h b/Engine/Plugins/Experimental/MeshModelingToolsetExp/Source/MeshModelingToolsExp/Public/BakeMeshAttributeMapsToolBase.h
--- a/Engine/Plugins/Experimental/MeshModelingToolsetExp/Source/MeshModelingToolsExp/Public/BakeMeshAttributeMapsToolBase.h
+++ b/Engine/Plugins/Experimental/MeshModelingToolsetExp/Source/MeshModelingToolsExp/Public/BakeMeshAttributeMapsToolBase.h
@@ -176,6 +176,9 @@ protected:
 	};
 	FBakeSettings CachedBakeSettings;
 
+	EBakeOpState UpdateResult_SampleFilterMask(UTexture2D* SampleFilterMask);
+	TSharedPtr<UE::Geometry::TImageBuilder<FVector4f>, ESPMode::ThreadSafe> CachedSampleFilterMask;
+
 	TSharedPtr<TArray<int32>, ESPMode::ThreadSafe> TargetMeshUVCharts;
 
 	/**
diff --git a/Engine/Plugins/Experimental/MeshModelingToolsetExp/Source/MeshModelingToolsExp/Public/BakeMultiMeshAttributeMapsTool.h b/Engine/Plugins/Experimental/MeshModelingToolsetExp/Source/MeshModelingToolsExp/Public/BakeMultiMeshAttributeMapsTool.h
--- a/Engine/Plugins/Experimental/MeshModelingToolsetExp/Source/MeshModelingToolsExp/Public/BakeMultiMeshAttributeMapsTool.h
+++ b/Engine/Plugins/Experimental/MeshModelingToolsetExp/Source/MeshModelingToolsExp/Public/BakeMultiMeshAttributeMapsTool.h
@@ -60,6 +60,10 @@ public:
 	UPROPERTY(EditAnywhere, Category = Textures)
 	EBakeTextureSamplesPerPixel SamplesPerPixel = EBakeTextureSamplesPerPixel::Sample1;
 
+	/** Mask texture for filtering out samples/pixels from the output texture */
+	UPROPERTY(EditAnywhere, Category = Textures, AdvancedDisplay)
+	TObjectPtr<UTexture2D> SampleFilterMask = nullptr;
+
 	UFUNCTION()
 	const TArray<FString>& GetMapPreviewNamesFunc()
 	{
diff --git a/Engine/Plugins/Runtime/GeometryProcessing/Source/DynamicMesh/Private/Sampling/MeshMapBaker.cpp b/Engine/Plugins/Runtime/GeometryProcessing/Source/DynamicMesh/Private/Sampling/MeshMapBaker.cpp
--- a/Engine/Plugins/Runtime/GeometryProcessing/Source/DynamicMesh/Private/Sampling/MeshMapBaker.cpp
+++ b/Engine/Plugins/Runtime/GeometryProcessing/Source/DynamicMesh/Private/Sampling/MeshMapBaker.cpp
@@ -512,9 +512,19 @@ void FMeshMapBaker::BakeSample(
 	float* Buffer = static_cast<float*>(FMemory_Alloca(sizeof(float) * BakeSampleBufferSize));
 	float* BufferPtr = Buffer;
 	const int32 NumBakers = Bakers.Num();
-	for (int32 Idx = 0; Idx < NumBakers; ++Idx)
+	if (SampleFilterF && SampleFilterF(ImageCoords, UVPosition, Sample.DetailTriID))
+	{
+		for (int32 Idx = 0; Idx < NumBakers; ++Idx)
+		{
+			BakeContexts[Idx].EvaluateDefault(BufferPtr, BakeContexts[Idx].EvalData);
+		}
+	}
+	else
 	{
-		BakeContexts[Idx].Evaluate(BufferPtr, Sample, BakeContexts[Idx].EvalData);
+		for (int32 Idx = 0; Idx < NumBakers; ++Idx)
+		{
+			BakeContexts[Idx].Evaluate(BufferPtr, Sample, BakeContexts[Idx].EvalData);
+		}
 	}
 	checkSlow((BufferPtr - Buffer) == BakeSampleBufferSize);
 
diff --git a/Engine/Plugins/Runtime/GeometryProcessing/Source/DynamicMesh/Public/Sampling/MeshMapBaker.h b/Engine/Plugins/Runtime/GeometryProcessing/Source/DynamicMesh/Public/Sampling/MeshMapBaker.h
--- a/Engine/Plugins/Runtime/GeometryProcessing/Source/DynamicMesh/Public/Sampling/MeshMapBaker.h
+++ b/Engine/Plugins/Runtime/GeometryProcessing/Source/DynamicMesh/Public/Sampling/MeshMapBaker.h
@@ -47,9 +47,16 @@ public:
 	/** @return the bake result image for a given baker index. */
 	const TArrayView<TUniquePtr<TImageBuilder<FVector4f>>> GetBakeResults(int32 EvalIdx);
 
-	/** if this function returns true, we should abort calculation */
+	/** @return true if we should abort calculation */
 	TFunction<bool(void)> CancelF = []() { return false; };
 
+	/**
+	 * @param ImageCoords the output image coordinates to be evaluated
+	 * @param UV the target mesh UV coordinates to be evaluated
+	 * @param TriID the target mesh triangle ID to be evaluated
+	 * @return true if the given sample should be filtered (return default evaluator result)
+	 */
+	TFunction<bool(const FVector2i&, const FVector2d&, int32)> SampleFilterF = nullptr;
 
 	//
 	// Parameters
