make msh2coll build collision in background operator

diff --git a/Engine/Plugins/Experimental/MeshModelingToolsetExp/Source/MeshModelingToolsExp/Private/Physics/SetCollisionGeometryTool.cpp b/Engine/Plugins/Experimental/MeshModelingToolsetExp/Source/MeshModelingToolsExp/Private/Physics/SetCollisionGeometryTool.cpp
--- a/Engine/Plugins/Experimental/MeshModelingToolsetExp/Source/MeshModelingToolsExp/Private/Physics/SetCollisionGeometryTool.cpp
+++ b/Engine/Plugins/Experimental/MeshModelingToolsetExp/Source/MeshModelingToolsExp/Private/Physics/SetCollisionGeometryTool.cpp
@@ -13,7 +13,6 @@
 #include "ShapeApproximation/ShapeDetection3.h"
 #include "ShapeApproximation/MeshSimpleShapeApproximation.h"
 
-#include "Physics/PhysicsDataCollection.h"
 #include "Physics/CollisionGeometryVisualization.h"
 
 // physics data
@@ -33,6 +32,105 @@ using namespace UE::Geometry;
 #define LOCTEXT_NAMESPACE "USetCollisionGeometryTool"
 
 
+/*
+ * Operators
+ */
+
+class FPhysicsCollectionOp : public TGenericDataOperator<FPhysicsDataCollection>
+{
+public:
+
+	TSharedPtr<FPhysicsDataCollection, ESPMode::ThreadSafe> InitialCollision;
+
+	TUniquePtr<FMeshSimpleShapeApproximation> UseShapeGenerator;
+	// Note: UseShapeGenerator holds raw pointers to these meshes, so we keep the array of shared pointers to prevent them from getting deleted while the op runs
+	TArray<TSharedPtr<FDynamicMesh3, ESPMode::ThreadSafe>> ActiveInputMeshes;
+
+	ECollisionGeometryType ComputeType;
+
+	bool bUseMaxCount;
+	bool bRemoveContained;
+	bool bAppendToExisting;
+
+	EProjectedHullAxis SweepAxis;
+	int32 MaxCount;
+
+	// Begin TGenericDataOperator interface
+	virtual void CalculateResult(FProgressCancel* Progress) override
+	{
+		check(UseShapeGenerator.IsValid());
+		check(InitialCollision.IsValid());
+
+		// calculate new collision
+		TUniquePtr<FPhysicsDataCollection> NewCollision = MakeUnique<FPhysicsDataCollection>();
+		NewCollision->InitializeFromExisting(*InitialCollision);
+		if (bAppendToExisting || ComputeType == ECollisionGeometryType::KeepExisting)
+		{
+			NewCollision->CopyGeometryFromExisting(*InitialCollision);
+		}
+
+		switch (ComputeType)
+		{
+		case ECollisionGeometryType::KeepExisting:
+		case ECollisionGeometryType::None:
+			break;
+		case ECollisionGeometryType::AlignedBoxes:
+			UseShapeGenerator->Generate_AlignedBoxes(NewCollision->Geometry);
+			break;
+		case ECollisionGeometryType::OrientedBoxes:
+			UseShapeGenerator->Generate_OrientedBoxes(NewCollision->Geometry);
+			break;
+		case ECollisionGeometryType::MinimalSpheres:
+			UseShapeGenerator->Generate_MinimalSpheres(NewCollision->Geometry);
+			break;
+		case ECollisionGeometryType::Capsules:
+			UseShapeGenerator->Generate_Capsules(NewCollision->Geometry);
+			break;
+		case ECollisionGeometryType::ConvexHulls:
+			if (UseShapeGenerator->ConvexDecompositionMaxPieces > 1)
+			{
+				UseShapeGenerator->Generate_ConvexHullDecompositions(NewCollision->Geometry);
+			}
+			else
+			{
+				UseShapeGenerator->Generate_ConvexHulls(NewCollision->Geometry);
+			}
+			break;
+		case ECollisionGeometryType::SweptHulls:
+			UseShapeGenerator->Generate_ProjectedHulls(NewCollision->Geometry,
+				(FMeshSimpleShapeApproximation::EProjectedHullAxisMode)(int32)SweepAxis);
+			break;
+		case ECollisionGeometryType::MinVolume:
+			UseShapeGenerator->Generate_MinVolume(NewCollision->Geometry);
+			break;
+		}
+
+		if (!NewCollision)
+		{
+			ensure(false);
+			return;
+		}
+
+		if (bRemoveContained)
+		{
+			NewCollision->Geometry.RemoveContainedGeometry();
+		}
+
+		
+		if (bUseMaxCount)
+		{
+			NewCollision->Geometry.FilterByVolume(MaxCount);
+		}
+
+		NewCollision->CopyGeometryToAggregate();
+
+		SetResult(MoveTemp(NewCollision));
+	}
+	// End TGenericDataOperator interface
+};
+
+
+
 const FToolTargetTypeRequirements& USetCollisionGeometryToolBuilder::GetTargetRequirements() const
 {
 	static FToolTargetTypeRequirements TypeRequirements({
@@ -83,9 +181,6 @@ void USetCollisionGeometryTool::Setup()
 		InitialSourceMeshes[k] = UE::ToolTarget::GetDynamicMeshCopy(Targets[k]);
 	});
 
-	// why is this here?
-	bInputMeshesValid = true;
-
 	UToolTarget* CollisionTarget = Targets[Targets.Num() - 1];
 	PreviewGeom = NewObject<UPreviewGeometry>(this);
 	FTransform PreviewTransform = (FTransform)UE::ToolTarget::GetLocalToWorldTransform(CollisionTarget);
@@ -105,24 +200,24 @@ void USetCollisionGeometryTool::Setup()
 	AddToolPropertySource(Settings);
 	Settings->bUseWorldSpace = (SourceObjectIndices.Num() > 1);
 	Settings->WatchProperty(Settings->InputMode, [this](ESetCollisionGeometryInputMode) { OnInputModeChanged(); });
-	Settings->WatchProperty(Settings->GeometryType, [this](ECollisionGeometryType) { bResultValid = false; });
+	Settings->WatchProperty(Settings->GeometryType, [this](ECollisionGeometryType) { InvalidateCompute(); });
 	Settings->WatchProperty(Settings->bUseWorldSpace, [this](bool) { bInputMeshesValid = false; });
-	Settings->WatchProperty(Settings->bAppendToExisting, [this](bool) { bResultValid = false; });
-	Settings->WatchProperty(Settings->bRemoveContained, [this](bool) { bResultValid = false; });
-	Settings->WatchProperty(Settings->bEnableMaxCount, [this](bool) { bResultValid = false; });
-	Settings->WatchProperty(Settings->MaxCount, [this](int32) { bResultValid = false; });
-	Settings->WatchProperty(Settings->MinThickness, [this](float) { bResultValid = false; });
-	Settings->WatchProperty(Settings->bDetectBoxes, [this](int32) { bResultValid = false; });
-	Settings->WatchProperty(Settings->bDetectSpheres, [this](int32) { bResultValid = false; });
-	Settings->WatchProperty(Settings->bDetectCapsules, [this](int32) { bResultValid = false; });
-	Settings->WatchProperty(Settings->bSimplifyHulls, [this](bool) { bResultValid = false; });
-	Settings->WatchProperty(Settings->HullTargetFaceCount, [this](int32) { bResultValid = false; });
-	Settings->WatchProperty(Settings->MaxHullsPerMesh, [this](int32) { bResultValid = false; });
-	Settings->WatchProperty(Settings->ConvexDecompositionSearchFactor, [this](int32) { bResultValid = false; });
-	Settings->WatchProperty(Settings->AddHullsErrorTolerance, [this](int32) { bResultValid = false; });
-	Settings->WatchProperty(Settings->bSimplifyPolygons, [this](bool) { bResultValid = false; });
-	Settings->WatchProperty(Settings->HullTolerance, [this](float) { bResultValid = false; });
-	Settings->WatchProperty(Settings->SweepAxis, [this](EProjectedHullAxis) { bResultValid = false; });
+	Settings->WatchProperty(Settings->bAppendToExisting, [this](bool) { InvalidateCompute(); });
+	Settings->WatchProperty(Settings->bRemoveContained, [this](bool) { InvalidateCompute(); });
+	Settings->WatchProperty(Settings->bEnableMaxCount, [this](bool) { InvalidateCompute(); });
+	Settings->WatchProperty(Settings->MaxCount, [this](int32) { InvalidateCompute(); });
+	Settings->WatchProperty(Settings->MinThickness, [this](float) { InvalidateCompute(); });
+	Settings->WatchProperty(Settings->bDetectBoxes, [this](int32) { InvalidateCompute(); });
+	Settings->WatchProperty(Settings->bDetectSpheres, [this](int32) { InvalidateCompute(); });
+	Settings->WatchProperty(Settings->bDetectCapsules, [this](int32) { InvalidateCompute(); });
+	Settings->WatchProperty(Settings->bSimplifyHulls, [this](bool) { InvalidateCompute(); });
+	Settings->WatchProperty(Settings->HullTargetFaceCount, [this](int32) { InvalidateCompute(); });
+	Settings->WatchProperty(Settings->MaxHullsPerMesh, [this](int32) { InvalidateCompute(); });
+	Settings->WatchProperty(Settings->ConvexDecompositionSearchFactor, [this](int32) { InvalidateCompute(); });
+	Settings->WatchProperty(Settings->AddHullsErrorTolerance, [this](int32) { InvalidateCompute(); });
+	Settings->WatchProperty(Settings->bSimplifyPolygons, [this](bool) { InvalidateCompute(); });
+	Settings->WatchProperty(Settings->HullTolerance, [this](float) { InvalidateCompute(); });
+	Settings->WatchProperty(Settings->SweepAxis, [this](EProjectedHullAxis) { InvalidateCompute(); });
 
 	if (InitialSourceMeshes.Num() == 1)
 	{
@@ -133,8 +228,6 @@ void USetCollisionGeometryTool::Setup()
 		AddToolPropertySource(PolygroupLayerProperties);
 	}
 
-	bResultValid = false;
-
 	VizSettings = NewObject<UCollisionGeometryVisualizationProperties>(this);
 	VizSettings->RestoreProperties(this);
 	AddToolPropertySource(VizSettings);
@@ -150,6 +243,63 @@ void USetCollisionGeometryTool::Setup()
 	GetToolManager()->DisplayMessage(
 		LOCTEXT("OnStartTool", "Initialize Simple Collision geometry for a Mesh from one or more input Meshes (including itself)."),
 		EToolMessageLevel::UserNotification);
+
+	// Make sure we are set to precompute input meshes on first tick
+	bInputMeshesValid = false;
+}
+
+
+TUniquePtr<UE::Geometry::TGenericDataOperator<FPhysicsDataCollection>> USetCollisionGeometryTool::MakeNewOperator()
+{
+	TUniquePtr<FPhysicsCollectionOp> Op = MakeUnique<FPhysicsCollectionOp>();
+
+	Op->InitialCollision = InitialCollision;
+
+	// Pick the approximator and input meshes that will be used by the op
+	TSharedPtr<UE::Geometry::FMeshSimpleShapeApproximation, ESPMode::ThreadSafe>* Approximator = nullptr;
+	TArray<TSharedPtr<FDynamicMesh3, ESPMode::ThreadSafe>>* Inputs = nullptr;
+	if (Settings->InputMode == ESetCollisionGeometryInputMode::CombineAll)
+	{
+		Approximator = &CombinedInputMeshesApproximator;
+		Inputs = &CombinedInputMeshes;
+	}
+	else if (Settings->InputMode == ESetCollisionGeometryInputMode::PerMeshComponent)
+	{
+		Approximator = &SeparatedMeshesApproximator;
+		Inputs = &SeparatedInputMeshes;
+	}
+	else if (Settings->InputMode == ESetCollisionGeometryInputMode::PerMeshGroup)
+	{
+		Approximator = &PerGroupMeshesApproximator;
+		Inputs = &PerGroupInputMeshes;
+	}
+	else
+	{
+		Approximator = &InputMeshesApproximator;
+		Inputs = &InputMeshes;
+	}
+	Op->UseShapeGenerator = MakeUnique<FMeshSimpleShapeApproximation>(**Approximator);
+	Op->ActiveInputMeshes = *Inputs;
+
+	Op->UseShapeGenerator->bDetectSpheres = Settings->bDetectSpheres;
+	Op->UseShapeGenerator->bDetectBoxes = Settings->bDetectBoxes;
+	Op->UseShapeGenerator->bDetectCapsules = Settings->bDetectCapsules;
+	Op->UseShapeGenerator->MinDimension = Settings->MinThickness;
+	Op->UseShapeGenerator->bSimplifyHulls = Settings->bSimplifyHulls;
+	Op->UseShapeGenerator->HullTargetFaceCount = Settings->HullTargetFaceCount;
+	Op->UseShapeGenerator->ConvexDecompositionMaxPieces = Settings->MaxHullsPerMesh;
+	Op->UseShapeGenerator->ConvexDecompositionSearchFactor = Settings->ConvexDecompositionSearchFactor;
+	Op->UseShapeGenerator->ConvexDecompositionErrorTolerance = Settings->AddHullsErrorTolerance;
+	Op->UseShapeGenerator->HullSimplifyTolerance = Settings->HullTolerance;
+
+	Op->ComputeType = Settings->GeometryType;
+	Op->bAppendToExisting = Settings->bAppendToExisting;
+	Op->bUseMaxCount = Settings->bEnableMaxCount;
+	Op->MaxCount = Settings->MaxCount;
+	Op->bRemoveContained = Settings->bRemoveContained;
+	Op->SweepAxis = Settings->SweepAxis;
+
+	return Op;
 }
 
 
@@ -173,6 +323,11 @@ void USetCollisionGeometryTool::OnShutdown(EToolShutdownType ShutdownType)
 		}
 	}
 
+	if (Compute)
+	{
+		Compute->Shutdown();
+	}
+
 	if (ShutdownType == EToolShutdownType::Accept)
 	{
 		// Make sure rendering is done so that we are not changing data being used by collision drawing.
@@ -245,13 +400,32 @@ void USetCollisionGeometryTool::OnTick(float DeltaTime)
 	{
 		PrecomputeInputMeshes();
 		bInputMeshesValid = true;
-		bResultValid = false;
+		InvalidateCompute();
 	}
 
-	if (bResultValid == false)
+	if (Compute)
 	{
-		UpdateGeneratedCollision();
-		bResultValid = true;
+		Compute->Tick(DeltaTime);
+
+		if (Compute->HaveValidResult())
+		{
+			TUniquePtr<FPhysicsDataCollection> Result = Compute->Shutdown();
+			if (Result.IsValid())
+			{
+				GeneratedCollision = MakeShareable<FPhysicsDataCollection>(Result.Release());
+
+				bVisualizationDirty = true;
+
+				// update visualization
+				PreviewGeom->RemoveAllLineSets();
+				UE::PhysicsTools::InitializePreviewGeometryLines(*GeneratedCollision, PreviewGeom,
+					VizSettings->Color, VizSettings->LineThickness, 0.0f, 16);
+
+				// update property set
+				CollisionProps->Reset();
+				UE::PhysicsTools::InitializePhysicsToolObjectPropertySet(GeneratedCollision.Get(), CollisionProps);
+			}
+		}
 	}
 
 	if (bVisualizationDirty)
@@ -262,6 +436,28 @@ void USetCollisionGeometryTool::OnTick(float DeltaTime)
 }
 
 
+void USetCollisionGeometryTool::InvalidateCompute()
+{
+	if (PreviewGeom)
+	{
+		PreviewGeom->RemoveAllLineSets();
+	}
+
+	if (!bInputMeshesValid)
+	{
+		// InvalidateCompute() will be called again when the input meshes are valid
+		return;
+	}
+
+	if (!Compute)
+	{
+		// Initialize background compute
+		Compute = MakeUnique<TGenericDataBackgroundCompute<FPhysicsDataCollection>>();
+		Compute->Setup(this);
+	}
+	Compute->InvalidateResult();
+}
+
 
 void USetCollisionGeometryTool::OnInputModeChanged()
 {
@@ -269,13 +465,13 @@ void USetCollisionGeometryTool::OnInputModeChanged()
 	{
 		SetToolPropertySourceEnabled(PolygroupLayerProperties, Settings->InputMode == ESetCollisionGeometryInputMode::PerMeshGroup);
 	}
-	bResultValid = false;
+	InvalidateCompute();
 }
 
 void USetCollisionGeometryTool::OnSelectedGroupLayerChanged()
 {
 	bInputMeshesValid = false;
-	bResultValid = false;
+	InvalidateCompute();
 }
 
 
@@ -319,101 +515,6 @@ void USetCollisionGeometryTool::UpdateVisualization()
 }
 
 
-void USetCollisionGeometryTool::UpdateGeneratedCollision()
-{
-	// calculate new collision
-	ECollisionGeometryType ComputeType = Settings->GeometryType;
-
-
-	TSharedPtr<FPhysicsDataCollection, ESPMode::ThreadSafe> NewCollision = MakeShared<FPhysicsDataCollection, ESPMode::ThreadSafe>();
-	NewCollision->InitializeFromExisting(*InitialCollision);
-	if (Settings->bAppendToExisting || ComputeType == ECollisionGeometryType::KeepExisting)
-	{
-		NewCollision->CopyGeometryFromExisting(*InitialCollision);
-	}
-
-	TSharedPtr<FMeshSimpleShapeApproximation, ESPMode::ThreadSafe> UseShapeGenerator = GetApproximator(Settings->InputMode);
-
-	UseShapeGenerator->bDetectSpheres = Settings->bDetectSpheres;
-	UseShapeGenerator->bDetectBoxes = Settings->bDetectBoxes;
-	UseShapeGenerator->bDetectCapsules = Settings->bDetectCapsules;
-	//UseShapeGenerator->bDetectConvexes = Settings->bDetectConvexes;
-
-	UseShapeGenerator->MinDimension = Settings->MinThickness;
-
-	switch (ComputeType)
-	{
-	case ECollisionGeometryType::KeepExisting:
-	case ECollisionGeometryType::None:
-		break;
-	case ECollisionGeometryType::AlignedBoxes:
-		UseShapeGenerator->Generate_AlignedBoxes(NewCollision->Geometry);
-		break;
-	case ECollisionGeometryType::OrientedBoxes:
-		UseShapeGenerator->Generate_OrientedBoxes(NewCollision->Geometry);
-		break;
-	case ECollisionGeometryType::MinimalSpheres:
-		UseShapeGenerator->Generate_MinimalSpheres(NewCollision->Geometry);
-		break;
-	case ECollisionGeometryType::Capsules:
-		UseShapeGenerator->Generate_Capsules(NewCollision->Geometry);
-		break;
-	case ECollisionGeometryType::ConvexHulls:
-		UseShapeGenerator->bSimplifyHulls = Settings->bSimplifyHulls;
-		UseShapeGenerator->HullTargetFaceCount = Settings->HullTargetFaceCount;
-		if (Settings->MaxHullsPerMesh > 1)
-		{
-			UseShapeGenerator->ConvexDecompositionMaxPieces = Settings->MaxHullsPerMesh;
-			UseShapeGenerator->ConvexDecompositionSearchFactor = Settings->ConvexDecompositionSearchFactor;
-			UseShapeGenerator->ConvexDecompositionErrorTolerance = Settings->AddHullsErrorTolerance;
-			UseShapeGenerator->Generate_ConvexHullDecompositions(NewCollision->Geometry);
-		}
-		else
-		{
-			UseShapeGenerator->Generate_ConvexHulls(NewCollision->Geometry);
-		}
-		break;
-	case ECollisionGeometryType::SweptHulls:
-		UseShapeGenerator->bSimplifyHulls = Settings->bSimplifyPolygons;
-		UseShapeGenerator->HullSimplifyTolerance = Settings->HullTolerance;
-		UseShapeGenerator->Generate_ProjectedHulls(NewCollision->Geometry, 
-			(FMeshSimpleShapeApproximation::EProjectedHullAxisMode)(int32)Settings->SweepAxis);
-		break;
-	case ECollisionGeometryType::MinVolume:
-		UseShapeGenerator->Generate_MinVolume(NewCollision->Geometry);
-		break;
-	}
-
-
-	if (!NewCollision)
-	{
-		ensure(false);
-		return;
-	}
-	GeneratedCollision = NewCollision;
-
-	if (Settings->bRemoveContained)
-	{
-		GeneratedCollision->Geometry.RemoveContainedGeometry();
-	}
-
-	bool bUseMaxCount = (Settings->bEnableMaxCount);
-	if (bUseMaxCount)
-	{
-		GeneratedCollision->Geometry.FilterByVolume(Settings->MaxCount);
-	}
-
-	GeneratedCollision->CopyGeometryToAggregate();
-
-	// update visualization
-	PreviewGeom->RemoveAllLineSets();
-	UE::PhysicsTools::InitializePreviewGeometryLines(*GeneratedCollision, PreviewGeom,
-		VizSettings->Color, VizSettings->LineThickness, 0.0f, 16);
-
-	// update property set
-	CollisionProps->Reset();
-	UE::PhysicsTools::InitializePhysicsToolObjectPropertySet(GeneratedCollision.Get(), CollisionProps);
-}
 
 
 
@@ -557,25 +658,4 @@ void USetCollisionGeometryTool::PrecomputeInputMeshes()
 }
 
 
-TSharedPtr<FMeshSimpleShapeApproximation, ESPMode::ThreadSafe>& USetCollisionGeometryTool::GetApproximator(ESetCollisionGeometryInputMode MeshSetMode)
-{
-	if (MeshSetMode == ESetCollisionGeometryInputMode::CombineAll)
-	{
-		return CombinedInputMeshesApproximator;
-	}
-	else if ( MeshSetMode == ESetCollisionGeometryInputMode::PerMeshComponent)
-	{
-		return SeparatedMeshesApproximator;
-	}
-	else if (MeshSetMode == ESetCollisionGeometryInputMode::PerMeshGroup)
-	{
-		return PerGroupMeshesApproximator;
-	}
-	else
-	{
-		return InputMeshesApproximator;
-	}
-}
-
-
 #undef LOCTEXT_NAMESPACE
\ No newline at end of file
diff --git a/Engine/Plugins/Experimental/MeshModelingToolsetExp/Source/MeshModelingToolsExp/Public/Physics/SetCollisionGeometryTool.h b/Engine/Plugins/Experimental/MeshModelingToolsetExp/Source/MeshModelingToolsExp/Public/Physics/SetCollisionGeometryTool.h
--- a/Engine/Plugins/Experimental/MeshModelingToolsetExp/Source/MeshModelingToolsExp/Public/Physics/SetCollisionGeometryTool.h
+++ b/Engine/Plugins/Experimental/MeshModelingToolsetExp/Source/MeshModelingToolsExp/Public/Physics/SetCollisionGeometryTool.h
@@ -10,8 +10,11 @@
 #include "OrientedBoxTypes.h"
 #include "CapsuleTypes.h"
 #include "Physics/CollisionPropertySets.h"
+#include "Physics/PhysicsDataCollection.h"
 #include "PropertySets/PolygroupLayersProperties.h"
 #include "Polygroups/PolygroupSet.h"
+#include "ModelingOperators.h"
+#include "MeshOpPreviewHelpers.h"
 #include "SetCollisionGeometryTool.generated.h"
 
 class UPreviewGeometry;
@@ -156,7 +159,7 @@ public:
  * Mesh Inspector Tool for visualizing mesh information
  */
 UCLASS()
-class MESHMODELINGTOOLSEXP_API USetCollisionGeometryTool : public UMultiSelectionMeshEditingTool
+class MESHMODELINGTOOLSEXP_API USetCollisionGeometryTool : public UMultiSelectionMeshEditingTool, public UE::Geometry::IGenericDataOperatorFactory<FPhysicsDataCollection>
 {
 	GENERATED_BODY()
 public:
@@ -167,6 +170,15 @@ public:
 
 	virtual bool HasCancel() const override { return true; }
 	virtual bool HasAccept() const override { return true; }
+	virtual bool CanAccept() const override
+	{
+		// allow accept when we're showing the current, valid result
+		return Super::CanAccept() && bInputMeshesValid && Compute && Compute->HaveValidResult() && !bVisualizationDirty;
+	}
+
+	// Begin IGenericDataOperatorFactory interface
+	virtual TUniquePtr<UE::Geometry::TGenericDataOperator<FPhysicsDataCollection>> MakeNewOperator() override;
+	// End IGenericDataOperatorFactory interface
 
 protected:
 
@@ -185,6 +197,11 @@ protected:
 	UPROPERTY()
 	TObjectPtr<UMaterialInterface> LineMaterial = nullptr;
 
+	//
+	// Background compute
+	//
+	TUniquePtr<TGenericDataBackgroundCompute<FPhysicsDataCollection>> Compute = nullptr;
+
 protected:
 	UPROPERTY()
 	TObjectPtr<UPreviewGeometry> PreviewGeom;
@@ -196,6 +213,11 @@ protected:
 
 	void OnInputModeChanged();
 
+	/**
+	 * Invalidates the background compute operator.
+	 */
+	void InvalidateCompute();
+
 	enum class EDetectedCollisionGeometry
 	{
 		None,
@@ -232,7 +254,6 @@ protected:
 		const TArray<TSharedPtr<FDynamicMesh3, ESPMode::ThreadSafe>>& FromInputMeshes,
 		TArray<TSharedPtr<FDynamicMesh3, ESPMode::ThreadSafe>>& ToMeshes,
 		TFunctionRef<bool(const FDynamicMesh3*, int32, int32)> TrisConnectedPredicate);
-	TSharedPtr<UE::Geometry::FMeshSimpleShapeApproximation, ESPMode::ThreadSafe>& GetApproximator(ESetCollisionGeometryInputMode MeshSetMode);
 
 	TUniquePtr<UE::Geometry::FPolygroupSet> ActiveGroupSet;
 	void OnSelectedGroupLayerChanged();
@@ -241,12 +262,9 @@ protected:
 	FTransform OrigTargetTransform;
 	FVector TargetScale3D;
 
-	bool bResultValid = false;
 	TSharedPtr<FPhysicsDataCollection, ESPMode::ThreadSafe> InitialCollision;
 	TSharedPtr<FPhysicsDataCollection, ESPMode::ThreadSafe> GeneratedCollision;
 
-	void UpdateGeneratedCollision();
-//	TSharedPtr<FPhysicsDataCollection> GenerateCollision_MinVolume();
 
 	bool bVisualizationDirty = false;
 	void UpdateVisualization();
