Fix crash when accepting Normal tool on a mesh with degenerate triangles

diff --git a/Engine/Plugins/Experimental/MeshModelingToolsetExp/Source/MeshModelingToolsExp/Private/EditNormalsTool.cpp b/Engine/Plugins/Experimental/MeshModelingToolsetExp/Source/MeshModelingToolsExp/Private/EditNormalsTool.cpp
--- a/Engine/Plugins/Experimental/MeshModelingToolsetExp/Source/MeshModelingToolsExp/Private/EditNormalsTool.cpp
+++ b/Engine/Plugins/Experimental/MeshModelingToolsetExp/Source/MeshModelingToolsExp/Private/EditNormalsTool.cpp
@@ -263,8 +263,20 @@ void UEditNormalsTool::GenerateAsset(const TArray<FDynamicMeshOpResult>& Results
 		const FDynamicMesh3* NewDynamicMesh = Results[ComponentIdx].Mesh.Get();
 		if (NewDynamicMesh)
 		{
-			bool bTopologyChanged = BasicProperties->WillTopologyChange();
-			UE::ToolTarget::CommitMeshDescriptionUpdateViaDynamicMesh(Targets[ComponentIdx], *NewDynamicMesh, bTopologyChanged);
+			if (bool bTopologyChanged = BasicProperties->WillTopologyChange())
+			{
+				// Tool may have changed the topology of the normal overlay (according to the specified tool properties), so we can't simply update the target mesh.
+				// Passing in bTopologyChanged = true will trigger the slower Convert function rather than the fast Update function.
+				UE::ToolTarget::CommitMeshDescriptionUpdateViaDynamicMesh(Targets[ComponentIdx], *NewDynamicMesh, bTopologyChanged);
+			}
+			else
+			{
+				// The tool didn't change the overlay topology so there's a chance we can do a fast path Update of the normal attributes.
+				// This function will still check if there is a mismatch between the dynamic mesh and target mesh in terms of triangles/vertices, and
+				// if so it will do the full conversion.
+				constexpr bool bUpdateTangents = false;
+				UE::ToolTarget::CommitDynamicMeshNormalsUpdate(Targets[ComponentIdx], NewDynamicMesh, bUpdateTangents);
+			}
 		}
 	}
 
