MeshModelingTools: Transact the bChangesMade flag in CubeGrid tool so that it is properly undone and redone.

diff --git a/Engine/Plugins/Experimental/MeshModelingToolsetExp/Source/MeshModelingToolsExp/Private/CubeGridTool.cpp b/Engine/Plugins/Experimental/MeshModelingToolsetExp/Source/MeshModelingToolsExp/Private/CubeGridTool.cpp
--- a/Engine/Plugins/Experimental/MeshModelingToolsetExp/Source/MeshModelingToolsExp/Private/CubeGridTool.cpp
+++ b/Engine/Plugins/Experimental/MeshModelingToolsetExp/Source/MeshModelingToolsExp/Private/CubeGridTool.cpp
@@ -193,6 +193,27 @@ namespace CubeGridToolLocals
 			return TEXT("CubeGridToolLocals::FCubeGridToolModeChange");
 		}
 	};
+
+	/** Undoes setting bChangesMade, which is used to determine whether the target needs saving. */
+	class FCubeGridChangesMadeChange : public FToolCommandChange
+	{
+	public:
+		FCubeGridChangesMadeChange() {};
+
+		virtual void Apply(UObject* Object) override
+		{
+			Cast<UCubeGridTool>(Object)->SetChangesMade(true);
+		}
+		virtual void Revert(UObject* Object) override
+		{
+			Cast<UCubeGridTool>(Object)->SetChangesMade(false);
+		}
+
+		virtual FString ToString() const override
+		{
+			return TEXT("CubeGridToolLocals::FCubeGridChangesMadeChange");
+		}
+	};
 	
 	// Attach a frame to the box such that Z points along the given direction.
 	// If we change the choices we make here, we may need to adjust how we pick the welded vertices in
@@ -1043,6 +1064,8 @@ void UCubeGridTool::OnTick(float DeltaTime)
 void UCubeGridTool::ApplyPreview()
 {
 	using namespace CubeGridToolLocals;
+	const FText TransactionText = LOCTEXT("CubeGridToolTransactionName", "Block Tool Change");
+	GetToolManager()->BeginUndoTransaction(TransactionText);
 
 	FDynamicMeshChangeTracker ChangeTracker(CurrentMesh.Get());
 	ChangeTracker.BeginChange();
@@ -1058,7 +1081,11 @@ void UCubeGridTool::ApplyPreview()
 		return;
 	}
 
-	bChangesMade = true; // TODO: make this undoable
+	if (!bChangesMade)
+	{
+		bChangesMade = true;
+		GetToolManager()->EmitObjectChange(this, MakeUnique<FCubeGridChangesMadeChange>(), TransactionText);
+	}
 
 	MeshSpatial->Build();
 
@@ -1067,9 +1094,6 @@ void UCubeGridTool::ApplyPreview()
 
 	UpdateComputeInputs();
 
-	const FText TransactionText = LOCTEXT("CubeGridToolTransactionName", "Block Tool Change");
-	GetToolManager()->BeginUndoTransaction(TransactionText);
-
 	bWaitingToApplyPreview = false;
 	bPreviewMayDiffer = false;
 	bBlockUntilPreviewUpdate = false;
@@ -2245,6 +2269,11 @@ void UCubeGridTool::RevertToDefaultMode()
 	}
 }
 
+void UCubeGridTool::SetChangesMade(bool bChangesMadeIn)
+{
+	bChangesMade = bChangesMadeIn;
+}
+
 bool UCubeGridTool::CanCurrentlyNestedCancel()
 {
 	return Mode == EMode::Corner || bHaveSelection;
diff --git a/Engine/Plugins/Experimental/MeshModelingToolsetExp/Source/MeshModelingToolsExp/Public/CubeGridTool.h b/Engine/Plugins/Experimental/MeshModelingToolsetExp/Source/MeshModelingToolsExp/Public/CubeGridTool.h
--- a/Engine/Plugins/Experimental/MeshModelingToolsetExp/Source/MeshModelingToolsExp/Public/CubeGridTool.h
+++ b/Engine/Plugins/Experimental/MeshModelingToolsetExp/Source/MeshModelingToolsExp/Public/CubeGridTool.h
@@ -334,6 +334,7 @@ public:
 	virtual void UpdateUsingMeshChange(const UE::Geometry::FDynamicMeshChange& MeshChange, bool bRevert);
 	virtual bool IsInDefaultMode() const;
 	virtual void RevertToDefaultMode();
+	virtual void SetChangesMade(bool bChangesMadeIn);
 
 	// IClickDragBehaviorTarget
 	virtual FInputRayHit CanBeginClickDragSequence(const FInputDeviceRay& PressPos) override;
