ModelingTools: Only enforce StaticMesh target requirement on source meshes in the BakeAll tool.

diff --git a/Engine/Plugins/Experimental/MeshModelingToolsetExp/Source/MeshModelingToolsExp/Private/BakeMultiMeshAttributeMapsTool.cpp b/Engine/Plugins/Experimental/MeshModelingToolsetExp/Source/MeshModelingToolsExp/Private/BakeMultiMeshAttributeMapsTool.cpp
--- a/Engine/Plugins/Experimental/MeshModelingToolsetExp/Source/MeshModelingToolsExp/Private/BakeMultiMeshAttributeMapsTool.cpp
+++ b/Engine/Plugins/Experimental/MeshModelingToolsetExp/Source/MeshModelingToolsExp/Private/BakeMultiMeshAttributeMapsTool.cpp
@@ -41,7 +41,6 @@ const FToolTargetTypeRequirements& UBakeMultiMeshAttributeMapsToolBuilder::GetTa
 	static FToolTargetTypeRequirements TypeRequirements({
 		UMeshDescriptionProvider::StaticClass(),
 		UPrimitiveComponentBackedTarget::StaticClass(),
-		UStaticMeshBackedTarget::StaticClass(),			// FMeshSceneAdapter currently only supports StaticMesh targets
 		UMaterialProvider::StaticClass()
 		});
 	return TypeRequirements;
@@ -50,7 +49,25 @@ const FToolTargetTypeRequirements& UBakeMultiMeshAttributeMapsToolBuilder::GetTa
 bool UBakeMultiMeshAttributeMapsToolBuilder::CanBuildTool(const FToolBuilderState& SceneState) const
 {
 	const int32 NumTargets = SceneState.TargetManager->CountSelectedAndTargetable(SceneState, GetTargetRequirements());
-	return (NumTargets > 1);
+	if (NumTargets > 1)
+	{
+		// FMeshSceneAdapter currently only supports StaticMesh targets.
+		// Restrict source targets to StaticMesh.
+		bool bValidTargets = true;
+		int TargetId = 0;
+		SceneState.TargetManager->EnumerateSelectedAndTargetableComponents(SceneState, GetTargetRequirements(),
+			[&bValidTargets, &TargetId](UActorComponent* Component)
+			{
+				if (TargetId > 0)
+				{
+					const UStaticMeshComponent* StaticMesh = Cast<UStaticMeshComponent>(Component);
+					bValidTargets = bValidTargets && StaticMesh;
+				}
+				++TargetId;
+			});
+		return bValidTargets;
+	}
+	return false;
 }
 
 UMultiSelectionMeshEditingTool* UBakeMultiMeshAttributeMapsToolBuilder::CreateNewTool(const FToolBuilderState& SceneState) const
