MeshModelingTools: Fix crash in CubeGrid tool when undoing out of the tool after modifying the target.

diff --git a/Engine/Plugins/Experimental/MeshModelingToolsetExp/Source/MeshModelingToolsExp/Private/CubeGridTool.cpp b/Engine/Plugins/Experimental/MeshModelingToolsetExp/Source/MeshModelingToolsExp/Private/CubeGridTool.cpp
--- a/Engine/Plugins/Experimental/MeshModelingToolsetExp/Source/MeshModelingToolsExp/Private/CubeGridTool.cpp
+++ b/Engine/Plugins/Experimental/MeshModelingToolsetExp/Source/MeshModelingToolsExp/Private/CubeGridTool.cpp
@@ -997,7 +997,10 @@ void UCubeGridTool::Shutdown(EToolShutdownType ShutdownType)
 	{
 		Cast<IPrimitiveComponentBackedTarget>(Target)->SetOwnerVisibility(true);
 
-		if (bChangesMade)
+		// We check the shutdown type because even though we are not an accept/cancel tool, we
+		// get a cancel shutdown via Ctrl+Z. In this case we definitely don't want to update because
+		// an update in an undo transaction results in a crash.
+		if (bChangesMade && ShutdownType != EToolShutdownType::Cancel)
 		{
 			GetToolManager()->BeginUndoTransaction(LOCTEXT("CubeGridToolEditTransactionName", "Block Tool Edit"));
 			FComponentMaterialSet OutputMaterialSet;
@@ -1008,7 +1011,7 @@ void UCubeGridTool::Shutdown(EToolShutdownType ShutdownType)
 			GetToolManager()->EndUndoTransaction();
 		}
 	}
-	else if (CurrentMesh->TriangleCount() > 0)
+	else if (CurrentMesh->TriangleCount() > 0 && ShutdownType != EToolShutdownType::Cancel)
 	{
 		GetToolManager()->BeginUndoTransaction(LOCTEXT("CubeGridToolCreateTransactionName", "Block Tool Create New"));
 
