MeshModelingTools: Fix crash on CubeGrid tool completion when underlying asset is deleted.

diff --git a/Engine/Plugins/Experimental/MeshModelingToolsetExp/Source/MeshModelingToolsExp/Private/CubeGridTool.cpp b/Engine/Plugins/Experimental/MeshModelingToolsetExp/Source/MeshModelingToolsExp/Private/CubeGridTool.cpp
--- a/Engine/Plugins/Experimental/MeshModelingToolsetExp/Source/MeshModelingToolsExp/Private/CubeGridTool.cpp
+++ b/Engine/Plugins/Experimental/MeshModelingToolsetExp/Source/MeshModelingToolsExp/Private/CubeGridTool.cpp
@@ -1000,7 +1000,7 @@ void UCubeGridTool::Shutdown(EToolShutdownType ShutdownType)
 		// We check the shutdown type because even though we are not an accept/cancel tool, we
 		// get a cancel shutdown via Ctrl+Z. In this case we definitely don't want to update because
 		// an update in an undo transaction results in a crash.
-		if (bChangesMade && ShutdownType != EToolShutdownType::Cancel)
+		if (bChangesMade && ShutdownType != EToolShutdownType::Cancel && Target->IsValid())
 		{
 			GetToolManager()->BeginUndoTransaction(LOCTEXT("CubeGridToolEditTransactionName", "Block Tool Edit"));
 			FComponentMaterialSet OutputMaterialSet;
@@ -1010,6 +1010,10 @@ void UCubeGridTool::Shutdown(EToolShutdownType ShutdownType)
 				FConversionToMeshDescriptionOptions(), &OutputMaterialSet);
 			GetToolManager()->EndUndoTransaction();
 		}
+		else if (!Target->IsValid())
+		{
+			UE_LOG(LogGeometry, Error, TEXT("CubeGridTool:: Edited mesh could not be committed (it was likely forcibly deleted from under the tool)."));
+		}
 	}
 	else if (CurrentMesh->TriangleCount() > 0 && ShutdownType != EToolShutdownType::Cancel)
 	{
