ModelingTools: Update TargetMesh property in BakeAll for skeletal and dynamic mesh targets.

diff --git a/Engine/Plugins/Experimental/MeshModelingToolsetExp/Source/MeshModelingToolsExp/Private/BakeMultiMeshAttributeMapsTool.cpp b/Engine/Plugins/Experimental/MeshModelingToolsetExp/Source/MeshModelingToolsExp/Private/BakeMultiMeshAttributeMapsTool.cpp
--- a/Engine/Plugins/Experimental/MeshModelingToolsetExp/Source/MeshModelingToolsExp/Private/BakeMultiMeshAttributeMapsTool.cpp
+++ b/Engine/Plugins/Experimental/MeshModelingToolsetExp/Source/MeshModelingToolsExp/Private/BakeMultiMeshAttributeMapsTool.cpp
@@ -465,6 +465,8 @@ void UBakeMultiMeshAttributeMapsTool::Setup()
 	InputMeshSettings->RestoreProperties(this);
 	AddToolPropertySource(InputMeshSettings);
 	InputMeshSettings->TargetStaticMesh = GetStaticMeshTarget(Target);
+	InputMeshSettings->TargetSkeletalMesh = GetSkeletalMeshTarget(Target);
+	InputMeshSettings->TargetDynamicMesh = GetDynamicMeshTarget(Target);
 	UpdateUVLayerNames(InputMeshSettings->TargetUVLayer, InputMeshSettings->TargetUVLayerNamesList, TargetMesh);
 	InputMeshSettings->WatchProperty(InputMeshSettings->TargetUVLayer, [this](FString) { OpState |= EBakeOpState::Evaluate; });
 	InputMeshSettings->WatchProperty(InputMeshSettings->ProjectionDistance, [this](float) { OpState |= EBakeOpState::Evaluate; });
diff --git a/Engine/Plugins/Experimental/MeshModelingToolsetExp/Source/MeshModelingToolsExp/Public/BakeMultiMeshAttributeMapsTool.h b/Engine/Plugins/Experimental/MeshModelingToolsetExp/Source/MeshModelingToolsExp/Public/BakeMultiMeshAttributeMapsTool.h
--- a/Engine/Plugins/Experimental/MeshModelingToolsetExp/Source/MeshModelingToolsExp/Public/BakeMultiMeshAttributeMapsTool.h
+++ b/Engine/Plugins/Experimental/MeshModelingToolsetExp/Source/MeshModelingToolsExp/Public/BakeMultiMeshAttributeMapsTool.h
@@ -103,9 +103,20 @@ class MESHMODELINGTOOLSEXP_API UBakeMultiMeshInputToolProperties : public UInter
 
 public:
 	/** Target mesh to sample to */
-	UPROPERTY(VisibleAnywhere, Category = BakeInput, DisplayName = "Target Mesh", meta = (TransientToolProperty))
+	UPROPERTY(VisibleAnywhere, Category = BakeInput, DisplayName = "Target Mesh", meta = (TransientToolProperty,
+		EditCondition = "TargetStaticMesh != nullptr", EditConditionHides))
 	TObjectPtr<UStaticMesh> TargetStaticMesh = nullptr;
 
+	/** Target mesh to sample to */
+	UPROPERTY(VisibleAnywhere, Category = BakeInput, DisplayName = "Target Mesh", meta = (TransientToolProperty,
+		EditCondition = "TargetSkeletalMesh != nullptr", EditConditionHides))
+	TObjectPtr<USkeletalMesh> TargetSkeletalMesh = nullptr;
+
+	/** Target mesh to sample to */
+	UPROPERTY(VisibleAnywhere, Category = BakeInput, DisplayName = "Target Mesh", meta = (TransientToolProperty,
+		EditCondition = "TargetDynamicMesh != nullptr", EditConditionHides))
+	TObjectPtr<AActor> TargetDynamicMesh = nullptr;
+
 	/** UV channel to use for the target mesh */
 	UPROPERTY(EditAnywhere, Category = BakeInput, meta = (DisplayName = "Target Mesh UV Channel",
 		GetOptions = GetTargetUVLayerNamesFunc, TransientToolProperty, NoResetToDefault))
