ModelingTools: Fixed BakeAll TangentNormal compute by ensuring that the normals are transformed into the correct space.

diff --git a/Engine/Plugins/Experimental/MeshModelingToolsetExp/Source/MeshModelingToolsExp/Private/BakeMultiMeshAttributeMapsTool.cpp b/Engine/Plugins/Experimental/MeshModelingToolsetExp/Source/MeshModelingToolsExp/Private/BakeMultiMeshAttributeMapsTool.cpp
--- a/Engine/Plugins/Experimental/MeshModelingToolsetExp/Source/MeshModelingToolsExp/Private/BakeMultiMeshAttributeMapsTool.cpp
+++ b/Engine/Plugins/Experimental/MeshModelingToolsetExp/Source/MeshModelingToolsExp/Private/BakeMultiMeshAttributeMapsTool.cpp
@@ -103,7 +103,7 @@ public:
 		{
 			if (ChildMesh)
 			{
-				ProcessFn(ChildMesh->MeshSpatial);
+				ProcessFn(ChildMesh);
 			}
 		};
 		MeshScene->ProcessActorChildMeshes(ProcessChildMesh);
@@ -111,8 +111,8 @@ public:
 
 	virtual int32 GetTriangleCount(const void* Mesh) const override
 	{
-		const IMeshSpatialWrapper* Spatial = static_cast<const IMeshSpatialWrapper*>(Mesh);
-		return Spatial->GetTriangleCount();
+		const FActorChildMesh* ChildMesh = static_cast<const FActorChildMesh*>(Mesh);
+		return ChildMesh->MeshSpatial->GetTriangleCount();
 	}
 
 	virtual void SetTextureMap(const void* Mesh, const FBakeDetailTexture& Map) override
@@ -167,9 +167,12 @@ public:
 		FVector3d& TriBaryCoords,
 		const IMeshSpatial::FQueryOptions& Options = IMeshSpatial::FQueryOptions()) const override
 	{
+		// Transform ray to world space
+		const FRay3d WorldRay = BaseToWorldTransform.TransformRay(Ray);
+		
 		// TODO: Pass-through max distance to FMeshSceneAdapter query
 		FMeshSceneRayHit HitResult;
-		const bool bHit = MeshScene->FindNearestRayIntersection(Ray, HitResult);
+		const bool bHit = MeshScene->FindNearestRayIntersection(WorldRay, HitResult);
 		const void* HitMesh = nullptr;
 
 		// Use TNumericLimits<float>::Max() for consistency with MeshAABBTree3.
@@ -179,7 +182,16 @@ public:
 			TriId = HitResult.HitMeshTriIndex;
 			NearestT = HitResult.RayDistance;
 			TriBaryCoords = HitResult.HitMeshBaryCoords;
-			HitMesh = HitResult.HitMeshSpatialWrapper;
+
+			// Use the component and component element index to identify the actual instance
+			// (FActorChildMesh) that was hit.
+			if (const FActorAdapter* const* FoundActor = ActorComponentMap.Find(HitResult.HitComponent); FoundActor != nullptr)
+			{
+				if (HitResult.HitComponentElementIndex < (*FoundActor)->ChildMeshes.Num())
+				{
+					HitMesh = (*FoundActor)->ChildMeshes[FMath::Max(0, HitResult.HitComponentElementIndex)].Get();
+				}
+			}
 		}
 		return HitMesh;
 	}
@@ -203,14 +215,14 @@ public:
 
 	virtual bool IsTriangle(const void* Mesh, const int TriId) const override
 	{
-		const IMeshSpatialWrapper* Spatial = static_cast<const IMeshSpatialWrapper*>(Mesh);
-		return Spatial->IsTriangle(TriId);
+		const FActorChildMesh* ChildMesh = static_cast<const FActorChildMesh*>(Mesh);
+		return ChildMesh->MeshSpatial->IsTriangle(TriId);
 	}
 
 	virtual FIndex3i GetTriangle(const void* Mesh, const int TriId) const override
 	{
-		const IMeshSpatialWrapper* Spatial = static_cast<const IMeshSpatialWrapper*>(Mesh);
-		return Spatial->GetTriangle(TriId);
+		const FActorChildMesh* ChildMesh = static_cast<const FActorChildMesh*>(Mesh);
+		return ChildMesh->MeshSpatial->GetTriangle(TriId);
 	}
 
 	virtual FVector3d GetTriNormal(const void* Mesh, const int TriId) const override
@@ -229,14 +241,14 @@ public:
 
 	virtual bool HasNormals(const void* Mesh) const override
 	{
-		const IMeshSpatialWrapper* Spatial = static_cast<const IMeshSpatialWrapper*>(Mesh);
-		return Spatial->HasNormals();
+		const FActorChildMesh* ChildMesh = static_cast<const FActorChildMesh*>(Mesh);
+		return ChildMesh->MeshSpatial->HasNormals();
 	}
 
 	virtual bool HasUVs(const void* Mesh, const int UVLayer = 0) const override
 	{
-		const IMeshSpatialWrapper* Spatial = static_cast<const IMeshSpatialWrapper*>(Mesh);
-		return Spatial->HasUVs(UVLayer);
+		const FActorChildMesh* ChildMesh = static_cast<const FActorChildMesh*>(Mesh);
+		return ChildMesh->MeshSpatial->HasUVs(UVLayer);
 	}
 
 	virtual bool HasTangents(const void* Mesh) const override
@@ -258,8 +270,10 @@ public:
 		const int32 TriId,
 		const FVector3d& BaryCoords) const override
 	{
-		const IMeshSpatialWrapper* Spatial = static_cast<const IMeshSpatialWrapper*>(Mesh);
-		return Spatial->TriBaryInterpolatePoint(TriId, BaryCoords);
+		const FActorChildMesh* ChildMesh = static_cast<const FActorChildMesh*>(Mesh);
+		const FVector3d Point = ChildMesh->MeshSpatial->TriBaryInterpolatePoint(TriId, BaryCoords);
+		const FVector3d WorldPoint = ChildMesh->WorldTransform.TransformPosition(Point);
+		return WorldPoint;
 	}
 
 	virtual bool TriBaryInterpolateNormal(
@@ -268,8 +282,13 @@ public:
 		const FVector3d& BaryCoords,
 		FVector3f& NormalOut) const override
 	{
-		const IMeshSpatialWrapper* Spatial = static_cast<const IMeshSpatialWrapper*>(Mesh);
-		return Spatial->TriBaryInterpolateNormal(TriId, BaryCoords, NormalOut);
+		const FActorChildMesh* ChildMesh = static_cast<const FActorChildMesh*>(Mesh);
+		const bool bSuccess = ChildMesh->MeshSpatial->TriBaryInterpolateNormal(TriId, BaryCoords, NormalOut);
+		const FVector3d NormalOut3d(NormalOut);
+		// As a stop-gap fix, transform our normal to the local space of the Base mesh here.
+		// TODO: Handle WorldToBase transformation in the evaluator.
+		NormalOut = WorldToBaseTransform.TransformNormal(FVector3f(ChildMesh->WorldTransform.TransformNormal(NormalOut3d)));
+		return bSuccess;
 	}
 
 	virtual bool TriBaryInterpolateUV(
@@ -279,8 +298,8 @@ public:
 		const int UVLayer,
 		FVector2f& UVOut ) const override
 	{
-		const IMeshSpatialWrapper* Spatial = static_cast<const IMeshSpatialWrapper*>(Mesh);
-		return Spatial->TriBaryInterpolateUV(TriId, BaryCoords, UVLayer, UVOut);
+		const FActorChildMesh* ChildMesh = static_cast<const FActorChildMesh*>(Mesh);
+		return ChildMesh->MeshSpatial->TriBaryInterpolateUV(TriId, BaryCoords, UVLayer, UVOut);
 	}
 
 	virtual bool TriBaryInterpolateColor(
@@ -327,6 +346,11 @@ public:
 		DetailNormalMaps = Map;
 	}
 
+public:
+	FTransformSRT3d BaseToWorldTransform;
+	FTransformSRT3d WorldToBaseTransform;
+	TMap<UActorComponent*, const FActorAdapter*> ActorComponentMap;
+
 protected:
 	FMeshSceneAdapter* MeshScene = nullptr;
 	FDetailTextureMap DetailTextureMaps;
@@ -350,6 +374,8 @@ public:
 	TSharedPtr<TArray<int32>, ESPMode::ThreadSafe> BaseMeshUVCharts;
 	TSharedPtr<UE::Geometry::TImageBuilder<FVector4f>> SampleFilterMask;
 
+	FTransformSRT3d BaseToWorldTransform;
+
 	// Detail bake data
 	TArray<TSharedPtr<UE::Geometry::TImageBuilder<FVector4f>>> CachedColorImages;
 	UBakeMultiMeshAttributeMapsTool::FTextureImageMap CachedMeshToColorImageMap;
@@ -378,6 +404,22 @@ public:
 		}
 		
 		FMeshBakerMeshSceneSampler DetailSampler(DetailMeshScene.Get());
+
+		// TODO: Precompute the ActorComponent to ActorChildMesh map in the tool.
+		auto ProcessChildMesh = [&DetailSampler](const FActorAdapter* Actor, const FActorChildMesh* ChildMesh)
+		{
+			// Here we rely on the fact that we populated the MeshSceneAdapter using AddComponents.
+			// This guarantees that each component has its own ActorAdapter. So we build a map
+			// of source component to actor adapter.
+			const FActorAdapter** FoundActor = DetailSampler.ActorComponentMap.Find(ChildMesh->SourceComponent);
+			if (Actor && ChildMesh && !FoundActor)
+			{
+				DetailSampler.ActorComponentMap.Add(ChildMesh->SourceComponent, Actor);
+			}
+		};
+		DetailMeshScene->ProcessActorChildMeshes(ProcessChildMesh);
+		DetailSampler.BaseToWorldTransform = BaseToWorldTransform;
+		DetailSampler.WorldToBaseTransform = BaseToWorldTransform.Inverse();
 		Baker->SetDetailSampler(&DetailSampler);
 
 		for (const EBakeMapType MapType : ENUM_EBAKEMAPTYPE_ALL)
@@ -419,15 +461,11 @@ void UBakeMultiMeshAttributeMapsTool::Setup()
 	Super::Setup();
 
 	// Initialize base mesh
-	const FTransformSRT3d BaseToWorld = UE::ToolTarget::GetLocalToWorldTransform(Targets[0]);
-	PreviewMesh->ProcessMesh([this, BaseToWorld](const FDynamicMesh3& Mesh)
+	PreviewMesh->ProcessMesh([this](const FDynamicMesh3& Mesh)
 	{
 		TargetMesh.Copy(Mesh);
 		TargetMeshTangents = MakeShared<FMeshTangentsd, ESPMode::ThreadSafe>(&TargetMesh);
 		TargetMeshTangents->CopyTriVertexTangents(Mesh);
-		
-		// FMeshSceneAdapter operates in world space, so ensure our mesh transformed to world.
-		MeshTransforms::ApplyTransform(TargetMesh, BaseToWorld);
 		TargetSpatial.SetMesh(&TargetMesh, true);
 	});
 
@@ -442,6 +480,9 @@ void UBakeMultiMeshAttributeMapsTool::Setup()
 		}
 	}
 	DetailMeshScene = MakeShared<FMeshSceneAdapter, ESPMode::ThreadSafe>();
+	// FMultiMeshMapBakerOp depends on adding meshes via AddComponents.
+	// If this changed, ensure the logic for generating the ActorComponentMap
+	// is updated in FMultiMeshMapBakerOp.
 	DetailMeshScene->AddComponents(DetailComponents);
 	DetailMeshScene->Build(FMeshSceneAdapterBuildOptions());
 	DetailMeshScene->BuildSpatialEvaluationCache();
@@ -547,6 +588,7 @@ TUniquePtr<UE::Geometry::TGenericDataOperator<FMeshMapBaker>> UBakeMultiMeshAttr
 	Op->DetailMeshScene = DetailMeshScene;
 	Op->BaseMesh = &TargetMesh;
 	Op->BaseMeshUVCharts = TargetMeshUVCharts;
+	Op->BaseToWorldTransform = UE::ToolTarget::GetLocalToWorldTransform(Targets[0]);
 	Op->BakeSettings = CachedBakeSettings;
 	Op->SampleFilterMask = CachedSampleFilterMask;
 
@@ -700,7 +742,7 @@ EBakeOpState UBakeMultiMeshAttributeMapsTool::UpdateResult_DetailMeshes()
 				if (static_cast<bool>(CachedBakeSettings.BakeMapTypes & EBakeMapType::Texture))
 				{
 					CachedMeshToColorImagesMap.Emplace(
-					ChildMesh->MeshSpatial,
+					(void*)ChildMesh,
 					FTextureImageData(CachedColorImages[*DataIndex].Get(), CachedColorUVLayers[*DataIndex]));
 				}
 			}
