MeshModelingTools: Allow user to create a new asset in CubeGrid tool if underlying asset is deleted.

diff --git a/Engine/Plugins/Experimental/MeshModelingToolsetExp/Source/MeshModelingToolsExp/Private/CubeGridTool.cpp b/Engine/Plugins/Experimental/MeshModelingToolsetExp/Source/MeshModelingToolsExp/Private/CubeGridTool.cpp
--- a/Engine/Plugins/Experimental/MeshModelingToolsetExp/Source/MeshModelingToolsExp/Private/CubeGridTool.cpp
+++ b/Engine/Plugins/Experimental/MeshModelingToolsetExp/Source/MeshModelingToolsExp/Private/CubeGridTool.cpp
@@ -21,6 +21,7 @@
 #include "Mechanics/DragAlignmentMechanic.h"
 #include "MeshOpPreviewHelpers.h" //UMeshOpPreviewWithBackgroundCompute
 #include "MeshDescriptionToDynamicMesh.h"
+#include "Misc/MessageDialog.h"
 #include "ModelingObjectsCreationAPI.h"
 #include "ModelingToolTargetUtil.h"
 #include "Properties/MeshMaterialProperties.h"
@@ -993,6 +994,7 @@ void UCubeGridTool::Shutdown(EToolShutdownType ShutdownType)
 		ApplyCornerMode(true);
 	}
 
+	bool bCreatingNewAsset = !Target && CurrentMesh->TriangleCount() > 0 && ShutdownType != EToolShutdownType::Cancel;
 	if (Target)
 	{
 		Cast<IPrimitiveComponentBackedTarget>(Target)->SetOwnerVisibility(true);
@@ -1010,12 +1012,23 @@ void UCubeGridTool::Shutdown(EToolShutdownType ShutdownType)
 				FConversionToMeshDescriptionOptions(), &OutputMaterialSet);
 			GetToolManager()->EndUndoTransaction();
 		}
-		else if (!Target->IsValid())
+		else if (!Target->IsValid() && CurrentMesh->TriangleCount() > 0)
 		{
-			UE_LOG(LogGeometry, Error, TEXT("CubeGridTool:: Edited mesh could not be committed (it was likely forcibly deleted from under the tool)."));
+			FText Title = LOCTEXT("RecreateAssetTitle", "Recreate Mesh Asset?");
+			EAppReturnType::Type Ret = FMessageDialog::Open(EAppMsgType::YesNo,
+				LOCTEXT("RecreateAssetQuestion", "The underlying asset that this tool was "
+					"operating on seems to no longer be valid (it was likely forcibly removed). "
+					"Would you like to recreate a new asset from the tool's current working "
+					"mesh? Selecting \"No\" or closing this window will discard the tool's "
+					"current work."));
+			if (Ret == EAppReturnType::Yes)
+			{
+				bCreatingNewAsset = true;
+			}
 		}
 	}
-	else if (CurrentMesh->TriangleCount() > 0 && ShutdownType != EToolShutdownType::Cancel)
+
+	if (bCreatingNewAsset)
 	{
 		GetToolManager()->BeginUndoTransaction(LOCTEXT("CubeGridToolCreateTransactionName", "Block Tool Create New"));
 
