ModelingTools: Add ObjectSpaceNormal and Position output types to BakeAll tool.

diff --git a/Engine/Plugins/Experimental/MeshModelingToolsetExp/Source/MeshModelingToolsExp/Private/BakeMultiMeshAttributeMapsTool.cpp b/Engine/Plugins/Experimental/MeshModelingToolsetExp/Source/MeshModelingToolsExp/Private/BakeMultiMeshAttributeMapsTool.cpp
--- a/Engine/Plugins/Experimental/MeshModelingToolsetExp/Source/MeshModelingToolsExp/Private/BakeMultiMeshAttributeMapsTool.cpp
+++ b/Engine/Plugins/Experimental/MeshModelingToolsetExp/Source/MeshModelingToolsExp/Private/BakeMultiMeshAttributeMapsTool.cpp
@@ -432,6 +432,19 @@ public:
 				Baker->AddEvaluator(NormalEval);
 				break;	
 			}
+			case EBakeMapType::Position:
+			{
+				TSharedPtr<FMeshPropertyMapEvaluator, ESPMode::ThreadSafe> PropertyEval = MakeShared<FMeshPropertyMapEvaluator, ESPMode::ThreadSafe>();
+				PropertyEval->Property = EMeshPropertyMapType::Position;
+				Baker->AddEvaluator(PropertyEval);
+				break;
+			}
+			case EBakeMapType::ObjectSpaceNormal:
+			{
+				TSharedPtr<FMeshPropertyMapEvaluator, ESPMode::ThreadSafe> PropertyEval = MakeShared<FMeshPropertyMapEvaluator, ESPMode::ThreadSafe>();
+				PropertyEval->Property = EMeshPropertyMapType::Normal;
+				Baker->AddEvaluator(PropertyEval);
+			}
 			case EBakeMapType::Texture:
 			{
 				TSharedPtr<FMeshResampleImageEvaluator, ESPMode::ThreadSafe> TextureEval = MakeShared<FMeshResampleImageEvaluator, ESPMode::ThreadSafe>();
diff --git a/Engine/Plugins/Experimental/MeshModelingToolsetExp/Source/MeshModelingToolsExp/Public/BakeMultiMeshAttributeMapsTool.h b/Engine/Plugins/Experimental/MeshModelingToolsetExp/Source/MeshModelingToolsExp/Public/BakeMultiMeshAttributeMapsTool.h
--- a/Engine/Plugins/Experimental/MeshModelingToolsetExp/Source/MeshModelingToolsExp/Public/BakeMultiMeshAttributeMapsTool.h
+++ b/Engine/Plugins/Experimental/MeshModelingToolsetExp/Source/MeshModelingToolsExp/Public/BakeMultiMeshAttributeMapsTool.h
@@ -40,7 +40,7 @@ class MESHMODELINGTOOLSEXP_API UBakeMultiMeshAttributeMapsToolProperties : publi
 public:
 	/** The map types to generate */
 	UPROPERTY(EditAnywhere, Category = BakeOutput, meta=(DisplayName="Output Types", Bitmask, BitmaskEnum=EBakeMapType,
-		ValidEnumValues="TangentSpaceNormal, Texture"))
+		ValidEnumValues="TangentSpaceNormal, ObjectSpaceNormal, Position, Texture"))
 	int32 MapTypes = (int32) EBakeMapType::None;
 
 	/** The map type index to preview */
