add a "Random Colors" option to the collision geometry visualizations

diff --git a/Engine/Plugins/Experimental/MeshLODToolset/Source/MeshLODToolset/Private/Tools/GenerateStaticMeshLODAssetTool.cpp b/Engine/Plugins/Experimental/MeshLODToolset/Source/MeshLODToolset/Private/Tools/GenerateStaticMeshLODAssetTool.cpp
--- a/Engine/Plugins/Experimental/MeshLODToolset/Source/MeshLODToolset/Private/Tools/GenerateStaticMeshLODAssetTool.cpp
+++ b/Engine/Plugins/Experimental/MeshLODToolset/Source/MeshLODToolset/Private/Tools/GenerateStaticMeshLODAssetTool.cpp
@@ -14,6 +14,7 @@
 #include "DynamicMeshEditor.h"
 #include "DynamicMesh/MeshTransforms.h"
 #include "DynamicMesh/MeshTangents.h"
+#include "Util/ColorConstants.h"
 
 #include "MeshDescriptionToDynamicMesh.h"
 #include "DynamicMeshToMeshDescription.h"
@@ -402,7 +403,7 @@ void UGenerateStaticMeshLODAssetTool::Setup()
 		PhysicsData.CopyGeometryToAggregate();
 		UE::PhysicsTools::InitializePreviewGeometryLines(PhysicsData,
 														 CollisionPreview,
-														 CollisionVizSettings->Color, CollisionVizSettings->LineThickness, 0.0f, 16);
+														 CollisionVizSettings->Color, CollisionVizSettings->LineThickness, 0.0f, 16, CollisionVizSettings->bRandomColors);
 
 		// Must happen on main thread, and GenerateProcess might be in use by an Op somewhere else
 		GenerateProcess->GraphEvalCriticalSection.Lock();
@@ -430,6 +431,7 @@ void UGenerateStaticMeshLODAssetTool::Setup()
 	AddToolPropertySource(CollisionVizSettings);
 	CollisionVizSettings->WatchProperty(CollisionVizSettings->LineThickness, [this](float NewValue) { bCollisionVisualizationDirty = true; });
 	CollisionVizSettings->WatchProperty(CollisionVizSettings->Color, [this](FColor NewValue) { bCollisionVisualizationDirty = true; });
+	CollisionVizSettings->WatchProperty(CollisionVizSettings->bRandomColors, [this](bool bNewValue) { bCollisionVisualizationDirty = true; });
 	CollisionVizSettings->WatchProperty(CollisionVizSettings->bShowHidden, [this](bool bNewValue) { bCollisionVisualizationDirty = true; });
 
 	CollisionPreview = NewObject<UPreviewGeometry>(this);
@@ -533,10 +535,11 @@ void UGenerateStaticMeshLODAssetTool::UpdateCollisionVisualization()
 	FColor UseColor = CollisionVizSettings->Color;
 	LineMaterial = ToolSetupUtil::GetDefaultLineComponentMaterial(GetToolManager(), !CollisionVizSettings->bShowHidden);
 
+	int32 ColorIdx = 0;
 	CollisionPreview->UpdateAllLineSets([&](ULineSetComponent* LineSet)
 	{
 		LineSet->SetAllLinesThickness(UseThickness);
-		LineSet->SetAllLinesColor(UseColor);
+		LineSet->SetAllLinesColor(CollisionVizSettings->bRandomColors ? LinearColors::SelectFColor(ColorIdx++) : UseColor);
 	});
 	CollisionPreview->SetAllLineSetsMaterial(LineMaterial);
 }
diff --git a/Engine/Plugins/Experimental/MeshModelingToolsetExp/Source/MeshModelingToolsExp/Private/Physics/CollisionGeometryVisualization.cpp b/Engine/Plugins/Experimental/MeshModelingToolsetExp/Source/MeshModelingToolsExp/Private/Physics/CollisionGeometryVisualization.cpp
--- a/Engine/Plugins/Experimental/MeshModelingToolsetExp/Source/MeshModelingToolsExp/Private/Physics/CollisionGeometryVisualization.cpp
+++ b/Engine/Plugins/Experimental/MeshModelingToolsetExp/Source/MeshModelingToolsExp/Private/Physics/CollisionGeometryVisualization.cpp
@@ -2,116 +2,145 @@
 
 #include "Physics/CollisionGeometryVisualization.h"
 #include "Generators/LineSegmentGenerators.h"
+#include "Util/ColorConstants.h"
 
 using namespace UE::Geometry;
 
 void UE::PhysicsTools::InitializePreviewGeometryLines(const FPhysicsDataCollection& PhysicsData, UPreviewGeometry* PreviewGeom,
-	const FColor& LineColor, float LineThickness, float DepthBias, int32 CircleStepResolution)
+	const FColor& LineColor, float LineThickness, float DepthBias, int32 CircleStepResolution, bool bRandomColors)
 {
 	int32 CircleSteps = FMath::Max(4, CircleStepResolution);
 	FColor SphereColor = LineColor;
 	FColor BoxColor = LineColor;
 	FColor ConvexColor = LineColor;
 	FColor CapsuleColor = LineColor;
+	int32 OverrideColorIdx = 0;
+
+	auto GetNextColor = [&OverrideColorIdx, bRandomColors](FColor DefaultColor)
+	{
+		if (bRandomColors)
+		{
+			return LinearColors::SelectFColor(OverrideColorIdx++);
+		}
+		return DefaultColor;
+	};
 
 	const FKAggregateGeom& AggGeom = PhysicsData.AggGeom;
 
 	// spheres are draw as 3 orthogonal circles
-	PreviewGeom->CreateOrUpdateLineSet(TEXT("Spheres"), AggGeom.SphereElems.Num(), [&](int32 Index, TArray<FRenderableLine>& LinesOut)
+	for (int32 Index = 0; Index < AggGeom.SphereElems.Num(); Index++)
 	{
-		const FKSphereElem& Sphere = AggGeom.SphereElems[Index];
-		FTransform ElemTransform = Sphere.GetTransform();
-		ElemTransform.ScaleTranslation(PhysicsData.ExternalScale3D);
-		FTransformSRT3f ElemTransformf(ElemTransform);
-		float Radius = PhysicsData.ExternalScale3D.GetAbsMin() * Sphere.Radius;
-		UE::Geometry::GenerateCircleSegments<float>(CircleSteps, Radius, FVector3f::Zero(), FVector3f::UnitX(), FVector3f::UnitY(), ElemTransformf,
-			[&](const FVector3f& A, const FVector3f& B) { LinesOut.Add(FRenderableLine((FVector)A, (FVector)B, SphereColor, LineThickness, DepthBias)); });
-		UE::Geometry::GenerateCircleSegments<float>(CircleSteps, Radius, FVector3f::Zero(), FVector3f::UnitX(), FVector3f::UnitZ(), ElemTransformf,
-			[&](const FVector3f& A, const FVector3f& B) { LinesOut.Add(FRenderableLine((FVector)A, (FVector)B, SphereColor, LineThickness, DepthBias)); });
-		UE::Geometry::GenerateCircleSegments<float>(CircleSteps, Radius, FVector3f::Zero(), FVector3f::UnitY(), FVector3f::UnitZ(), ElemTransformf,
-			[&](const FVector3f& A, const FVector3f& B) { LinesOut.Add(FRenderableLine((FVector)A, (FVector)B, SphereColor, LineThickness, DepthBias)); });
-	});
+		PreviewGeom->CreateOrUpdateLineSet(FString::Printf(TEXT("Spheres %d"), Index), 1, [&](int32 UnusedIndex, TArray<FRenderableLine>& LinesOut)
+		{
+			FColor Color = GetNextColor(SphereColor);
+
+			const FKSphereElem& Sphere = AggGeom.SphereElems[Index];
+			FTransform ElemTransform = Sphere.GetTransform();
+			ElemTransform.ScaleTranslation(PhysicsData.ExternalScale3D);
+			FTransformSRT3f ElemTransformf(ElemTransform);
+			float Radius = PhysicsData.ExternalScale3D.GetAbsMin() * Sphere.Radius;
+			UE::Geometry::GenerateCircleSegments<float>(CircleSteps, Radius, FVector3f::Zero(), FVector3f::UnitX(), FVector3f::UnitY(), ElemTransformf,
+				[&](const FVector3f& A, const FVector3f& B) { LinesOut.Add(FRenderableLine((FVector)A, (FVector)B, Color, LineThickness, DepthBias)); });
+			UE::Geometry::GenerateCircleSegments<float>(CircleSteps, Radius, FVector3f::Zero(), FVector3f::UnitX(), FVector3f::UnitZ(), ElemTransformf,
+				[&](const FVector3f& A, const FVector3f& B) { LinesOut.Add(FRenderableLine((FVector)A, (FVector)B, Color, LineThickness, DepthBias)); });
+			UE::Geometry::GenerateCircleSegments<float>(CircleSteps, Radius, FVector3f::Zero(), FVector3f::UnitY(), FVector3f::UnitZ(), ElemTransformf,
+				[&](const FVector3f& A, const FVector3f& B) { LinesOut.Add(FRenderableLine((FVector)A, (FVector)B, Color, LineThickness, DepthBias)); });
+		});
+	}
 
 
 	// boxes are drawn as boxes
-	PreviewGeom->CreateOrUpdateLineSet(TEXT("Boxes"), AggGeom.BoxElems.Num(), [&](int32 Index, TArray<FRenderableLine>& LinesOut)
+	for (int32 Index = 0; Index < AggGeom.BoxElems.Num(); Index++)
 	{
-		const FKBoxElem& Box = AggGeom.BoxElems[Index];
-		FTransform ElemTransform = Box.GetTransform();
-		ElemTransform.ScaleTranslation(PhysicsData.ExternalScale3D);
-		FTransformSRT3f ElemTransformf(ElemTransform);
-		FVector3f HalfDimensions(
-			PhysicsData.ExternalScale3D.X * Box.X * 0.5f,
-			PhysicsData.ExternalScale3D.Y * Box.Y * 0.5f,
-			PhysicsData.ExternalScale3D.Z * Box.Z * 0.5f);
-		UE::Geometry::GenerateBoxSegments<float>(HalfDimensions, FVector3f::Zero(), FVector3f::UnitX(), FVector3f::UnitY(), FVector3f::UnitZ(), ElemTransformf,
-			[&](const FVector3f& A, const FVector3f& B) { LinesOut.Add(FRenderableLine((FVector)A, (FVector)B, BoxColor, LineThickness, DepthBias)); });
-	});
+		PreviewGeom->CreateOrUpdateLineSet(FString::Printf(TEXT("Boxes %d"), Index), 1, [&](int32 UnusedIndex, TArray<FRenderableLine>& LinesOut)
+		{
+			FColor Color = GetNextColor(BoxColor);
+
+			const FKBoxElem& Box = AggGeom.BoxElems[Index];
+			FTransform ElemTransform = Box.GetTransform();
+			ElemTransform.ScaleTranslation(PhysicsData.ExternalScale3D);
+			FTransformSRT3f ElemTransformf(ElemTransform);
+			FVector3f HalfDimensions(
+				PhysicsData.ExternalScale3D.X * Box.X * 0.5f,
+				PhysicsData.ExternalScale3D.Y * Box.Y * 0.5f,
+				PhysicsData.ExternalScale3D.Z * Box.Z * 0.5f);
+			UE::Geometry::GenerateBoxSegments<float>(HalfDimensions, FVector3f::Zero(), FVector3f::UnitX(), FVector3f::UnitY(), FVector3f::UnitZ(), ElemTransformf,
+				[&](const FVector3f& A, const FVector3f& B) { LinesOut.Add(FRenderableLine((FVector)A, (FVector)B, Color, LineThickness, DepthBias)); });
+		});
+	}
 
 
 	// capsules are draw as two hemispheres (with 3 intersecting arcs/circles) and connecting lines
-	PreviewGeom->CreateOrUpdateLineSet(TEXT("Capsules"), AggGeom.SphylElems.Num(), [&](int32 Index, TArray<FRenderableLine>& LinesOut)
+	for (int32 Index = 0; Index < AggGeom.SphylElems.Num(); Index++)
 	{
-		const FKSphylElem& Capsule = AggGeom.SphylElems[Index];
-		FTransform ElemTransform = Capsule.GetTransform();
-		ElemTransform.ScaleTranslation(PhysicsData.ExternalScale3D);
-		FTransformSRT3f ElemTransformf(ElemTransform);
-		const float HalfLength = Capsule.GetScaledCylinderLength(PhysicsData.ExternalScale3D) * .5f;
-		const float Radius = Capsule.GetScaledRadius(PhysicsData.ExternalScale3D);
-		FVector3f Top(0, 0, HalfLength), Bottom(0, 0, -HalfLength);
-
-		// top and bottom circles
-		UE::Geometry::GenerateCircleSegments<float>(CircleSteps, Radius, Top, FVector3f::UnitX(), FVector3f::UnitY(), ElemTransformf,
-			[&](const FVector3f& A, const FVector3f& B) { LinesOut.Add(FRenderableLine((FVector)A, (FVector)B, CapsuleColor, LineThickness, DepthBias)); });
-		UE::Geometry::GenerateCircleSegments<float>(CircleSteps, Radius, Bottom, FVector3f::UnitX(), FVector3f::UnitY(), ElemTransformf,
-			[&](const FVector3f& A, const FVector3f& B) { LinesOut.Add(FRenderableLine((FVector)A, (FVector)B, CapsuleColor, LineThickness, DepthBias)); });
-
-		// top dome
-		UE::Geometry::GenerateArcSegments<float>(CircleSteps, Radius, 0.0, PI, Top, FVector3f::UnitY(), FVector3f::UnitZ(), ElemTransformf,
-			[&](const FVector3f& A, const FVector3f& B) { LinesOut.Add(FRenderableLine((FVector)A, (FVector)B, CapsuleColor, LineThickness, DepthBias)); });
-		UE::Geometry::GenerateArcSegments<float>(CircleSteps, Radius, 0.0, PI, Top, FVector3f::UnitX(), FVector3f::UnitZ(), ElemTransformf,
-			[&](const FVector3f& A, const FVector3f& B) { LinesOut.Add(FRenderableLine((FVector)A, (FVector)B, CapsuleColor, LineThickness, DepthBias)); });
-
-		// bottom dome
-		UE::Geometry::GenerateArcSegments<float>(CircleSteps, Radius, 0.0, -PI, Bottom, FVector3f::UnitY(), FVector3f::UnitZ(), ElemTransformf,
-			[&](const FVector3f& A, const FVector3f& B) { LinesOut.Add(FRenderableLine((FVector)A, (FVector)B, CapsuleColor, LineThickness, DepthBias)); });
-		UE::Geometry::GenerateArcSegments<float>(CircleSteps, Radius, 0.0, -PI, Bottom, FVector3f::UnitX(), FVector3f::UnitZ(), ElemTransformf,
-			[&](const FVector3f& A, const FVector3f& B) { LinesOut.Add(FRenderableLine((FVector)A, (FVector)B, CapsuleColor, LineThickness, DepthBias)); });
-
-		// connecting lines
-		for (int k = 0; k < 2; ++k)
+		PreviewGeom->CreateOrUpdateLineSet(FString::Printf(TEXT("Capsules %d"), Index), 1, [&](int32 UnusedIndex, TArray<FRenderableLine>& LinesOut)
 		{
-			FVector DX = (k < 1) ? FVector(-Radius, 0, 0) : FVector(Radius, 0, 0);
-			LinesOut.Add(FRenderableLine(
-				ElemTransform.TransformPosition((FVector)Top + DX), 
-				ElemTransform.TransformPosition((FVector)Bottom + DX), CapsuleColor, LineThickness, DepthBias));
-			FVector DY = (k < 1) ? FVector(0, -Radius, 0) : FVector(0, Radius, 0);
-			LinesOut.Add(FRenderableLine(
-				ElemTransform.TransformPosition((FVector)Top + DY), 
-				ElemTransform.TransformPosition((FVector)Bottom + DY), CapsuleColor, LineThickness, DepthBias));
-		}
-	});
-
-
+			FColor Color = GetNextColor(CapsuleColor);
+
+			const FKSphylElem& Capsule = AggGeom.SphylElems[Index];
+			FTransform ElemTransform = Capsule.GetTransform();
+			ElemTransform.ScaleTranslation(PhysicsData.ExternalScale3D);
+			FTransformSRT3f ElemTransformf(ElemTransform);
+			const float HalfLength = Capsule.GetScaledCylinderLength(PhysicsData.ExternalScale3D) * .5f;
+			const float Radius = Capsule.GetScaledRadius(PhysicsData.ExternalScale3D);
+			FVector3f Top(0, 0, HalfLength), Bottom(0, 0, -HalfLength);
+
+			// top and bottom circles
+			UE::Geometry::GenerateCircleSegments<float>(CircleSteps, Radius, Top, FVector3f::UnitX(), FVector3f::UnitY(), ElemTransformf,
+				[&](const FVector3f& A, const FVector3f& B) { LinesOut.Add(FRenderableLine((FVector)A, (FVector)B, Color, LineThickness, DepthBias)); });
+			UE::Geometry::GenerateCircleSegments<float>(CircleSteps, Radius, Bottom, FVector3f::UnitX(), FVector3f::UnitY(), ElemTransformf,
+				[&](const FVector3f& A, const FVector3f& B) { LinesOut.Add(FRenderableLine((FVector)A, (FVector)B, Color, LineThickness, DepthBias)); });
+
+			// top dome
+			UE::Geometry::GenerateArcSegments<float>(CircleSteps, Radius, 0.0, PI, Top, FVector3f::UnitY(), FVector3f::UnitZ(), ElemTransformf,
+				[&](const FVector3f& A, const FVector3f& B) { LinesOut.Add(FRenderableLine((FVector)A, (FVector)B, Color, LineThickness, DepthBias)); });
+			UE::Geometry::GenerateArcSegments<float>(CircleSteps, Radius, 0.0, PI, Top, FVector3f::UnitX(), FVector3f::UnitZ(), ElemTransformf,
+				[&](const FVector3f& A, const FVector3f& B) { LinesOut.Add(FRenderableLine((FVector)A, (FVector)B, Color, LineThickness, DepthBias)); });
+
+			// bottom dome
+			UE::Geometry::GenerateArcSegments<float>(CircleSteps, Radius, 0.0, -PI, Bottom, FVector3f::UnitY(), FVector3f::UnitZ(), ElemTransformf,
+				[&](const FVector3f& A, const FVector3f& B) { LinesOut.Add(FRenderableLine((FVector)A, (FVector)B, Color, LineThickness, DepthBias)); });
+			UE::Geometry::GenerateArcSegments<float>(CircleSteps, Radius, 0.0, -PI, Bottom, FVector3f::UnitX(), FVector3f::UnitZ(), ElemTransformf,
+				[&](const FVector3f& A, const FVector3f& B) { LinesOut.Add(FRenderableLine((FVector)A, (FVector)B, Color, LineThickness, DepthBias)); });
+
+			// connecting lines
+			for (int k = 0; k < 2; ++k)
+			{
+				FVector DX = (k < 1) ? FVector(-Radius, 0, 0) : FVector(Radius, 0, 0);
+				LinesOut.Add(FRenderableLine(
+					ElemTransform.TransformPosition((FVector)Top + DX),
+					ElemTransform.TransformPosition((FVector)Bottom + DX), Color, LineThickness, DepthBias));
+				FVector DY = (k < 1) ? FVector(0, -Radius, 0) : FVector(0, Radius, 0);
+				LinesOut.Add(FRenderableLine(
+					ElemTransform.TransformPosition((FVector)Top + DY),
+					ElemTransform.TransformPosition((FVector)Bottom + DY), Color, LineThickness, DepthBias));
+			}
+		});
+	}
 
 	// convexes are drawn as mesh edges
-	PreviewGeom->CreateOrUpdateLineSet(TEXT("Convexes"), AggGeom.ConvexElems.Num(), [&](int32 Index, TArray<FRenderableLine>& LinesOut)
+	for (int32 Index = 0; Index < AggGeom.ConvexElems.Num(); Index++)
 	{
-		const FKConvexElem& Convex = AggGeom.ConvexElems[Index];
-		FTransform ElemTransform = Convex.GetTransform();
-		ElemTransform.ScaleTranslation(PhysicsData.ExternalScale3D);
-		ElemTransform.SetScale3D(PhysicsData.ExternalScale3D);
-		int32 NumTriangles = Convex.IndexData.Num() / 3;
-		for (int32 k = 0; k < NumTriangles; ++k)
+		PreviewGeom->CreateOrUpdateLineSet(FString::Printf(TEXT("Convex %d"), Index), 1, [&](int32 UnusedIndex, TArray<FRenderableLine>& LinesOut)
 		{
-			FVector A = ElemTransform.TransformPosition(Convex.VertexData[Convex.IndexData[3 * k]]);
-			FVector B = ElemTransform.TransformPosition(Convex.VertexData[Convex.IndexData[3 * k + 1]]);
-			FVector C = ElemTransform.TransformPosition(Convex.VertexData[Convex.IndexData[3 * k + 2]]);
-			LinesOut.Add(FRenderableLine(A, B, ConvexColor, LineThickness, DepthBias));
-			LinesOut.Add(FRenderableLine(B, C, ConvexColor, LineThickness, DepthBias));
-			LinesOut.Add(FRenderableLine(C, A, ConvexColor, LineThickness, DepthBias));
-		}
-	});
+			FColor Color = GetNextColor(ConvexColor);
+
+			const FKConvexElem& Convex = AggGeom.ConvexElems[Index];
+			FTransform ElemTransform = Convex.GetTransform();
+			ElemTransform.ScaleTranslation(PhysicsData.ExternalScale3D);
+			ElemTransform.SetScale3D(PhysicsData.ExternalScale3D);
+			int32 NumTriangles = Convex.IndexData.Num() / 3;
+			for (int32 k = 0; k < NumTriangles; ++k)
+			{
+				FVector A = ElemTransform.TransformPosition(Convex.VertexData[Convex.IndexData[3 * k]]);
+				FVector B = ElemTransform.TransformPosition(Convex.VertexData[Convex.IndexData[3 * k + 1]]);
+				FVector C = ElemTransform.TransformPosition(Convex.VertexData[Convex.IndexData[3 * k + 2]]);
+				LinesOut.Add(FRenderableLine(A, B, Color, LineThickness, DepthBias));
+				LinesOut.Add(FRenderableLine(B, C, Color, LineThickness, DepthBias));
+				LinesOut.Add(FRenderableLine(C, A, Color, LineThickness, DepthBias));
+			}
+		});
+	}
 
 
 	// Unclear whether we actually use these in the Engine, for UBodySetup? Does not appear to be supported by UxX import system,
diff --git a/Engine/Plugins/Experimental/MeshModelingToolsetExp/Source/MeshModelingToolsExp/Private/Physics/ExtractCollisionGeometryTool.cpp b/Engine/Plugins/Experimental/MeshModelingToolsetExp/Source/MeshModelingToolsExp/Private/Physics/ExtractCollisionGeometryTool.cpp
--- a/Engine/Plugins/Experimental/MeshModelingToolsetExp/Source/MeshModelingToolsExp/Private/Physics/ExtractCollisionGeometryTool.cpp
+++ b/Engine/Plugins/Experimental/MeshModelingToolsetExp/Source/MeshModelingToolsExp/Private/Physics/ExtractCollisionGeometryTool.cpp
@@ -89,6 +89,8 @@ void UExtractCollisionGeometryTool::Setup()
 	AddToolPropertySource(VizSettings);
 	VizSettings->WatchProperty(VizSettings->LineThickness, [this](float NewValue) { bVisualizationDirty = true; });
 	VizSettings->WatchProperty(VizSettings->Color, [this](FColor NewValue) { bVisualizationDirty = true; });
+	VizSettings->WatchProperty(VizSettings->bRandomColors, [this](bool bNewValue) { bVisualizationDirty = true; });
+	VizSettings->WatchProperty(VizSettings->bShowHidden, [this](bool bNewValue) { bVisualizationDirty = true; });
 
 	UBodySetup* BodySetup = UE::ToolTarget::GetPhysicsBodySetup(Target);
 	if (BodySetup)
@@ -103,7 +105,7 @@ void UExtractCollisionGeometryTool::Setup()
 		PreviewElements->CreateInWorld(UE::ToolTarget::GetTargetActor(Target)->GetWorld(), TargetTransform);
 		
 		UE::PhysicsTools::InitializePreviewGeometryLines(*PhysicsInfo, PreviewElements,
-			VizSettings->Color, VizSettings->LineThickness, 0.0f, 16);
+			VizSettings->Color, VizSettings->LineThickness, 0.0f, 16, VizSettings->bRandomColors);
 
 		ObjectProps = NewObject<UPhysicsObjectToolPropertySet>(this);
 		UE::PhysicsTools::InitializePhysicsToolObjectPropertySet(PhysicsInfo.Get(), ObjectProps);
@@ -222,10 +224,11 @@ void UExtractCollisionGeometryTool::UpdateVisualization()
 {
 	float UseThickness = VizSettings->LineThickness;
 	FColor UseColor = VizSettings->Color;
+	int32 ColorIdx = 0;
 	PreviewElements->UpdateAllLineSets([&](ULineSetComponent* LineSet)
 		{
 			LineSet->SetAllLinesThickness(UseThickness);
-			LineSet->SetAllLinesColor(UseColor);
+			LineSet->SetAllLinesColor(VizSettings->bRandomColors ? LinearColors::SelectFColor(ColorIdx++) : UseColor);
 		});
 	
 	UMaterialInterface* LineMaterial = ToolSetupUtil::GetDefaultLineComponentMaterial(GetToolManager(), !VizSettings->bShowHidden);
diff --git a/Engine/Plugins/Experimental/MeshModelingToolsetExp/Source/MeshModelingToolsExp/Private/Physics/PhysicsInspectorTool.cpp b/Engine/Plugins/Experimental/MeshModelingToolsetExp/Source/MeshModelingToolsExp/Private/Physics/PhysicsInspectorTool.cpp
--- a/Engine/Plugins/Experimental/MeshModelingToolsetExp/Source/MeshModelingToolsExp/Private/Physics/PhysicsInspectorTool.cpp
+++ b/Engine/Plugins/Experimental/MeshModelingToolsetExp/Source/MeshModelingToolsExp/Private/Physics/PhysicsInspectorTool.cpp
@@ -6,6 +6,7 @@
 #include "ToolSetupUtil.h"
 #include "ModelingToolTargetUtil.h"
 #include "Drawing/PreviewGeometryActor.h"
+#include "Util/ColorConstants.h"
 
 #include "Physics/PhysicsDataCollection.h"
 #include "Physics/CollisionGeometryVisualization.h"
@@ -48,6 +49,7 @@ void UPhysicsInspectorTool::Setup()
 	AddToolPropertySource(VizSettings);
 	VizSettings->WatchProperty(VizSettings->LineThickness, [this](float NewValue) { bVisualizationDirty = true; });
 	VizSettings->WatchProperty(VizSettings->Color, [this](FColor NewValue) { bVisualizationDirty = true; });
+	VizSettings->WatchProperty(VizSettings->bRandomColors, [this](bool bNewValue) { bVisualizationDirty = true; });
 	VizSettings->WatchProperty(VizSettings->bShowHidden, [this](bool bNewValue) { bVisualizationDirty = true; });
 
 	for (int32 ComponentIdx = 0; ComponentIdx < Targets.Num(); ComponentIdx++)
@@ -113,12 +115,13 @@ void UPhysicsInspectorTool::UpdateVisualization()
 	FColor UseColor = VizSettings->Color;
 	LineMaterial = ToolSetupUtil::GetDefaultLineComponentMaterial(GetToolManager(), !VizSettings->bShowHidden);
 
+	int32 ColorIdx = 0;
 	for (UPreviewGeometry* Preview : PreviewElements)
 	{
 		Preview->UpdateAllLineSets([&](ULineSetComponent* LineSet)
 		{
 			LineSet->SetAllLinesThickness(UseThickness);
-			LineSet->SetAllLinesColor(UseColor);
+			LineSet->SetAllLinesColor(VizSettings->bRandomColors ? LinearColors::SelectFColor(ColorIdx++) : UseColor);
 		});
 		Preview->SetAllLineSetsMaterial(LineMaterial);
 	}
@@ -136,7 +139,7 @@ void UPhysicsInspectorTool::InitializeObjectProperties(const FPhysicsDataCollect
 void UPhysicsInspectorTool::InitializeGeometry(const FPhysicsDataCollection& PhysicsData, UPreviewGeometry* PreviewGeom)
 {
 	UE::PhysicsTools::InitializePreviewGeometryLines(PhysicsData, PreviewGeom,
-		VizSettings->Color, VizSettings->LineThickness, 0.0f, 16);
+		VizSettings->Color, VizSettings->LineThickness, 0.0f, 16, VizSettings->bRandomColors);
 }
 
 
diff --git a/Engine/Plugins/Experimental/MeshModelingToolsetExp/Source/MeshModelingToolsExp/Private/Physics/SetCollisionGeometryTool.cpp b/Engine/Plugins/Experimental/MeshModelingToolsetExp/Source/MeshModelingToolsExp/Private/Physics/SetCollisionGeometryTool.cpp
--- a/Engine/Plugins/Experimental/MeshModelingToolsetExp/Source/MeshModelingToolsExp/Private/Physics/SetCollisionGeometryTool.cpp
+++ b/Engine/Plugins/Experimental/MeshModelingToolsetExp/Source/MeshModelingToolsExp/Private/Physics/SetCollisionGeometryTool.cpp
@@ -9,6 +9,7 @@
 #include "Selections/MeshConnectedComponents.h"
 #include "DynamicSubmesh3.h"
 #include "Polygroups/PolygroupUtil.h"
+#include "Util/ColorConstants.h"
 
 #include "ShapeApproximation/ShapeDetection3.h"
 #include "ShapeApproximation/MeshSimpleShapeApproximation.h"
@@ -233,6 +234,7 @@ void USetCollisionGeometryTool::Setup()
 	AddToolPropertySource(VizSettings);
 	VizSettings->WatchProperty(VizSettings->LineThickness, [this](float NewValue) { bVisualizationDirty = true; });
 	VizSettings->WatchProperty(VizSettings->Color, [this](FColor NewValue) { bVisualizationDirty = true; });
+	VizSettings->WatchProperty(VizSettings->bRandomColors, [this](bool bNewValue) { bVisualizationDirty = true; });
 	VizSettings->WatchProperty(VizSettings->bShowHidden, [this](bool bNewValue) { bVisualizationDirty = true; });
 
 	// add option for collision properties
@@ -419,7 +421,7 @@ void USetCollisionGeometryTool::OnTick(float DeltaTime)
 				// update visualization
 				PreviewGeom->RemoveAllLineSets();
 				UE::PhysicsTools::InitializePreviewGeometryLines(*GeneratedCollision, PreviewGeom,
-					VizSettings->Color, VizSettings->LineThickness, 0.0f, 16);
+					VizSettings->Color, VizSettings->LineThickness, 0.0f, 16, VizSettings->bRandomColors);
 
 				// update property set
 				CollisionProps->Reset();
@@ -504,10 +506,11 @@ void USetCollisionGeometryTool::UpdateVisualization()
 {
 	float UseThickness = VizSettings->LineThickness;
 	FColor UseColor = VizSettings->Color;
+	int32 ColorIdx = 0;
 	PreviewGeom->UpdateAllLineSets([&](ULineSetComponent* LineSet)
 	{
 		LineSet->SetAllLinesThickness(UseThickness);
-		LineSet->SetAllLinesColor(UseColor);
+		LineSet->SetAllLinesColor(VizSettings->bRandomColors ? LinearColors::SelectFColor(ColorIdx++) : UseColor);
 	});
 
 	LineMaterial = ToolSetupUtil::GetDefaultLineComponentMaterial(GetToolManager(), !VizSettings->bShowHidden);
diff --git a/Engine/Plugins/Experimental/MeshModelingToolsetExp/Source/MeshModelingToolsExp/Public/Physics/CollisionGeometryVisualization.h b/Engine/Plugins/Experimental/MeshModelingToolsetExp/Source/MeshModelingToolsExp/Public/Physics/CollisionGeometryVisualization.h
--- a/Engine/Plugins/Experimental/MeshModelingToolsetExp/Source/MeshModelingToolsExp/Public/Physics/CollisionGeometryVisualization.h
+++ b/Engine/Plugins/Experimental/MeshModelingToolsetExp/Source/MeshModelingToolsExp/Public/Physics/CollisionGeometryVisualization.h
@@ -14,7 +14,7 @@ namespace UE
 		 * Spheres and Capsules are drawn as 3-axis wireframes. Convexes are added as wireframes.
 		 */
 		void MESHMODELINGTOOLSEXP_API InitializePreviewGeometryLines(const FPhysicsDataCollection& PhysicsData, UPreviewGeometry* PreviewGeom,
-			const FColor& LineColor, float LineThickness, float DepthBias = 0.0, int32 CircleStepResolution = 16 );
+			const FColor& LineColor, float LineThickness, float DepthBias = 0.0, int32 CircleStepResolution = 16, bool bRandomColors = true );
 	}
 }
 
diff --git a/Engine/Plugins/Experimental/MeshModelingToolsetExp/Source/MeshModelingToolsExp/Public/Physics/CollisionPropertySets.h b/Engine/Plugins/Experimental/MeshModelingToolsetExp/Source/MeshModelingToolsExp/Public/Physics/CollisionPropertySets.h
--- a/Engine/Plugins/Experimental/MeshModelingToolsetExp/Source/MeshModelingToolsExp/Public/Physics/CollisionPropertySets.h
+++ b/Engine/Plugins/Experimental/MeshModelingToolsetExp/Source/MeshModelingToolsExp/Public/Physics/CollisionPropertySets.h
@@ -132,6 +132,9 @@ public:
 	bool bShowHidden = false;
 
 	UPROPERTY(EditAnywhere, Category = "Collision Visualization")
+	bool bRandomColors = true;
+
+	UPROPERTY(EditAnywhere, Category = "Collision Visualization", meta = (EditCondition = "!bRandomColors"))
 	FColor Color = FColor::Red;
 };
 
