ModelingTools: Fix PVSStudio static analysis warning in BakeAll.

diff --git a/Engine/Plugins/Experimental/MeshModelingToolsetExp/Source/MeshModelingToolsExp/Private/BakeMultiMeshAttributeMapsTool.cpp b/Engine/Plugins/Experimental/MeshModelingToolsetExp/Source/MeshModelingToolsExp/Private/BakeMultiMeshAttributeMapsTool.cpp
--- a/Engine/Plugins/Experimental/MeshModelingToolsetExp/Source/MeshModelingToolsExp/Private/BakeMultiMeshAttributeMapsTool.cpp
+++ b/Engine/Plugins/Experimental/MeshModelingToolsetExp/Source/MeshModelingToolsExp/Private/BakeMultiMeshAttributeMapsTool.cpp
@@ -408,13 +408,16 @@ public:
 		// TODO: Precompute the ActorComponent to ActorChildMesh map in the tool.
 		auto ProcessChildMesh = [&DetailSampler](const FActorAdapter* Actor, const FActorChildMesh* ChildMesh)
 		{
-			// Here we rely on the fact that we populated the MeshSceneAdapter using AddComponents.
-			// This guarantees that each component has its own ActorAdapter. So we build a map
-			// of source component to actor adapter.
-			const FActorAdapter** FoundActor = DetailSampler.ActorComponentMap.Find(ChildMesh->SourceComponent);
-			if (Actor && ChildMesh && !FoundActor)
+			if (Actor && ChildMesh)
 			{
-				DetailSampler.ActorComponentMap.Add(ChildMesh->SourceComponent, Actor);
+				// Here we rely on the fact that we populated the MeshSceneAdapter using AddComponents.
+				// This guarantees that each component has its own ActorAdapter. So we build a map
+				// of source component to actor adapter.
+				const FActorAdapter** FoundActor = DetailSampler.ActorComponentMap.Find(ChildMesh->SourceComponent);
+				if (!FoundActor)
+				{
+					DetailSampler.ActorComponentMap.Add(ChildMesh->SourceComponent, Actor);
+				}
 			}
 		};
 		DetailMeshScene->ProcessActorChildMeshes(ProcessChildMesh);
