expose min part thickness parameter in convex decomposition algorithm out to Msh2Coll tool and Geometry Script api

diff --git a/Engine/Plugins/Experimental/GeometryScripting/Source/GeometryScriptingCore/Private/CollisionFunctions.cpp b/Engine/Plugins/Experimental/GeometryScripting/Source/GeometryScriptingCore/Private/CollisionFunctions.cpp
--- a/Engine/Plugins/Experimental/GeometryScripting/Source/GeometryScriptingCore/Private/CollisionFunctions.cpp
+++ b/Engine/Plugins/Experimental/GeometryScripting/Source/GeometryScriptingCore/Private/CollisionFunctions.cpp
@@ -97,6 +97,7 @@ void ComputeCollisionFromMesh(
 			ShapeGenerator.ConvexDecompositionMaxPieces = Options.MaxConvexHullsPerMesh;
 			ShapeGenerator.ConvexDecompositionSearchFactor = Options.ConvexDecompositionSearchFactor;
 			ShapeGenerator.ConvexDecompositionErrorTolerance = Options.ConvexDecompositionErrorTolerance;
+			ShapeGenerator.ConvexDecompositionMinPartThickness = Options.ConvexDecompositionMinPartThickness;
 			ShapeGenerator.Generate_ConvexHullDecompositions(NewCollision.Geometry);
 		}
 		else
diff --git a/Engine/Plugins/Experimental/GeometryScripting/Source/GeometryScriptingCore/Public/GeometryScript/CollisionFunctions.h b/Engine/Plugins/Experimental/GeometryScripting/Source/GeometryScriptingCore/Public/GeometryScript/CollisionFunctions.h
--- a/Engine/Plugins/Experimental/GeometryScripting/Source/GeometryScriptingCore/Public/GeometryScript/CollisionFunctions.h
+++ b/Engine/Plugins/Experimental/GeometryScripting/Source/GeometryScriptingCore/Public/GeometryScript/CollisionFunctions.h
@@ -75,6 +75,9 @@ public:
 	UPROPERTY(BlueprintReadWrite, Category = Options)
 	float ConvexDecompositionErrorTolerance = 0;
 
+	UPROPERTY(BlueprintReadWrite, Category = Options)
+	float ConvexDecompositionMinPartThickness = 0.1;
+
 	UPROPERTY(BlueprintReadWrite, Category = Options)
 	float SweptHullSimplifyTolerance = 0.1;
 
diff --git a/Engine/Plugins/Experimental/MeshModelingToolsetExp/Source/MeshModelingToolsExp/Private/Physics/SetCollisionGeometryTool.cpp b/Engine/Plugins/Experimental/MeshModelingToolsetExp/Source/MeshModelingToolsExp/Private/Physics/SetCollisionGeometryTool.cpp
--- a/Engine/Plugins/Experimental/MeshModelingToolsetExp/Source/MeshModelingToolsExp/Private/Physics/SetCollisionGeometryTool.cpp
+++ b/Engine/Plugins/Experimental/MeshModelingToolsetExp/Source/MeshModelingToolsExp/Private/Physics/SetCollisionGeometryTool.cpp
@@ -216,6 +216,7 @@ void USetCollisionGeometryTool::Setup()
 	Settings->WatchProperty(Settings->MaxHullsPerMesh, [this](int32) { InvalidateCompute(); });
 	Settings->WatchProperty(Settings->ConvexDecompositionSearchFactor, [this](int32) { InvalidateCompute(); });
 	Settings->WatchProperty(Settings->AddHullsErrorTolerance, [this](int32) { InvalidateCompute(); });
+	Settings->WatchProperty(Settings->MinPartThickness, [this](int32) { InvalidateCompute(); });
 	Settings->WatchProperty(Settings->bSimplifyPolygons, [this](bool) { InvalidateCompute(); });
 	Settings->WatchProperty(Settings->HullTolerance, [this](float) { InvalidateCompute(); });
 	Settings->WatchProperty(Settings->SweepAxis, [this](EProjectedHullAxis) { InvalidateCompute(); });
@@ -292,6 +293,7 @@ TUniquePtr<UE::Geometry::TGenericDataOperator<FPhysicsDataCollection>> USetColli
 	Op->UseShapeGenerator->ConvexDecompositionMaxPieces = Settings->MaxHullsPerMesh;
 	Op->UseShapeGenerator->ConvexDecompositionSearchFactor = Settings->ConvexDecompositionSearchFactor;
 	Op->UseShapeGenerator->ConvexDecompositionErrorTolerance = Settings->AddHullsErrorTolerance;
+	Op->UseShapeGenerator->ConvexDecompositionMinPartThickness = Settings->MinPartThickness;
 	Op->UseShapeGenerator->HullSimplifyTolerance = Settings->HullTolerance;
 
 	Op->ComputeType = Settings->GeometryType;
diff --git a/Engine/Plugins/Experimental/MeshModelingToolsetExp/Source/MeshModelingToolsExp/Public/Physics/SetCollisionGeometryTool.h b/Engine/Plugins/Experimental/MeshModelingToolsetExp/Source/MeshModelingToolsExp/Public/Physics/SetCollisionGeometryTool.h
--- a/Engine/Plugins/Experimental/MeshModelingToolsetExp/Source/MeshModelingToolsExp/Public/Physics/SetCollisionGeometryTool.h
+++ b/Engine/Plugins/Experimental/MeshModelingToolsetExp/Source/MeshModelingToolsExp/Public/Physics/SetCollisionGeometryTool.h
@@ -133,6 +133,11 @@ public:
 		EditConditionHides, EditCondition = "GeometryType == ECollisionGeometryType::ConvexHulls && MaxHullsPerMesh > 1"))
 	float AddHullsErrorTolerance = 0;
 
+	/** Minimum part thickness for convex decomposition (in cm); hulls thinner than this will be merged into adjacent hulls, if possible. */
+	UPROPERTY(EditAnywhere, Category = ConvexHulls, meta = (UIMin = "0", UIMax = "1", ClampMin = "0",
+		EditConditionHides, EditCondition = "GeometryType == ECollisionGeometryType::ConvexHulls && MaxHullsPerMesh > 1"))
+	float MinPartThickness = 0.1;
+
 	UPROPERTY(EditAnywhere, Category = SweptHulls, meta = (EditConditionHides, EditCondition = "GeometryType == ECollisionGeometryType::SweptHulls"))
 	bool bSimplifyPolygons = true;
 
diff --git a/Engine/Plugins/Runtime/GeometryProcessing/Source/DynamicMesh/Private/ShapeApproximation/MeshSimpleShapeApproximation.cpp b/Engine/Plugins/Runtime/GeometryProcessing/Source/DynamicMesh/Private/ShapeApproximation/MeshSimpleShapeApproximation.cpp
--- a/Engine/Plugins/Runtime/GeometryProcessing/Source/DynamicMesh/Private/ShapeApproximation/MeshSimpleShapeApproximation.cpp
+++ b/Engine/Plugins/Runtime/GeometryProcessing/Source/DynamicMesh/Private/ShapeApproximation/MeshSimpleShapeApproximation.cpp
@@ -324,7 +324,7 @@ void FMeshSimpleShapeApproximation::Generate_ConvexHullDecompositions(FSimpleSha
 		// TODO: if (bSimplifyHulls), also consider simplifying the input?
 		FConvexDecomposition3 Decomposition(SourceMesh);
 
-		Decomposition.Compute(ConvexDecompositionMaxPieces, ConvexDecompositionMaxPieces * ConvexDecompositionSearchFactor, ConvexDecompositionErrorTolerance);
+		Decomposition.Compute(ConvexDecompositionMaxPieces, ConvexDecompositionMaxPieces * ConvexDecompositionSearchFactor, ConvexDecompositionErrorTolerance, ConvexDecompositionMinPartThickness);
 
 		for (int32 HullIdx = 0; HullIdx < Decomposition.NumHulls(); HullIdx++)
 		{
diff --git a/Engine/Plugins/Runtime/GeometryProcessing/Source/DynamicMesh/Public/ShapeApproximation/MeshSimpleShapeApproximation.h b/Engine/Plugins/Runtime/GeometryProcessing/Source/DynamicMesh/Public/ShapeApproximation/MeshSimpleShapeApproximation.h
--- a/Engine/Plugins/Runtime/GeometryProcessing/Source/DynamicMesh/Public/ShapeApproximation/MeshSimpleShapeApproximation.h
+++ b/Engine/Plugins/Runtime/GeometryProcessing/Source/DynamicMesh/Public/ShapeApproximation/MeshSimpleShapeApproximation.h
@@ -74,6 +74,8 @@ public:
 	float ConvexDecompositionSearchFactor = .5;
 	/** Error tolerance to guide convex decomposition (in cm); we stop adding new parts if the volume error is below the threshold.  For volumetric errors, value will be cubed. */
 	double ConvexDecompositionErrorTolerance = 0;
+	/** Minimum part thickness for convex decomposition (in cm); hulls thinner than this will be merged into adjacent hulls, if possible. */
+	double ConvexDecompositionMinPartThickness = .1;
 
 	bool bUseExactComputationForBox = false;
 
