Render Capture Baking: Reimplement infill in new baking framework using callbacks

#rnx

diff --git a/Engine/Plugins/Experimental/MeshModelingToolsetExp/Source/MeshModelingToolsExp/Private/BakeRenderCaptureTool.cpp b/Engine/Plugins/Experimental/MeshModelingToolsetExp/Source/MeshModelingToolsExp/Private/BakeRenderCaptureTool.cpp
--- a/Engine/Plugins/Experimental/MeshModelingToolsetExp/Source/MeshModelingToolsExp/Private/BakeRenderCaptureTool.cpp
+++ b/Engine/Plugins/Experimental/MeshModelingToolsetExp/Source/MeshModelingToolsExp/Private/BakeRenderCaptureTool.cpp
@@ -18,6 +18,7 @@
 #include "Sampling/MeshMapBaker.h"
 #include "Sampling/MeshGenericWorldPositionEvaluator.h"
 #include "Image/ImageInfilling.h"
+#include "Algo/NoneOf.h"
 
 
 using namespace UE::Geometry;
@@ -215,6 +216,112 @@ static void ImageBuildersFromPhotoSet(
 		return (HitTID == IndexConstants::InvalidID);
 	};
 
+	struct FInfillData
+	{
+		struct FSampleStats {
+			uint16 NumValid = 0;
+			uint16 NumInvalid = 0;
+
+			// The ==, !=, += operators and the Zero() function are required by the TMarchingPixelInfill implementation
+
+			bool operator==(const FSampleStats& Other) const
+			{
+				return (NumValid == Other.NumValid) && (NumInvalid == Other.NumInvalid);
+			}
+
+			bool operator!=(const FSampleStats& Other) const
+			{
+				return !(*this == Other);
+			}
+
+			FSampleStats& operator+=(const FSampleStats& Other)
+			{
+				NumValid += Other.NumValid;
+				NumInvalid += Other.NumInvalid;
+				return *this;
+			}
+
+			static FSampleStats Zero()
+			{
+				return FSampleStats{0, 0};
+			}
+		};
+
+		// Collect some sample stats per pixel, used to determine if a pixel requires infill or not
+		TImageBuilder<FSampleStats> SampleStats;
+
+		// The i-th element of this array indicates if the i-th evaluator needs infill
+		TArray<bool> EvaluatorNeedsInfill;
+	} InfillData;
+
+	InfillData.SampleStats.SetDimensions(FImageDimensions(Options.TextureImageSize, Options.TextureImageSize));
+	InfillData.SampleStats.Clear(FInfillData::FSampleStats{});
+
+	auto RegisterSampleStats = [&InfillData](bool bSampleValid, const FMeshMapEvaluator::FCorrespondenceSample& Sample, const FVector2d& UVPosition, const FVector2i& ImageCoords)
+	{
+		checkSlow(InfillData.SampleStats.GetDimensions().IsValidCoords(ImageCoords));
+		if (bSampleValid)
+		{
+			InfillData.SampleStats.GetPixel(ImageCoords).NumValid += 1;
+		}
+		else
+		{
+			InfillData.SampleStats.GetPixel(ImageCoords).NumInvalid += 1;
+		}
+	};
+
+	auto ComputeAndApplyInfill = [&InfillData](TArray<TUniquePtr<TImageBuilder<FVector4f>>>& BakeResults)
+	{
+		check(BakeResults.Num() == InfillData.EvaluatorNeedsInfill.Num());
+
+		if (BakeResults.IsEmpty() || Algo::NoneOf(InfillData.EvaluatorNeedsInfill))
+		{
+			return;
+		}
+
+		// Find pixels that need infill
+		TArray<FVector2i> MissingPixels;
+		FCriticalSection MissingPixelsLock;
+		ParallelFor(InfillData.SampleStats.GetDimensions().GetHeight(), [&MissingPixels, &MissingPixelsLock, &InfillData](int32 Y)
+		{
+			for (int32 X = 0; X < InfillData.SampleStats.GetDimensions().GetWidth(); X++)
+			{
+				const FInfillData::FSampleStats& Stats = InfillData.SampleStats.GetPixel(X, Y);
+				// TODO experiment with other classifications
+				if (Stats.NumInvalid > 0 && Stats.NumValid == 0)
+				{
+					MissingPixelsLock.Lock();
+					MissingPixels.Add(FVector2i(X, Y));
+					MissingPixelsLock.Unlock();
+				}
+			}
+		});
+
+		auto NormalizeFunc = [](FVector4f SumValue, int32 Count)
+		{
+			float InvSum = (Count == 0) ? 1.0f : (1.0f / Count);
+			return FVector4f(SumValue.X * InvSum, SumValue.Y * InvSum, SumValue.Z * InvSum, 1.0f);
+		};
+		
+		auto DummyNormalizeStatsFunc = [](FInfillData::FSampleStats SumValue, int32 Count)
+		{
+			// The return value must be different from MissingValue below so that ComputeInfill works correctly
+			return FInfillData::FSampleStats{TNumericLimits<uint16>::Max(), TNumericLimits<uint16>::Max()};
+		};
+
+		TMarchingPixelInfill<FInfillData::FSampleStats> Infill;
+		// This must be the same as the value of exterior pixels, otherwise infill will spread the exterior values into the texture
+		FInfillData::FSampleStats MissingValue{0, 0};
+		Infill.ComputeInfill(InfillData.SampleStats, MissingPixels, MissingValue, DummyNormalizeStatsFunc);
+		for (int EvaluatorIndex = 0; EvaluatorIndex < BakeResults.Num(); EvaluatorIndex++)
+		{
+			if (InfillData.EvaluatorNeedsInfill[EvaluatorIndex])
+			{
+				Infill.ApplyInfill<FVector4f>(*BakeResults[EvaluatorIndex], NormalizeFunc);
+			}
+		}
+	};
+
 	FMeshMapBaker Baker;
 	Baker.SetTargetMesh(BaseMesh);
 	Baker.SetTargetMeshTangents(BaseMeshTangents);
@@ -222,6 +329,8 @@ static void ImageBuildersFromPhotoSet(
 	Baker.SetSamplesPerPixel(static_cast<int32>(Options.SamplesPerPixel));
 	Baker.SetFilter(FMeshMapBaker::EBakeFilterType::BSpline);
 	Baker.SetTargetMeshUVLayer(Options.TargetUVLayer);
+	Baker.InteriorSampleCallback = RegisterSampleStats;
+	Baker.PostWriteToImageCallback = ComputeAndApplyInfill;
 
 	FSceneCapturePhotoSetSampler Sampler(SceneCapture, VisibilityFunction, BaseMesh, &BaseMeshSpatial, BaseMeshTangents.Get());
 	Baker.SetDetailSampler(&Sampler);
@@ -256,6 +365,7 @@ static void ImageBuildersFromPhotoSet(
 		
 		int32 EvaluatorIndex = Baker.AddEvaluator(Evaluator);
 		
+		InfillData.EvaluatorNeedsInfill.Add(true);
 		CaptureTypeToEvaluatorIndexMap.Emplace(CaptureType, EvaluatorIndex);
 	};
 
@@ -306,6 +416,10 @@ static void ImageBuildersFromPhotoSet(
 		int32 EvaluatorIndex = Baker.AddEvaluator(Evaluator);
 		
 		CaptureTypeToEvaluatorIndexMap.Emplace(ERenderCaptureType::WorldNormal, EvaluatorIndex);
+
+		// Note: No infill on normal map for now, doesn't make sense to do after mapping to tangent space!
+		//  (should we build baked normal map in world space, and then resample to tangent space??)
+		InfillData.EvaluatorNeedsInfill.Add(false);
 	}
 
 	{
@@ -344,85 +458,6 @@ static void ImageBuildersFromPhotoSet(
 			Results->NormalImage = MoveTemp(Baker.GetBakeResults(Item.Value)[0]);
 		}
 	}
-
-	// TODO Remove the FMeshImageBakingCache and reimplement infill to work with the new baking framework
-
-	FMeshImageBakingCache TempBakeCache;
-	{
-		TRACE_CPUPROFILER_EVENT_SCOPE(ImageBuildersFromPhotoSet_MakeMeshImageBakingCache);
-
-		TempBakeCache.SetDetailMesh(BaseMesh, &BaseMeshSpatial);
-		TempBakeCache.SetBakeTargetMesh(BaseMesh);
-		TempBakeCache.SetDimensions(FImageDimensions(Options.TextureImageSize, Options.TextureImageSize));
-		TempBakeCache.SetUVLayer(Options.TargetUVLayer);
-		TempBakeCache.SetThickness(0.1);
-		TempBakeCache.SetCorrespondenceStrategy(FMeshImageBakingCache::ECorrespondenceStrategy::Identity);
-		TempBakeCache.ValidateCache();
-	}
-
-	TMarchingPixelInfill<FVector4f> Infill;
-
-	{
-		TRACE_CPUPROFILER_EVENT_SCOPE(ImageBuildersFromPhotoSet_FindAndComputeInfill);
-		
-		// Find missing pixels
-		TArray<FVector2i> MissingPixels;
-		TempBakeCache.FindSamplingHoles([&](const FVector2i& Coords)
-		{
-			return Results->ColorImage->GetPixel(Coords) == InvalidColor;
-		}, MissingPixels);
-
-		// Solve infill for the holes while also caching infill information
-		Infill.ComputeInfill(*Results->ColorImage, MissingPixels, InvalidColor,
-			[](FVector4f SumValue, int32 Count) {
-			float InvSum = (Count == 0) ? 1.0f : (1.0f / Count);
-			return FVector4f(SumValue.X * InvSum, SumValue.Y * InvSum, SumValue.Z * InvSum, 1.0f);
-		});
-	}
-
-	auto ApplyInfill = [&Infill](TUniquePtr<TImageBuilder<FVector4f>>& Image)
-	{
-		Infill.ApplyInfill(*Image,
-			[](FVector4f SumValue, int32 Count)
-			{
-				float InvSum = (Count == 0) ? 1.0f : (1.0f / Count);
-				return FVector4f(SumValue.X * InvSum, SumValue.Y * InvSum, SumValue.Z * InvSum, 1.0f);
-			});
-	};
-
-	{
-		TRACE_CPUPROFILER_EVENT_SCOPE(ImageBuildersFromPhotoSet_ApplyInfill);
-
-		// Note: Infill for Results->ColorImage was applied when the infill was computed
-
-		if (Options.bUsePackedMRS)
-		{
-			ApplyInfill(Results->PackedMRSImage);
-		}
-		else
-		{
-			if (Options.bBakeRoughness)
-			{
-				ApplyInfill(Results->RoughnessImage);
-			}
-			if (Options.bBakeMetallic)
-			{
-				ApplyInfill(Results->MetallicImage);
-			}
-			if (Options.bBakeSpecular)
-			{
-				ApplyInfill(Results->SpecularImage);
-			}
-		}
-
-		if (Options.bBakeEmissive)
-		{
-			ApplyInfill(Results->EmissiveImage);
-		}
-
-		// Note: No infill on normal map for now, doesn't make sense to do after mapping to tangent space!
-		//  (should we build baked normal map in world space, and then resample to tangent space??)
-	}
 }
 
 //
diff --git a/Engine/Plugins/Runtime/GeometryProcessing/Source/DynamicMesh/Private/Sampling/MeshMapBaker.cpp b/Engine/Plugins/Runtime/GeometryProcessing/Source/DynamicMesh/Private/Sampling/MeshMapBaker.cpp
--- a/Engine/Plugins/Runtime/GeometryProcessing/Source/DynamicMesh/Private/Sampling/MeshMapBaker.cpp
+++ b/Engine/Plugins/Runtime/GeometryProcessing/Source/DynamicMesh/Private/Sampling/MeshMapBaker.cpp
@@ -376,22 +376,20 @@ void FMeshMapBaker::Bake()
 							const FVector2d UVPosition = (FVector2d)OccupancyMap.TexelQueryUV[LinearIdx];
 							const int32 UVTriangleID = OccupancyMap.TexelQueryTriangle[LinearIdx];
 
-							// Compute the per-sample correspondence data
-							bool bValidSample = false;
+							// Compute the per-sample correspondence data 
+							// Note: Since we check LinearIdx is an interior sample above we know we'll get a valid
+							// SampleInfo because interior samples all have valid UVTriangleIDs. We don't check the
+							// return value of QuerySampleInfo though because this is a tight inner loop
 							FMeshUVSampleInfo SampleInfo;
-							if (MeshUVSampler.QuerySampleInfo(UVTriangleID, UVPosition, SampleInfo))
+							MeshUVSampler.QuerySampleInfo(UVTriangleID, UVPosition, SampleInfo);
 							{
 								FMeshMapEvaluator::FCorrespondenceSample Sample;
-								if (ComputeCorrespondenceSample(SampleInfo, Sample))
+								bool bSampleValid = ComputeCorrespondenceSample(SampleInfo, Sample);
+								if (bSampleValid)
 								{
 									BakeSample(*TileBuffer, Sample, UVPosition, ImageCoords, OccupancyMap);
-									bValidSample = true;
 								}
-							}
-
-							if (bValidSample == false)
-							{
-								// RecordInvalidSample(Sample, UVPosition, ImageCoords); // TODO
+								InteriorSampleCallback(bSampleValid, Sample, UVPosition, ImageCoords);
 							}
 						}
 					}
@@ -512,6 +510,13 @@ void FMeshMapBaker::Bake()
 		return;
 	}
 
+	PostWriteToImageCallback(BakeResults);
+
+	if (CancelF())
+	{
+		return;
+	}
+
 	// Gutter Texel processing
 	if (bGutterEnabled)
 	{
diff --git a/Engine/Plugins/Runtime/GeometryProcessing/Source/DynamicMesh/Public/Sampling/MeshMapBaker.h b/Engine/Plugins/Runtime/GeometryProcessing/Source/DynamicMesh/Public/Sampling/MeshMapBaker.h
--- a/Engine/Plugins/Runtime/GeometryProcessing/Source/DynamicMesh/Public/Sampling/MeshMapBaker.h
+++ b/Engine/Plugins/Runtime/GeometryProcessing/Source/DynamicMesh/Public/Sampling/MeshMapBaker.h
@@ -50,6 +50,14 @@ public:
 	/** @return true if we should abort calculation */
 	TFunction<bool(void)> CancelF = []() { return false; };
 
+	/** Function to call after evaluator data is written to the final image, but before any gutter processing occurs */
+	TFunction<void(TArray<TUniquePtr<TImageBuilder<FVector4f>>>&)> PostWriteToImageCallback
+		= [](TArray<TUniquePtr<TImageBuilder<FVector4f>>>& PostWriteToImageBakeResults) {};
+
+	/* Function to call for each interior sample */
+	TFunction<void(bool, const FMeshMapEvaluator::FCorrespondenceSample&, const FVector2d&, const FVector2i&)> InteriorSampleCallback
+		= [](bool bSampleValid, const FMeshMapEvaluator::FCorrespondenceSample& Sample, const FVector2d& UVPosition, const FVector2i& ImageCoords) {};
+
 	/**
 	 * @param ImageCoords the output image coordinates to be evaluated
 	 * @param UV the target mesh UV coordinates to be evaluated
diff --git a/Engine/Source/Runtime/GeometryCore/Public/Image/ImageInfilling.h b/Engine/Source/Runtime/GeometryCore/Public/Image/ImageInfilling.h
--- a/Engine/Source/Runtime/GeometryCore/Public/Image/ImageInfilling.h
+++ b/Engine/Source/Runtime/GeometryCore/Public/Image/ImageInfilling.h
@@ -91,7 +91,7 @@ public:
 				FVector2i NbrCoords = CenterCoords + IndexUtil::GridOffsets8[j];
 				if (Image.ContainsPixel(NbrCoords))
 				{
-					FVector4f PixelColor = Image.GetPixel(NbrCoords);
+					PixelType PixelColor = Image.GetPixel(NbrCoords);
 					if (PixelColor != MissingValue)
 					{
 						ValidNbrAvg += PixelColor;
@@ -158,7 +158,7 @@ public:
 					NbrSum += Image.GetPixel(NbrCoords);
 					NbrCount++;
 				}
-				PixelType InfillValue = NormalizeFunc(NbrSum, NbrCount);
+				OtherPixelType InfillValue = NormalizeFunc(NbrSum, NbrCount);
 				Image.SetPixel(CenterCoords, InfillValue);
 			}
 		}
diff --git a/Engine/Source/Runtime/GeometryCore/Public/Image/SpatialPhotoSet.h b/Engine/Source/Runtime/GeometryCore/Public/Image/SpatialPhotoSet.h
--- a/Engine/Source/Runtime/GeometryCore/Public/Image/SpatialPhotoSet.h
+++ b/Engine/Source/Runtime/GeometryCore/Public/Image/SpatialPhotoSet.h
@@ -160,7 +160,9 @@ PixelType TSpatialPhotoSet<PixelType, RealType>::ComputeSample(
 }
 
 
-// TODO Consider refactoring this to deduplicate code from ComputeSample, or not, since this is called in a tight inner loop
+// TODO Consider refactoring this to deduplicate code from ComputeSample, or not, since this is called in a tight inner
+// loop. Note: Unlike ComputeSample, which searches for the best valid photo, this function immediately returns from the
+// innermost if since we just want to know if at least one photo has passed the visibility check
 template<typename PixelType, typename RealType>
 bool TSpatialPhotoSet<PixelType, RealType>::IsValidSample(
 	const FVector3d& Position, 
