fix polycut handling of non-uniform scale

diff --git a/Engine/Plugins/Experimental/MeshModelingToolsetExp/Source/MeshModelingToolsEditorOnlyExp/Private/PolygonOnMeshTool.cpp b/Engine/Plugins/Experimental/MeshModelingToolsetExp/Source/MeshModelingToolsEditorOnlyExp/Private/PolygonOnMeshTool.cpp
--- a/Engine/Plugins/Experimental/MeshModelingToolsetExp/Source/MeshModelingToolsEditorOnlyExp/Private/PolygonOnMeshTool.cpp
+++ b/Engine/Plugins/Experimental/MeshModelingToolsetExp/Source/MeshModelingToolsEditorOnlyExp/Private/PolygonOnMeshTool.cpp
@@ -260,15 +260,23 @@ TUniquePtr<FDynamicMeshOperator> UPolygonOnMeshTool::MakeNewOperator()
 		BasicProperties->Operation == EEmbeddedPolygonOpMethod::TrimOutside;
 	EmbedOp->bAttemptFixHolesOnBoolean = !bOpLeavesOpenBoundaries && BasicProperties->bTryToFixCracks;
 
-	FFrame3d LocalFrame = DrawPlaneWorld;
-	FTransform3d ToLocal = WorldTransform.Inverse();
-	LocalFrame.Transform(ToLocal);
+	// Match the world plane in the local space
+	FVector3d LocalOrigin = WorldTransform.InverseTransformPosition(DrawPlaneWorld.Origin);
+	FVector3d TempLocalX = WorldTransform.InverseTransformVector(DrawPlaneWorld.GetAxis(0));
+	FVector3d TempLocalY = WorldTransform.InverseTransformVector(DrawPlaneWorld.GetAxis(1));
+	FVector3d LocalZ = TempLocalX.Cross(TempLocalY);
+	FFrame3d LocalFrame(LocalOrigin, LocalZ);
 	EmbedOp->PolygonFrame = LocalFrame;
 	
-	FVector2d LocalFrameScale(ToLocal.TransformVector(LocalFrame.X()).Length(), ToLocal.TransformVector(LocalFrame.Y()).Length());
-	LocalFrameScale *= BasicProperties->PolygonScale;
+	// Transform the active polygon by the polygon scale put it on the local space's frame
 	EmbedOp->EmbedPolygon = ActivePolygon;
-	EmbedOp->EmbedPolygon.Scale(LocalFrameScale, FVector2d::Zero());
+	const TArray<FVector2d>& Vertices = EmbedOp->EmbedPolygon.GetVertices();
+	for (int32 Idx = 0; Idx < EmbedOp->EmbedPolygon.VertexCount(); ++Idx)
+	{
+		FVector2d World2d = Vertices[Idx] * BasicProperties->PolygonScale;
+		FVector2d LocalPos = LocalFrame.ToPlaneUV(WorldTransform.InverseTransformPosition(DrawPlaneWorld.FromPlaneUV(World2d)));
+		EmbedOp->EmbedPolygon.Set(Idx, LocalPos);
+	}
 
 	// TODO: scale any extrude by ToLocal.TransformVector(LocalFrame.Z()).Length() ??
 	// EmbedOp->ExtrudeDistance = Tool->BasicProperties->ExtrudeDistance;
