Render Capture Baking: Move visibility raycast from evaluators to detail sampler

Benchmark configuration:
- Test scene with 4 meshes: a mossy stone nanite mesh, a beach ball, a ball with an emissive blue cross and a hut mesh
- Target mesh: 11.8k triangles mesh, created by voxel wrapping the meshes in the test scene elements
- 4096x4096 output texture with 16x multisampling baking BaseColor, Emissive, PackedMRS and Normal maps
- 4096x4096 photo capture resolution

Unreal Insights reports the following times for the Bake call:
- Before this change: 2 mins 45 seconds
- After this change: 55.9 seconds

Aside: (Unaffected by this change) In both cases the photo capture, which is done on the game thread, took about 1 min

#rnx

diff --git a/Engine/Plugins/Experimental/MeshModelingToolsetExp/Source/MeshModelingToolsExp/Private/BakeRenderCaptureTool.cpp b/Engine/Plugins/Experimental/MeshModelingToolsetExp/Source/MeshModelingToolsExp/Private/BakeRenderCaptureTool.cpp
--- a/Engine/Plugins/Experimental/MeshModelingToolsetExp/Source/MeshModelingToolsExp/Private/BakeRenderCaptureTool.cpp
+++ b/Engine/Plugins/Experimental/MeshModelingToolsetExp/Source/MeshModelingToolsExp/Private/BakeRenderCaptureTool.cpp
@@ -16,7 +16,7 @@
 
 #include "Sampling/MeshImageBakingCache.h"
 #include "Sampling/MeshMapBaker.h"
-#include "Sampling/MeshGenericWorldPositionEvaluator.h"
+#include "Sampling/GenericEvaluator.h"
 #include "Image/ImageInfilling.h"
 #include "Algo/NoneOf.h"
 
@@ -31,12 +31,11 @@ using namespace UE::Geometry;
 
 
 
-// TODO Reimplement this when the render capture algorithm is refactored to decouple the visibility checking from the evaluation
 class FSceneCapturePhotoSetSampler : public FMeshBakerDynamicMeshSampler
 {
 public:
 	FSceneCapturePhotoSetSampler(
-		const FSceneCapturePhotoSet* SceneCapture,
+		FSceneCapturePhotoSet* SceneCapture,
 		TFunctionRef<bool(const FVector3d&, const FVector3d&)> VisibilityFunction,
 		const FDynamicMesh3* Mesh,
 		const FDynamicMeshAABBTree3* Spatial,
@@ -47,13 +46,36 @@ public:
 	{
 	}
 
+	virtual bool SupportsCustomCorrespondence() const override
+	{
+		return true;
+	}
+
+	// Warning: Expects that Sample.BaseSample.SurfacePoint and Sample.BaseNormal are set when the function is called
+	virtual void* ComputeCustomCorrespondence(const FMeshUVSampleInfo& SampleInfo, FMeshMapEvaluator::FCorrespondenceSample& Sample) const override
+	{
+		// Perform a ray-cast to determine which photo/coordinate, if any, should be sampled
+		int PhotoIndex;
+		FVector2d PhotoCoords;
+		SceneCapture->ComputeSampleLocation(Sample.BaseSample.SurfacePoint, Sample.BaseNormal, VisibilityFunction, PhotoIndex, PhotoCoords);
+
+		// Store the photo coordinates and index in the correspondence sample
+		Sample.DetailMesh = SceneCapture;
+		Sample.DetailTriID = PhotoIndex;
+		Sample.DetailBaryCoords.X = PhotoCoords.X;
+		Sample.DetailBaryCoords.Y = PhotoCoords.Y;
+
+		// This will be set to Sample.DetailMesh but we can already do that internally so it's kindof redundant
+		return SceneCapture;
+	}
+
 	virtual bool IsValidCorrespondence(const FMeshMapEvaluator::FCorrespondenceSample& Sample) const override
 	{
-		return SceneCapture->IsValidSample(Sample.BaseSample.SurfacePoint, Sample.BaseNormal, VisibilityFunction);
+		return Sample.DetailTriID != IndexConstants::InvalidID;
 	}
 
 public:
-	const FSceneCapturePhotoSet* SceneCapture = nullptr;
+	FSceneCapturePhotoSet* SceneCapture = nullptr;
 	TFunctionRef<bool(const FVector3d&, const FVector3d&)> VisibilityFunction;
 };
 
@@ -195,6 +217,32 @@ static TUniquePtr<FSceneCapturePhotoSet> CapturePhotoSet(
 	return SceneCapture;
 }
 
+template <ERenderCaptureType CaptureType>
+TSharedPtr<FGenericEvaluator<FVector4f>>
+MakeColorEvaluator(const FSceneCapturePhotoSet::FSceneSample& DefaultSample, const FSceneCapturePhotoSet* SceneCapture)
+{
+	TSharedPtr<FGenericEvaluator<FVector4f>> Evaluator = MakeShared<FGenericEvaluator<FVector4f>>();
+
+	Evaluator->DefaultResult = DefaultSample.GetValue4f(CaptureType);
+
+	Evaluator->EvaluateSampleCallback = [&DefaultSample, SceneCapture](const FMeshMapEvaluator::FCorrespondenceSample& Sample)
+	{
+		const int PhotoIndex = Sample.DetailTriID;
+		const FVector2d PhotoCoords(Sample.DetailBaryCoords.X, Sample.DetailBaryCoords.Y);
+		const FVector4f SampleColor = SceneCapture->ComputeSample<CaptureType>(PhotoIndex, PhotoCoords, DefaultSample);
+		return SampleColor;
+	};
+
+	Evaluator->EvaluateColorCallback = [](const int DataIdx, float*& In)
+	{
+		const FVector4f Out(In[0], In[1], In[2], In[3]);
+		In += 4;
+		return Out;
+	};
+
+	return Evaluator;
+}
+
 static void ImageBuildersFromPhotoSet(
 	FSceneCapturePhotoSet* SceneCapture,
 	const FBakeRenderCaptureOptions::FOptions& Options,
@@ -334,7 +382,7 @@ static void ImageBuildersFromPhotoSet(
 
 	FSceneCapturePhotoSetSampler Sampler(SceneCapture, VisibilityFunction, BaseMesh, &BaseMeshSpatial, BaseMeshTangents.Get());
 	Baker.SetDetailSampler(&Sampler);
-	Baker.SetCorrespondenceStrategy(FMeshMapBaker::ECorrespondenceStrategy::Identity);
+	Baker.SetCorrespondenceStrategy(FMeshMapBaker::ECorrespondenceStrategy::Custom);
 
 	TMap<ERenderCaptureType, int32> CaptureTypeToEvaluatorIndexMap;
 
@@ -351,70 +399,100 @@ static void ImageBuildersFromPhotoSet(
 	DefaultColorSample.Emissive = FVector3f(InvalidColor.X, InvalidColor.Y, InvalidColor.Z);
 	DefaultColorSample.WorldNormal = FVector4f((DefaultNormal + FVector3f::One()) * .5f, InvalidColor.W);
 
-	auto AddColorEvaluator = [&](ERenderCaptureType CaptureType)
+	auto AddColorEvaluator = [&Baker, &InfillData, &CaptureTypeToEvaluatorIndexMap] (
+		const TSharedPtr<FGenericEvaluator<FVector4f>>& Evaluator,
+		ERenderCaptureType CaptureType)
 	{
-		TSharedPtr<FMeshGenericWorldPositionColorEvaluator> Evaluator = MakeShared<FMeshGenericWorldPositionColorEvaluator>();
-		Evaluator->DefaultColor = DefaultColorSample.GetValue4f(CaptureType);
-		Evaluator->ColorSampleFunction =
-			[&DefaultColorSample, &VisibilityFunction, SceneCapture, CaptureType](FVector3d Position, FVector3d Normal)
-		{
-			FSceneCapturePhotoSet::FSceneSample ColorSample = DefaultColorSample;
-			SceneCapture->ComputeSample(FRenderCaptureTypeFlags::Single(CaptureType), Position, Normal, VisibilityFunction, ColorSample);
-			return ColorSample.GetValue4f(CaptureType);
-		};
-		
 		int32 EvaluatorIndex = Baker.AddEvaluator(Evaluator);
-		
 		InfillData.EvaluatorNeedsInfill.Add(true);
 		CaptureTypeToEvaluatorIndexMap.Emplace(CaptureType, EvaluatorIndex);
 	};
 
-	AddColorEvaluator(ERenderCaptureType::BaseColor);
+	AddColorEvaluator(
+		MakeColorEvaluator<ERenderCaptureType::BaseColor>(DefaultColorSample, SceneCapture),
+		ERenderCaptureType::BaseColor);
 
 	if (Options.bUsePackedMRS)
 	{
-		AddColorEvaluator(ERenderCaptureType::CombinedMRS);
+		AddColorEvaluator(
+			MakeColorEvaluator<ERenderCaptureType::CombinedMRS>(DefaultColorSample, SceneCapture),
+			ERenderCaptureType::CombinedMRS);
 	}
 	else
 	{
 		if (Options.bBakeRoughness)
 		{
-			AddColorEvaluator(ERenderCaptureType::Roughness);
+			AddColorEvaluator(
+				MakeColorEvaluator<ERenderCaptureType::Roughness>(DefaultColorSample, SceneCapture),
+				ERenderCaptureType::Roughness);
 		}
 		if (Options.bBakeMetallic)
 		{
-			AddColorEvaluator(ERenderCaptureType::Metallic);
+			AddColorEvaluator(
+				MakeColorEvaluator<ERenderCaptureType::Metallic>(DefaultColorSample, SceneCapture),
+				ERenderCaptureType::Metallic);
 		}
 		if (Options.bBakeSpecular)
 		{
-			AddColorEvaluator(ERenderCaptureType::Specular);
+			AddColorEvaluator(
+				MakeColorEvaluator<ERenderCaptureType::Specular>(DefaultColorSample, SceneCapture),
+				ERenderCaptureType::Specular);
 		}
 	}
 
 	if (Options.bBakeEmissive)
 	{
-		AddColorEvaluator(ERenderCaptureType::Emissive);
+		AddColorEvaluator(
+			MakeColorEvaluator<ERenderCaptureType::Emissive>(DefaultColorSample, SceneCapture),
+			ERenderCaptureType::Emissive);
 	}
 
 	if (Options.bBakeNormalMap) {
-		TSharedPtr<FMeshGenericWorldPositionNormalEvaluator> Evaluator = MakeShared<FMeshGenericWorldPositionNormalEvaluator>();
-		Evaluator->DefaultUnitWorldNormal = DefaultNormal;
-		Evaluator->UnitWorldNormalSampleFunction =
-			[&DefaultColorSample, &VisibilityFunction, SceneCapture](FVector3d Position, FVector3d Normal)
+		TSharedPtr<FGenericEvaluator<FVector3f, FMeshMapEvaluator::EComponents::Float3>> Evaluator =
+			MakeShared<FGenericEvaluator<FVector3f, FMeshMapEvaluator::EComponents::Float3>>();
+
+		Evaluator->DefaultResult = DefaultNormal;
+
+		Evaluator->EvaluateSampleCallback = [&DefaultColorSample, &BaseMeshTangents, SceneCapture](const FMeshMapEvaluator::FCorrespondenceSample& Sample)
 		{
-			FSceneCapturePhotoSet::FSceneSample ColorSample = DefaultColorSample;
-			SceneCapture->ComputeSample(FRenderCaptureTypeFlags::WorldNormal(), Position, Normal, VisibilityFunction, ColorSample);
-			FVector3f NormalColor = ColorSample.WorldNormal;
+			const int32 TriangleID = Sample.BaseSample.TriangleIndex;
+			const FVector3d BaryCoords = Sample.BaseSample.BaryCoords;
+			const int PhotoIndex = Sample.DetailTriID;
+			const FVector2d PhotoCoords(Sample.DetailBaryCoords.X, Sample.DetailBaryCoords.Y);
+
+			const FVector4f NormalColor = SceneCapture->ComputeSample<ERenderCaptureType::WorldNormal>(PhotoIndex, PhotoCoords, DefaultColorSample);
 
 			// Map from color components [0,1] to normal components [-1,1]
-			float x = (NormalColor.X - 0.5f) * 2.0f;
-			float y = (NormalColor.Y - 0.5f) * 2.0f;
-			float z = (NormalColor.Z - 0.5f) * 2.0f;
-			return FVector3f(x, y, z);
+			const FVector3f WorldSpaceNormal(
+				(NormalColor.X - 0.5f) * 2.0f,
+				(NormalColor.Y - 0.5f) * 2.0f,
+				(NormalColor.Z - 0.5f) * 2.0f);
+
+			// Get tangents on base mesh
+			FVector3d BaseTangentX, BaseTangentY;
+			BaseMeshTangents->GetInterpolatedTriangleTangent(TriangleID, BaryCoords, BaseTangentX, BaseTangentY);
+
+			// Compute normal in tangent space
+			const FVector3f TangentSpaceNormal(
+					(float)WorldSpaceNormal.Dot(FVector3f(BaseTangentX)),
+					(float)WorldSpaceNormal.Dot(FVector3f(BaseTangentY)),
+					(float)WorldSpaceNormal.Dot(FVector3f(Sample.BaseNormal)));
+
+			return TangentSpaceNormal;
 		};
-		
+
+		Evaluator->EvaluateColorCallback = [](const int DataIdx, float*& In)
+		{
+			// Map normal space [-1,1] to color space [0,1]
+			const FVector3f Normal(In[0], In[1], In[2]);
+			const FVector3f Color = (Normal + FVector3f::One()) * 0.5f;
+			const FVector4f Out(Color.X, Color.Y, Color.Z, 1.0f);
+			In += 3;
+			return Out;
+		};
+
 		int32 EvaluatorIndex = Baker.AddEvaluator(Evaluator);
-		
+
 		CaptureTypeToEvaluatorIndexMap.Emplace(ERenderCaptureType::WorldNormal, EvaluatorIndex);
 
 		// Note: No infill on normal map for now, doesn't make sense to do after mapping to tangent space!
diff --git a/Engine/Plugins/Runtime/GeometryProcessing/Source/DynamicMesh/Private/Sampling/MeshGenericWorldPositionEvaluator.cpp b/Engine/Plugins/Runtime/GeometryProcessing/Source/DynamicMesh/Private/Sampling/MeshGenericWorldPositionEvaluator.cpp
deleted file mode 100644
diff --git a/Engine/Plugins/Runtime/GeometryProcessing/Source/DynamicMesh/Private/Sampling/MeshMapBaker.cpp b/Engine/Plugins/Runtime/GeometryProcessing/Source/DynamicMesh/Private/Sampling/MeshMapBaker.cpp
--- a/Engine/Plugins/Runtime/GeometryProcessing/Source/DynamicMesh/Private/Sampling/MeshMapBaker.cpp
+++ b/Engine/Plugins/Runtime/GeometryProcessing/Source/DynamicMesh/Private/Sampling/MeshMapBaker.cpp
@@ -203,7 +203,6 @@ void FMeshMapBaker::Bake()
 		else if (UseStrategy == ECorrespondenceStrategy::Custom && DetailSampler->SupportsCustomCorrespondence())
 		{
 			ValueOut.DetailMesh = DetailSampler->ComputeCustomCorrespondence(SampleInfo, ValueOut);
-			check(false); // @Incomplete
 		}
 		else	// Fall back to raycast strategy
 		{
@@ -218,12 +217,7 @@ void FMeshMapBaker::Bake()
 				(UseStrategy == ECorrespondenceStrategy::RaycastStandardThenNearest));
 		}
 
-		if (ValueOut.DetailMesh && DetailSampler->IsTriangle(ValueOut.DetailMesh, ValueOut.DetailTriID))
-		{
-			return DetailSampler->IsValidCorrespondence(ValueOut);
-		}
-
-		return false;
+		return DetailSampler->IsValidCorrespondence(ValueOut);
 	};
 
 	// This computes a FMeshUVSampleInfo to pass to the ComputeCorrespondenceSample function, which will find the
@@ -378,10 +372,9 @@ void FMeshMapBaker::Bake()
 
 							// Compute the per-sample correspondence data 
 							// Note: Since we check LinearIdx is an interior sample above we know we'll get a valid
-							// SampleInfo because interior samples all have valid UVTriangleIDs. We don't check the
-							// return value of QuerySampleInfo though because this is a tight inner loop
+							// SampleInfo because interior samples all have valid UVTriangleIDs.
 							FMeshUVSampleInfo SampleInfo;
-							MeshUVSampler.QuerySampleInfo(UVTriangleID, UVPosition, SampleInfo);
+							if (MeshUVSampler.QuerySampleInfo(UVTriangleID, UVPosition, SampleInfo))
 							{
 								FMeshMapEvaluator::FCorrespondenceSample Sample;
 								bool bSampleValid = ComputeCorrespondenceSample(SampleInfo, Sample);
diff --git a/Engine/Plugins/Runtime/GeometryProcessing/Source/DynamicMesh/Public/Sampling/GenericEvaluator.h b/Engine/Plugins/Runtime/GeometryProcessing/Source/DynamicMesh/Public/Sampling/GenericEvaluator.h
new file mode 100644
--- /dev/null
+++ b/Engine/Plugins/Runtime/GeometryProcessing/Source/DynamicMesh/Public/Sampling/GenericEvaluator.h
@@ -0,0 +1,63 @@
+﻿// Copyright Epic Games, Inc. All Rights Reserved.
+
+#pragma once
+
+#include "Sampling/MeshMapEvaluator.h"
+
+namespace UE
+{
+namespace Geometry
+{
+
+template <
+	typename DataType,
+	FMeshMapEvaluator::EComponents NumComponents = FMeshMapEvaluator::EComponents::Float4,
+	FMeshMapEvaluator::EAccumulateMode Mode = FMeshMapEvaluator::EAccumulateMode::Add>
+class FGenericEvaluator : public FMeshMapEvaluator
+{
+public:
+
+	DataType DefaultResult;
+	TFunction<DataType(const FCorrespondenceSample& Sample)> EvaluateSampleCallback;
+	TFunction<FVector4f(const int DataIdx, float*& In)> EvaluateColorCallback;
+
+public:
+	virtual ~FGenericEvaluator() {}
+
+	// Begin FMeshMapEvaluator interface
+	virtual void Setup(const FMeshBaseBaker& Baker, FEvaluationContext& Context) override
+	{
+		Context.Evaluate = &EvaluateSample;
+		Context.EvaluateDefault = &EvaluateDefault;
+		Context.EvaluateColor = &EvaluateColor;
+		Context.EvalData = this;
+		Context.AccumulateMode = Mode;
+		Context.DataLayout = { NumComponents };
+	}
+
+	virtual EMeshMapEvaluatorType Type() const override { return EMeshMapEvaluatorType::Generic; }
+	// End FMeshMapEvaluator interface
+
+	static void EvaluateSample(float*& Out, const FCorrespondenceSample& Sample, void* EvalData)
+	{
+		const FGenericEvaluator* Eval = static_cast<FGenericEvaluator*>(EvalData);
+		const DataType SampleResult = Eval->EvaluateSampleCallback(Sample);
+		WriteToBuffer(Out, SampleResult);
+	}
+
+	static void EvaluateDefault(float*& Out, void* EvalData)
+	{
+		const FGenericEvaluator* Eval = static_cast<FGenericEvaluator*>(EvalData);
+		WriteToBuffer(Out, Eval->DefaultResult);
+	}
+
+	static void EvaluateColor(const int DataIdx, float*& In, FVector4f& Out, void* EvalData)
+	{
+		const FGenericEvaluator* Eval = static_cast<FGenericEvaluator*>(EvalData);
+		Out = Eval->EvaluateColorCallback(DataIdx, In);
+	}
+};
+
+
+} // end namespace UE::Geometry
+} // end namespace UE
diff --git a/Engine/Plugins/Runtime/GeometryProcessing/Source/DynamicMesh/Public/Sampling/MeshBakerCommon.h b/Engine/Plugins/Runtime/GeometryProcessing/Source/DynamicMesh/Public/Sampling/MeshBakerCommon.h
--- a/Engine/Plugins/Runtime/GeometryProcessing/Source/DynamicMesh/Public/Sampling/MeshBakerCommon.h
+++ b/Engine/Plugins/Runtime/GeometryProcessing/Source/DynamicMesh/Public/Sampling/MeshBakerCommon.h
@@ -159,9 +159,9 @@ public:
 		return nullptr;
 	}
 
-	virtual bool IsValidCorrespondence(const FMeshMapEvaluator::FCorrespondenceSample& CorrespondenceSample) const
+	virtual bool IsValidCorrespondence(const FMeshMapEvaluator::FCorrespondenceSample& Sample) const
 	{
-		return true;
+		return Sample.DetailMesh && IsTriangle(Sample.DetailMesh, Sample.DetailTriID);
 	}
 
 	/**
diff --git a/Engine/Plugins/Runtime/GeometryProcessing/Source/DynamicMesh/Public/Sampling/MeshGenericWorldPositionEvaluator.h b/Engine/Plugins/Runtime/GeometryProcessing/Source/DynamicMesh/Public/Sampling/MeshGenericWorldPositionEvaluator.h
deleted file mode 100644
diff --git a/Engine/Plugins/Runtime/GeometryProcessing/Source/DynamicMesh/Public/Sampling/MeshMapEvaluator.h b/Engine/Plugins/Runtime/GeometryProcessing/Source/DynamicMesh/Public/Sampling/MeshMapEvaluator.h
--- a/Engine/Plugins/Runtime/GeometryProcessing/Source/DynamicMesh/Public/Sampling/MeshMapEvaluator.h
+++ b/Engine/Plugins/Runtime/GeometryProcessing/Source/DynamicMesh/Public/Sampling/MeshMapEvaluator.h
@@ -20,8 +20,7 @@ enum class EMeshMapEvaluatorType
 	Property,
 	ResampleImage,
 	MultiResampleImage,
-	GenericWorldPositionColor,
-	GenericWorldPositionNormal
+	Generic
 };
 
 class FMeshMapEvaluator
@@ -53,7 +52,8 @@ public:
 		FMeshUVSampleInfo BaseSample;
 		FVector3d BaseNormal = FVector3d::Zero();
 
-		const void* DetailMesh = nullptr;
+		// The following data has an interpretation which depends on the concrete MeshMapEvaluator/IMeshBakerDetailSampler
+		const void* DetailMesh = nullptr; // Arbitrary pointer to the mesh/surface being sampled
 		int32 DetailTriID = IndexConstants::InvalidID;
 		FVector3d DetailBaryCoords = FVector3d::Zero();
 	};
diff --git a/Engine/Plugins/Runtime/MeshModelingToolset/Source/ModelingComponents/Private/Scene/SceneCapturePhotoSet.cpp b/Engine/Plugins/Runtime/MeshModelingToolset/Source/ModelingComponents/Private/Scene/SceneCapturePhotoSet.cpp
--- a/Engine/Plugins/Runtime/MeshModelingToolset/Source/ModelingComponents/Private/Scene/SceneCapturePhotoSet.cpp
+++ b/Engine/Plugins/Runtime/MeshModelingToolset/Source/ModelingComponents/Private/Scene/SceneCapturePhotoSet.cpp
@@ -263,7 +263,15 @@ FVector4f FSceneCapturePhotoSet::FSceneSample::GetValue4f(ERenderCaptureType Cap
 	return FVector4f::Zero();
 }
 
-
+bool FSceneCapturePhotoSet::ComputeSampleLocation(
+	const FVector3d& Position,
+	const FVector3d& Normal,
+	TFunctionRef<bool(const FVector3d&, const FVector3d&)> VisibilityFunction,
+	int& PhotoIndex,
+	FVector2d& PhotoCoords) const
+{
+	return BaseColorPhotoSet.ComputeSampleLocation(Position, Normal, VisibilityFunction, PhotoIndex, PhotoCoords);
+}
 
 bool FSceneCapturePhotoSet::ComputeSample(
 	const FRenderCaptureTypeFlags& SampleChannels,
@@ -274,6 +282,7 @@ bool FSceneCapturePhotoSet::ComputeSample(
 {
 	// This could be much more efficient if (eg) we knew that all the photo sets have
 	// the same captures, then the query only has to be done once and can be used to sample each specific photo
+	// This is implemented in the other ComputeSample overload
 
 	if (SampleChannels.bBaseColor)
 	{
@@ -326,14 +335,6 @@ bool FSceneCapturePhotoSet::ComputeSample(
 	return true;
 }
 
-bool FSceneCapturePhotoSet::IsValidSample(
-	const FVector3d& Position,
-	const FVector3d& Normal,
-	TFunctionRef<bool(const FVector3d&, const FVector3d&)> VisibilityFunction) const
-{
-	return BaseColorPhotoSet.IsValidSample(Position, Normal, VisibilityFunction);
-}
-
 void FSceneCapturePhotoSet::SetEnableWriteDebugImages(bool bEnable, FString FolderName)
 {
 	bWriteDebugImages = bEnable;
diff --git a/Engine/Plugins/Runtime/MeshModelingToolset/Source/ModelingComponents/Public/Scene/SceneCapturePhotoSet.h b/Engine/Plugins/Runtime/MeshModelingToolset/Source/ModelingComponents/Public/Scene/SceneCapturePhotoSet.h
--- a/Engine/Plugins/Runtime/MeshModelingToolset/Source/ModelingComponents/Public/Scene/SceneCapturePhotoSet.h
+++ b/Engine/Plugins/Runtime/MeshModelingToolset/Source/ModelingComponents/Public/Scene/SceneCapturePhotoSet.h
@@ -112,11 +112,18 @@ public:
 		TFunctionRef<bool(const FVector3d&, const FVector3d&)> VisibilityFunction,
 		FSceneSample& DefaultsInResultsOut) const;
 
-	bool IsValidSample(
+	bool ComputeSampleLocation(
 		const FVector3d& Position,
 		const FVector3d& Normal,
-		TFunctionRef<bool(const FVector3d&, const FVector3d&)> VisibilityFunction) const;
+		TFunctionRef<bool(const FVector3d&, const FVector3d&)> VisibilityFunction,
+		int& PhotoIndex,
+		FVector2d& PhotoCoords) const;
 
+	template <ERenderCaptureType CaptureType>
+	FVector4f ComputeSample(
+		const int& PhotoIndex,
+		const FVector2d& PhotoCoords,
+		const FSceneSample& DefaultSample) const;
 
 	const FSpatialPhotoSet3f& GetBaseColorPhotoSet() { return BaseColorPhotoSet; }
 	const FSpatialPhotoSet1f& GetRoughnessPhotoSet() { return RoughnessPhotoSet; }
@@ -159,6 +166,59 @@ protected:
 };
 
 
+template <ERenderCaptureType CaptureType>
+FVector4f FSceneCapturePhotoSet::ComputeSample(
+	const int& PhotoIndex,
+	const FVector2d& PhotoCoords,
+	const FSceneSample& DefaultSample) const
+{
+	if constexpr (CaptureType == ERenderCaptureType::BaseColor)
+	{
+		FVector3f BaseColor = BaseColorPhotoSet.ComputeSample(PhotoIndex, PhotoCoords, DefaultSample.BaseColor);
+		return FVector4f(BaseColor, 1.f);
+	}
+
+	if constexpr (CaptureType == ERenderCaptureType::Roughness)
+	{
+		float Roughness = RoughnessPhotoSet.ComputeSample(PhotoIndex, PhotoCoords, DefaultSample.Roughness);
+		return FVector4f(Roughness, Roughness, Roughness, 1.f);
+	}
+
+	if constexpr (CaptureType == ERenderCaptureType::Specular)
+	{
+		float Specular = SpecularPhotoSet.ComputeSample(PhotoIndex, PhotoCoords, DefaultSample.Specular);
+		return FVector4f(Specular, Specular, Specular, 1.f);
+	}
+
+	if constexpr (CaptureType == ERenderCaptureType::Metallic)
+	{
+		float Metallic = MetallicPhotoSet.ComputeSample(PhotoIndex, PhotoCoords, DefaultSample.Metallic);
+		return FVector4f(Metallic, Metallic, Metallic, 1.f);
+	}
+
+	if constexpr (CaptureType == ERenderCaptureType::CombinedMRS)
+	{
+		FVector3f MRSValue(DefaultSample.Metallic, DefaultSample.Roughness, DefaultSample.Specular);
+		MRSValue = PackedMRSPhotoSet.ComputeSample(PhotoIndex, PhotoCoords, MRSValue);
+		return FVector4f(MRSValue, 1.f);
+	}
+
+	if constexpr (CaptureType == ERenderCaptureType::Emissive)
+	{
+		FVector3f Emissive = EmissivePhotoSet.ComputeSample(PhotoIndex, PhotoCoords, DefaultSample.Emissive);
+		return FVector4f(Emissive, 1.f);
+	}
+
+	if constexpr (CaptureType == ERenderCaptureType::WorldNormal)
+	{
+		FVector3f WorldNormal = WorldNormalPhotoSet.ComputeSample(PhotoIndex, PhotoCoords, DefaultSample.WorldNormal);
+		return FVector4f(WorldNormal, 1.f);
+	}
+
+	ensure(false);
+	return FVector4f::Zero();
+}
+
 
 } // end namespace UE::Geometry
 } // end namespace UE
diff --git a/Engine/Source/Runtime/GeometryCore/Public/Image/SpatialPhotoSet.h b/Engine/Source/Runtime/GeometryCore/Public/Image/SpatialPhotoSet.h
--- a/Engine/Source/Runtime/GeometryCore/Public/Image/SpatialPhotoSet.h
+++ b/Engine/Source/Runtime/GeometryCore/Public/Image/SpatialPhotoSet.h
@@ -81,10 +81,18 @@ public:
 		const PixelType& DefaultValue
 	) const;
 
-	bool IsValidSample(
+	PixelType ComputeSample(
+		const int& PhotoIndex,
+		const FVector2d& PhotoCoords,
+		const PixelType& DefaultValue
+	) const;
+
+	bool ComputeSampleLocation(
 		const FVector3d& Position, 
 		const FVector3d& Normal, 
-		TFunctionRef<bool(const FVector3d&,const FVector3d&)> VisibilityFunction
+		TFunctionRef<bool(const FVector3d&,const FVector3d&)> VisibilityFunction,
+		int& PhotoIndex,
+		FVector2d& PhotoCoords
 	) const;
 
 protected:
@@ -94,19 +102,19 @@ typedef TSpatialPhotoSet<FVector4f, float> FSpatialPhotoSet4f;
 typedef TSpatialPhotoSet<FVector3f, float> FSpatialPhotoSet3f;
 typedef TSpatialPhotoSet<float, float> FSpatialPhotoSet1f;
 
-
-
-
 template<typename PixelType, typename RealType>
-PixelType TSpatialPhotoSet<PixelType, RealType>::ComputeSample(
+bool TSpatialPhotoSet<PixelType, RealType>::ComputeSampleLocation(
 	const FVector3d& Position, 
 	const FVector3d& Normal, 
 	TFunctionRef<bool(const FVector3d&, const FVector3d&)> VisibilityFunction,
-	const PixelType& DefaultValue ) const
+	int& PhotoIndex,
+	FVector2d& PhotoCoords) const
 {
 	double DotTolerance = -0.1;		// dot should be negative for normal pointing towards photo
 
-	PixelType BestSample = DefaultValue;
+	PhotoIndex = IndexConstants::InvalidID;
+	PhotoCoords = FVector2d(0., 0.);
+
 	double MinDot = 1.0;
 
 	int32 NumPhotos = Num();
@@ -145,79 +153,49 @@ PixelType TSpatialPhotoSet<PixelType, RealType>::ComputeSample(
 				double v = -(PlaneY / ViewPlaneHeightWorld);
 				if (FMathd::Abs(u) < 1 && FMathd::Abs(v) < 1)
 				{
-					double x = (u/2.0 + 0.5) * (double)Photo.Dimensions.GetWidth();
-					double y = (v/2.0 + 0.5) * (double)Photo.Dimensions.GetHeight();
-					PixelType Sample = Photo.Image.template BilinearSample<RealType>(FVector2d(x, y), DefaultValue);
-
+					PhotoCoords.X = (u/2.0 + 0.5) * (double)Photo.Dimensions.GetWidth();
+					PhotoCoords.Y = (v/2.0 + 0.5) * (double)Photo.Dimensions.GetHeight();
+					PhotoIndex = pi;
 					MinDot = ViewDot;
-					BestSample = Sample;
 				}
 			}
 		}
 	}
 
-	return BestSample;
+	return PhotoIndex != IndexConstants::InvalidID;
 }
 
-
-// TODO Consider refactoring this to deduplicate code from ComputeSample, or not, since this is called in a tight inner
-// loop. Note: Unlike ComputeSample, which searches for the best valid photo, this function immediately returns from the
-// innermost if since we just want to know if at least one photo has passed the visibility check
 template<typename PixelType, typename RealType>
-bool TSpatialPhotoSet<PixelType, RealType>::IsValidSample(
+PixelType TSpatialPhotoSet<PixelType, RealType>::ComputeSample(
 	const FVector3d& Position, 
 	const FVector3d& Normal, 
-	TFunctionRef<bool(const FVector3d&, const FVector3d&)> VisibilityFunction) const
+	TFunctionRef<bool(const FVector3d&, const FVector3d&)> VisibilityFunction,
+	const PixelType& DefaultValue ) const
 {
-	double DotTolerance = -0.1;		// dot should be negative for normal pointing towards photo
-
-	double MinDot = 1.0;
-
-	int32 NumPhotos = Num();
-	for (int32 pi = 0; pi < NumPhotos; ++pi)
+	PixelType Result = DefaultValue;
+	
+	int PhotoIndex;
+	FVector2d PhotoCoords;
+	if (ComputeSampleLocation(Position, Normal, VisibilityFunction, PhotoIndex, PhotoCoords))
 	{
-		const TSpatialPhoto<PixelType>& Photo = *Photos[pi];
-		check(Photo.Dimensions.IsSquare());
-
-		FVector3d ViewDirection = Photo.Frame.X();
-		double ViewDot = ViewDirection.Dot(Normal);
-		if (ViewDot > DotTolerance || ViewDot > MinDot)
-		{
-			continue;
-		}
-
-		FFrame3d ViewPlane = Photo.Frame;
-		ViewPlane.Origin += Photo.NearPlaneDist * ViewDirection;
-
-		double ViewPlaneWidthWorld = Photo.NearPlaneDist * FMathd::Tan(Photo.HorzFOVDegrees * 0.5 * FMathd::DegToRad);
-		double ViewPlaneHeightWorld = ViewPlaneWidthWorld;
-
-		FVector3d RayOrigin = Photo.Frame.Origin;
-		FVector3d RayDir = Normalized(Position - RayOrigin);
-		FVector3d HitPoint;
-		bool bHit = ViewPlane.RayPlaneIntersection(RayOrigin, RayDir, 0, HitPoint);
-		if (bHit)
-		{
-			bool bVisible = VisibilityFunction(Position, HitPoint);
-			if ( bVisible )
-			{
-				double PlaneX = (HitPoint - ViewPlane.Origin).Dot(ViewPlane.Y());
-				double PlaneY = (HitPoint - ViewPlane.Origin).Dot(ViewPlane.Z());
-
-				//FVector2d PlanePos = ViewPlane.ToPlaneUV(HitPoint, 0);
-				double u = PlaneX / ViewPlaneWidthWorld;
-				double v = -(PlaneY / ViewPlaneHeightWorld);
-				if (FMathd::Abs(u) < 1 && FMathd::Abs(v) < 1)
-				{
-					return true;
-				}
-			}
-		}
+		const TSpatialPhoto<PixelType>& Photo = *Photos[PhotoIndex];
+		Result = Photo.Image.template BilinearSample<RealType>(PhotoCoords, DefaultValue);
 	}
+	
+	return Result;
+}
 
-	return false;
+template<typename PixelType, typename RealType>
+PixelType TSpatialPhotoSet<PixelType, RealType>::ComputeSample(
+	const int& PhotoIndex, 
+	const FVector2d& PhotoCoords, 
+	const PixelType& DefaultValue) const
+{
+	const TSpatialPhoto<PixelType>& Photo = *Photos[PhotoIndex];
+	return Photo.Image.template BilinearSample<RealType>(PhotoCoords, DefaultValue);
 }
 
 
+
 } // end namespace UE::Geometry
 } // end namespace UE
\ No newline at end of file
