Render Capture Baking: Add FScopedSlowTask to give user progress feedback on game thread scene capture

#rnx

diff --git a/Engine/Plugins/Experimental/MeshModelingToolsetExp/Source/MeshModelingToolsExp/Private/BakeRenderCaptureTool.cpp b/Engine/Plugins/Experimental/MeshModelingToolsetExp/Source/MeshModelingToolsExp/Private/BakeRenderCaptureTool.cpp
--- a/Engine/Plugins/Experimental/MeshModelingToolsetExp/Source/MeshModelingToolsExp/Private/BakeRenderCaptureTool.cpp
+++ b/Engine/Plugins/Experimental/MeshModelingToolsetExp/Source/MeshModelingToolsExp/Private/BakeRenderCaptureTool.cpp
@@ -19,6 +19,7 @@
 #include "Sampling/GenericEvaluator.h"
 #include "Image/ImageInfilling.h"
 #include "Algo/NoneOf.h"
+#include "Misc/ScopedSlowTask.h"
 
 
 using namespace UE::Geometry;
@@ -175,17 +176,23 @@ public:
 
 static TUniquePtr<FSceneCapturePhotoSet> CapturePhotoSet(
 	const TArray<TObjectPtr<AActor>>& Actors,
-	const FBakeRenderCaptureOptions::FOptions& Options
+	const FBakeRenderCaptureOptions::FOptions& Options,
+	bool bAllowCancel
 )
 {
 	TRACE_CPUPROFILER_EVENT_SCOPE(CapturePhotoSet);
 
+	FScopedSlowTask Progress(1.f, LOCTEXT("CapturingScene", "Capturing Scene..."));
+	Progress.EnterProgressFrame(1.f);
+	Progress.MakeDialog(bAllowCancel);
+
 	double FieldOfView = Options.FieldOfViewDegrees;
 	double NearPlaneDist = Options.NearPlaneDist;
 
 	FImageDimensions CaptureDimensions(Options.RenderCaptureImageSize, Options.RenderCaptureImageSize);
 
 	TUniquePtr<FSceneCapturePhotoSet> SceneCapture = MakeUnique<FSceneCapturePhotoSet>();
+	SceneCapture->SetAllowCancel(bAllowCancel);
 
 	SceneCapture->SetCaptureTypeEnabled(ERenderCaptureType::BaseColor, Options.bBakeBaseColor);
 	SceneCapture->SetCaptureTypeEnabled(ERenderCaptureType::WorldNormal, Options.bBakeNormalMap);
@@ -213,7 +220,7 @@ static TUniquePtr<FSceneCapturePhotoSet> CapturePhotoSet(
 	SceneCapture->AddStandardExteriorCapturesFromBoundingBox(
 		CaptureDimensions, FieldOfView, NearPlaneDist,
 		true, true, true);
-	
+
 	return SceneCapture;
 }
 
@@ -676,6 +683,9 @@ void UBakeRenderCaptureTool::Setup()
 	AddToolPropertySource(ResultSettings);
 	SetToolPropertySourceEnabled(ResultSettings, true);
 
+	// Used to implmement SceneCapture cancellation
+	ComputedSettings = NewObject<UBakeRenderCaptureToolProperties>(this);
+
 	// Hide the render capture meshes since this baker operates solely in world space which will occlude the preview of
 	// the target mesh.
 	for (int Idx = 1; Idx < NumTargets; ++Idx)
@@ -1011,7 +1021,6 @@ void UBakeRenderCaptureTool::InvalidateComputeRC()
 		ComputeRC->OnResultUpdated.AddLambda([this](const TUniquePtr<FBakeRenderCaptureResultsBuilder>& NewResult) { OnMapsUpdatedRC(NewResult); });
 	}
 	ComputeRC->InvalidateResult();
-	OpState = EBakeOpState::Clean;
 }
 
 
@@ -1035,14 +1044,43 @@ void UBakeRenderCaptureTool::UpdateResult()
 		UE::ToolTarget::ShowSourceObject(Targets[Idx]);
 	}
 
+	// Do not allow user-cancellation on the call that occurs when the Render Capture Tool starts up
+	const bool bAllowCancel = (bFirstEverSceneCapture == false);
+
 	SceneCapture.Reset();
-	SceneCapture = CapturePhotoSet(Actors, Options);
+	SceneCapture = CapturePhotoSet(Actors, Options, bAllowCancel);
 
 	for (int Idx = 1; Idx < Targets.Num(); ++Idx)
 	{
 		UE::ToolTarget::HideSourceObject(Targets[Idx]);
 	}
 
+	if (SceneCapture->Cancelled())
+	{
+		// Restore the settings present before the change that invoked the scene capture recompute
+		Settings->SamplesPerPixel = ComputedSettings->SamplesPerPixel;
+		Settings->RenderCaptureResolution = ComputedSettings->RenderCaptureResolution;
+		Settings->MaterialSettings = ComputedSettings->MaterialSettings;
+		Settings->CaptureFieldOfView = ComputedSettings->CaptureFieldOfView;
+		Settings->NearPlaneDist = ComputedSettings->NearPlaneDist;
+
+		// Silently make the above Settings updates so we don't overwrite the change to OpState below and call this function again
+		Settings->SilentUpdateWatched();
+
+		OpState = EBakeOpState::Clean;
+
+		return;
+	}
+
+	bFirstEverSceneCapture = false;
+
+	// Cache Settings used to compute this SceneCapture so we can restore them if the user cancels a SceneCapture recompute
+	ComputedSettings->SamplesPerPixel = Settings->SamplesPerPixel;
+	ComputedSettings->RenderCaptureResolution = Settings->RenderCaptureResolution;
+	ComputedSettings->MaterialSettings = Settings->MaterialSettings;
+	ComputedSettings->CaptureFieldOfView = Settings->CaptureFieldOfView;
+	ComputedSettings->NearPlaneDist = Settings->NearPlaneDist;
+
 	// @Incomplete Clear our invalid bitflag to check for valid inputs.
 	//OpState &= ~EBakeOpState::Invalid;
 	//OpState |= CheckValidInputs();
@@ -1053,9 +1091,9 @@ void UBakeRenderCaptureTool::UpdateResult()
 	//}
 	
 	InvalidateResults();
-
-	// This should be the only point of compute invalidation to minimize synchronization issues.
 	InvalidateComputeRC();
+
+	OpState = EBakeOpState::Clean;
 }
 
 
diff --git a/Engine/Plugins/Experimental/MeshModelingToolsetExp/Source/MeshModelingToolsExp/Public/BakeRenderCaptureTool.h b/Engine/Plugins/Experimental/MeshModelingToolsetExp/Source/MeshModelingToolsExp/Public/BakeRenderCaptureTool.h
--- a/Engine/Plugins/Experimental/MeshModelingToolsetExp/Source/MeshModelingToolsExp/Public/BakeRenderCaptureTool.h
+++ b/Engine/Plugins/Experimental/MeshModelingToolsetExp/Source/MeshModelingToolsExp/Public/BakeRenderCaptureTool.h
@@ -320,8 +320,13 @@ protected:
 	TObjectPtr<UMaterialInstanceDynamic> PreviewMaterialPackedRC;
 	void InitializePreviewMaterials();
 
-	// @Refactor TODO We currently need to compute this on the game thread because the implementation has checks for this
+	// TODO We currently need to compute this on the game thread because the implementation has checks for this
 	TUniquePtr<UE::Geometry::FSceneCapturePhotoSet> SceneCapture;
+	bool bFirstEverSceneCapture = true;
+
+	// If the user cancels a scene capture before the computation completes then the settings which changed to invoke
+	// the capture are reverted to these values
+	TObjectPtr<UBakeRenderCaptureToolProperties> ComputedSettings;
 
 	// Analytics
 	virtual FString GetAnalyticsEventName() const override
diff --git a/Engine/Plugins/Runtime/MeshModelingToolset/Source/ModelingComponents/Private/Scene/SceneCapturePhotoSet.cpp b/Engine/Plugins/Runtime/MeshModelingToolset/Source/ModelingComponents/Private/Scene/SceneCapturePhotoSet.cpp
--- a/Engine/Plugins/Runtime/MeshModelingToolset/Source/ModelingComponents/Private/Scene/SceneCapturePhotoSet.cpp
+++ b/Engine/Plugins/Runtime/MeshModelingToolset/Source/ModelingComponents/Private/Scene/SceneCapturePhotoSet.cpp
@@ -2,9 +2,12 @@
 
 #include "Scene/SceneCapturePhotoSet.h"
 #include "EngineUtils.h"
+#include "Misc/ScopedSlowTask.h"
 
 using namespace UE::Geometry;
 
+#define LOCTEXT_NAMESPACE "SceneCapture"
+
 void FSceneCapturePhotoSet::SetCaptureSceneActors(UWorld* World, const TArray<AActor*>& Actors)
 {
 	this->TargetWorld = World;
@@ -90,6 +93,9 @@ void FSceneCapturePhotoSet::AddExteriorCaptures(
 	const TArray<FVector3d>& Directions)
 {
 	check(this->TargetWorld != nullptr);
+	
+	FScopedSlowTask Progress(Directions.Num(), LOCTEXT("ComputingViewpoints", "Computing Viewpoints..."));
+	Progress.MakeDialog(bAllowCancel);
 
 	// Workaround for Nanite scene proxies visibility
 	// Unregister all components to remove unwanted proxies from the scene. This is currently the only way to "hide" nanite meshes.
@@ -104,6 +110,16 @@ void FSceneCapturePhotoSet::AddExteriorCaptures(
 		}
 	}
 
+	ON_SCOPE_EXIT
+	{
+		// Workaround for Nanite scene proxies visibility
+		// Reregister all components we previously unregistered
+		for (AActor* Actor : ActorsToRegister)
+		{
+			Actor->RegisterAllComponents();
+		}
+	};
+
 	FWorldRenderCapture RenderCapture;
 	RenderCapture.SetWorld(TargetWorld);
 	RenderCapture.SetVisibleActors(VisibleActors);
@@ -119,6 +135,13 @@ void FSceneCapturePhotoSet::AddExteriorCaptures(
 	int32 NumDirections = Directions.Num();
 	for (int32 di = 0; di < NumDirections; ++di)
 	{
+		Progress.EnterProgressFrame(1.f);
+		if (Progress.ShouldCancel())
+		{
+			bWasCancelled = true;
+			return;
+		}
+
 		FVector3d ViewDirection = Directions[di];
 		ViewDirection.Normalize();
 
@@ -134,12 +157,13 @@ void FSceneCapturePhotoSet::AddExteriorCaptures(
 		BasePhoto3f.HorzFOVDegrees = HorizontalFOVDegrees;
 		BasePhoto3f.Dimensions = PhotoDimensions;
 
-		auto CaptureImageTypeFunc_3f = [&BasePhoto3f, &RenderCapture](ERenderCaptureType CaptureType, FSpatialPhotoSet3f& PhotoSet)
+		auto CaptureImageTypeFunc_3f = [&Progress, &BasePhoto3f, &RenderCapture](ERenderCaptureType CaptureType, FSpatialPhotoSet3f& PhotoSet)
 		{
 			FSpatialPhoto3f NewPhoto = BasePhoto3f;
 			FImageAdapter Image(&NewPhoto.Image);
 			RenderCapture.CaptureFromPosition(CaptureType, NewPhoto.Frame, NewPhoto.HorzFOVDegrees, NewPhoto.NearPlaneDist, Image);
 			PhotoSet.Add(MoveTemp(NewPhoto));
+			Progress.TickProgress();
 		};
 
 		FSpatialPhoto1f BasePhoto1f;
@@ -148,12 +172,13 @@ void FSceneCapturePhotoSet::AddExteriorCaptures(
 		BasePhoto1f.HorzFOVDegrees = HorizontalFOVDegrees;
 		BasePhoto1f.Dimensions = PhotoDimensions;
 
-		auto CaptureImageTypeFunc_1f = [&BasePhoto1f, &RenderCapture](ERenderCaptureType CaptureType, FSpatialPhotoSet1f& PhotoSet)
+		auto CaptureImageTypeFunc_1f = [&Progress, &BasePhoto1f, &RenderCapture](ERenderCaptureType CaptureType, FSpatialPhotoSet1f& PhotoSet)
 		{
 			FSpatialPhoto1f NewPhoto = BasePhoto1f;
 			FImageAdapter Image(&NewPhoto.Image);
 			RenderCapture.CaptureFromPosition(CaptureType, NewPhoto.Frame, NewPhoto.HorzFOVDegrees, NewPhoto.NearPlaneDist, Image);
 			PhotoSet.Add(MoveTemp(NewPhoto));
+			Progress.TickProgress();
 		};
 
 		if (bEnableBaseColor)
@@ -185,13 +210,6 @@ void FSceneCapturePhotoSet::AddExteriorCaptures(
 			CaptureImageTypeFunc_3f(ERenderCaptureType::Emissive, EmissivePhotoSet);
 		}
 	}
-
-	// Workaround for Nanite scene proxies visibility
-	// Reregister all components we previously unregistered
-	for (AActor* Actor : ActorsToRegister)
-	{
-		Actor->RegisterAllComponents();
-	}
 }
 
 
@@ -342,4 +360,6 @@ void FSceneCapturePhotoSet::SetEnableWriteDebugImages(bool bEnable, FString Fold
 	{
 		DebugImagesFolderName = FolderName;
 	}
-}
\ No newline at end of file
+}
+
+#undef LOCTEXT_NAMESPACE
\ No newline at end of file
diff --git a/Engine/Plugins/Runtime/MeshModelingToolset/Source/ModelingComponents/Public/MeshOpPreviewHelpers.h b/Engine/Plugins/Runtime/MeshModelingToolset/Source/ModelingComponents/Public/MeshOpPreviewHelpers.h
--- a/Engine/Plugins/Runtime/MeshModelingToolset/Source/ModelingComponents/Public/MeshOpPreviewHelpers.h
+++ b/Engine/Plugins/Runtime/MeshModelingToolset/Source/ModelingComponents/Public/MeshOpPreviewHelpers.h
@@ -358,9 +358,13 @@ public:
 		check(BackgroundCompute);
 		if (BackgroundCompute)
 		{
-			return BackgroundCompute->GetElapsedComputeTime();
+			typename ComputeSourceType::FStatus Status = BackgroundCompute->CheckStatus();
+			if (Status.TaskStatus == UE::Geometry::EBackgroundComputeTaskStatus::InProgress)
+			{
+				return static_cast<float>(Status.ElapsedTime);
+			}
 		}
-		return 0.0;
+		return 0.f;
 	}
 
 
diff --git a/Engine/Plugins/Runtime/MeshModelingToolset/Source/ModelingComponents/Public/Scene/SceneCapturePhotoSet.h b/Engine/Plugins/Runtime/MeshModelingToolset/Source/ModelingComponents/Public/Scene/SceneCapturePhotoSet.h
--- a/Engine/Plugins/Runtime/MeshModelingToolset/Source/ModelingComponents/Public/Scene/SceneCapturePhotoSet.h
+++ b/Engine/Plugins/Runtime/MeshModelingToolset/Source/ModelingComponents/Public/Scene/SceneCapturePhotoSet.h
@@ -8,6 +8,7 @@
 
 class UWorld;
 class AActor;
+struct FScopedSlowTask;
 
 namespace UE
 {
@@ -140,6 +141,15 @@ public:
 	 */
 	void SetEnableWriteDebugImages(bool bEnable, FString FolderName = FString());
 
+	void SetAllowCancel(bool bAllowCancelIn)
+	{
+		bAllowCancel = bAllowCancelIn;
+	}
+
+	bool Cancelled() const
+	{
+		return bWasCancelled;
+	}
 
 protected:
 	UWorld* TargetWorld = nullptr;
@@ -163,6 +173,9 @@ protected:
 
 	bool bWriteDebugImages = false;
 	FString DebugImagesFolderName = TEXT("SceneCapturePhotoSet");
+
+	bool bAllowCancel = false;
+	bool bWasCancelled = false;
 };
 
 
