Render Capture Baking: Refactor tool properties so that only genuine render capture options trigger game thread recomputes, also improve names

#rnx

diff --git a/Engine/Plugins/Experimental/MeshModelingToolsetExp/Source/MeshModelingToolsExp/Private/BakeRenderCaptureTool.cpp b/Engine/Plugins/Experimental/MeshModelingToolsetExp/Source/MeshModelingToolsExp/Private/BakeRenderCaptureTool.cpp
--- a/Engine/Plugins/Experimental/MeshModelingToolsetExp/Source/MeshModelingToolsExp/Private/BakeRenderCaptureTool.cpp
+++ b/Engine/Plugins/Experimental/MeshModelingToolsetExp/Source/MeshModelingToolsExp/Private/BakeRenderCaptureTool.cpp
@@ -40,11 +40,16 @@ public:
 		TFunctionRef<bool(const FVector3d&, const FVector3d&)> VisibilityFunction,
 		const FDynamicMesh3* Mesh,
 		const FDynamicMeshAABBTree3* Spatial,
-		const FMeshTangentsd* Tangents = nullptr) :
+		const FMeshTangentsd* Tangents)
+	:
 		FMeshBakerDynamicMeshSampler(Mesh, Spatial, Tangents),
 		SceneCapture(SceneCapture),
 		VisibilityFunction(VisibilityFunction)
 	{
+		check(SceneCapture != nullptr);
+		check(Mesh != nullptr);
+		check(Spatial != nullptr);
+		check(Tangents != nullptr);
 	}
 
 	virtual bool SupportsCustomCorrespondence() const override
@@ -82,15 +87,15 @@ public:
 
 
 
-static FName BaseColorTexParamName = FName("BaseColor");
-static FName RoughnessTexParamName = FName("Roughness");
-static FName MetallicTexParamName = FName("Metallic");
-static FName SpecularTexParamName = FName("Specular");
-static FName EmissiveTexParamName = FName("Emissive");
-static FName NormalTexParamName = FName("NormalMap");
-static FName PackedMRSTexParamName = FName("PackedMRS");
+static FString BaseColorTexParamName = TEXT("BaseColor");
+static FString RoughnessTexParamName = TEXT("Roughness");
+static FString MetallicTexParamName = TEXT("Metallic");
+static FString SpecularTexParamName = TEXT("Specular");
+static FString EmissiveTexParamName = TEXT("Emissive");
+static FString NormalTexParamName = TEXT("NormalMap");
+static FString PackedMRSTexParamName = TEXT("PackedMRS");
 
-class FBakeRenderCaptureOptions
+class FRenderCaptureSettings
 {
 public:
 	enum class ETextureSizePolicy : uint8
@@ -128,11 +133,6 @@ public:
 		
 		bool bUsePackedMRS = true;
 
-		// output texture options
-		int32 TextureImageSize = 1024;
-
-		EBakeTextureSamplesPerPixel SamplesPerPixel = EBakeTextureSamplesPerPixel::Sample1;
-
 		//
 		// Mesh settings
 		//
@@ -145,7 +145,7 @@ public:
 	 * Construct an FOptions from the provided FMeshApproximationSettings.
 	 */
 	static FOptions ConstructOptions(
-		const UBakeRenderCaptureToolProperties& UseSettings,
+		const URenderCaptureProperties& RenderCaptureProperties,
 		const UBakeRenderCaptureInputToolProperties& InputMeshSettings)
 	{
 		//
@@ -154,21 +154,19 @@ public:
 		FOptions Options;
 
 		Options.TargetUVLayer = InputMeshSettings.TargetUVLayerNamesList.IndexOfByKey(InputMeshSettings.TargetUVLayer);
-		
-		Options.RenderCaptureImageSize = (UseSettings.RenderCaptureResolution == 0) ? Options.TextureImageSize : UseSettings.RenderCaptureResolution;
-		Options.FieldOfViewDegrees = UseSettings.CaptureFieldOfView;
-		Options.NearPlaneDist = UseSettings.NearPlaneDist;
-		
+
+		Options.RenderCaptureImageSize = static_cast<int32>(RenderCaptureProperties.Resolution);
+
 		//Options.bBakeBaseColor // This is always true
-		Options.bBakeRoughness = UseSettings.MaterialSettings.bRoughnessMap;
-		Options.bBakeMetallic = UseSettings.MaterialSettings.bMetallicMap;
-		Options.bBakeSpecular = UseSettings.MaterialSettings.bSpecularMap;
-		Options.bBakeEmissive = UseSettings.MaterialSettings.bEmissiveMap;
-		Options.bBakeNormalMap = UseSettings.MaterialSettings.bNormalMap;
-		Options.bUsePackedMRS = UseSettings.MaterialSettings.bPackedMRSMap;
+		Options.bBakeRoughness = RenderCaptureProperties.bRoughnessMap;
+		Options.bBakeMetallic = RenderCaptureProperties.bMetallicMap;
+		Options.bBakeSpecular = RenderCaptureProperties.bSpecularMap;
+		Options.bBakeEmissive = RenderCaptureProperties.bEmissiveMap;
+		Options.bBakeNormalMap = RenderCaptureProperties.bNormalMap;
+		Options.bUsePackedMRS = RenderCaptureProperties.bPackedMRSMap;
 
-		Options.TextureImageSize = UseSettings.MaterialSettings.TextureSize;
-		Options.SamplesPerPixel = UseSettings.SamplesPerPixel;
+		Options.FieldOfViewDegrees = RenderCaptureProperties.CaptureFieldOfView;
+		Options.NearPlaneDist = RenderCaptureProperties.NearPlaneDist;
 
 		return Options;
 	}
@@ -176,7 +174,7 @@ public:
 
 static TUniquePtr<FSceneCapturePhotoSet> CapturePhotoSet(
 	const TArray<TObjectPtr<AActor>>& Actors,
-	const FBakeRenderCaptureOptions::FOptions& Options,
+	const FRenderCaptureSettings::FOptions& Options,
 	bool bAllowCancel
 )
 {
@@ -252,7 +250,9 @@ MakeColorEvaluator(const FSceneCapturePhotoSet::FSceneSample& DefaultSample, con
 
 static void ImageBuildersFromPhotoSet(
 	FSceneCapturePhotoSet* SceneCapture,
-	const FBakeRenderCaptureOptions::FOptions& Options,
+	const FRenderCaptureSettings::FOptions& Options,
+	const int32 TextureImageSize,
+	const EBakeTextureSamplesPerPixel SamplesPerPixel,
 	const FDynamicMesh3* BaseMesh,
 	const TSharedPtr<UE::Geometry::FMeshTangentsd, ESPMode::ThreadSafe>& BaseMeshTangents,
 	TUniquePtr<FBakeRenderCaptureResultsBuilder>& Results)
@@ -261,7 +261,7 @@ static void ImageBuildersFromPhotoSet(
 
 	FDynamicMeshAABBTree3 BaseMeshSpatial(BaseMesh, true);
 
-	double RayOffsetHackDist = (double)(100.0f * FMathf::ZeroTolerance * BaseMesh->GetBounds().MinDim() );
+	double RayOffsetHackDist = (double)(100.0 * FMathf::ZeroTolerance * BaseMesh->GetBounds().MinDim() );
 	auto VisibilityFunction = [&BaseMeshSpatial, RayOffsetHackDist](const FVector3d& SurfPos, const FVector3d& ImagePosWorld)
 	{
 		FVector3d RayDir = ImagePosWorld - SurfPos;
@@ -309,7 +309,7 @@ static void ImageBuildersFromPhotoSet(
 		TArray<bool> EvaluatorNeedsInfill;
 	} InfillData;
 
-	InfillData.SampleStats.SetDimensions(FImageDimensions(Options.TextureImageSize, Options.TextureImageSize));
+	InfillData.SampleStats.SetDimensions(FImageDimensions(TextureImageSize, TextureImageSize));
 	InfillData.SampleStats.Clear(FInfillData::FSampleStats{});
 
 	auto RegisterSampleStats = [&InfillData](bool bSampleValid, const FMeshMapEvaluator::FCorrespondenceSample& Sample, const FVector2d& UVPosition, const FVector2i& ImageCoords)
@@ -380,8 +380,8 @@ static void ImageBuildersFromPhotoSet(
 	FMeshMapBaker Baker;
 	Baker.SetTargetMesh(BaseMesh);
 	Baker.SetTargetMeshTangents(BaseMeshTangents);
-	Baker.SetDimensions(FImageDimensions(Options.TextureImageSize, Options.TextureImageSize));
-	Baker.SetSamplesPerPixel(static_cast<int32>(Options.SamplesPerPixel));
+	Baker.SetDimensions(FImageDimensions(TextureImageSize, TextureImageSize));
+	Baker.SetSamplesPerPixel(static_cast<int32>(SamplesPerPixel));
 	Baker.SetFilter(FMeshMapBaker::EBakeFilterType::BSpline);
 	Baker.SetTargetMeshUVLayer(Options.TargetUVLayer);
 	Baker.InteriorSampleCallback = RegisterSampleStats;
@@ -556,9 +556,11 @@ public:
 	TArray<TObjectPtr<AActor>>* Actors;
 	UE::Geometry::FDynamicMesh3* BaseMesh = nullptr;
 	TSharedPtr<UE::Geometry::FMeshTangentsd, ESPMode::ThreadSafe> BaseMeshTangents;
-	FBakeRenderCaptureOptions::FOptions Options;
+	FRenderCaptureSettings::FOptions Options;
+	int32 TextureImageSize;
+	EBakeTextureSamplesPerPixel SamplesPerPixel;
 	TObjectPtr<UBakeRenderCaptureInputToolProperties> InputMeshSettings;
-	TUniquePtr<FSceneCapturePhotoSet> SceneCapture;
+	FSceneCapturePhotoSet* SceneCapture = nullptr;
 	
 	// Begin TGenericDataOperator interface
 	// @Incomplete Use the Progress thing
@@ -567,10 +569,10 @@ public:
 		check(Actors != nullptr);
 		check(BaseMesh != nullptr);
 		check(BaseMeshTangents.IsValid());
-		check(SceneCapture.IsValid());
+		check(SceneCapture != nullptr);
 
 		// Bake textures onto the base/target mesh by projecting/sampling the set of captured photos
-		ImageBuildersFromPhotoSet(SceneCapture.Get(), Options, BaseMesh, BaseMeshTangents, Result);
+		ImageBuildersFromPhotoSet(SceneCapture, Options, TextureImageSize, SamplesPerPixel, BaseMesh, BaseMeshTangents, Result);
 	}
 	// End TGenericDataOperator interface
 };
@@ -652,13 +654,24 @@ void UBakeRenderCaptureTool::Setup()
 	Settings->RestoreProperties(this);
 	AddToolPropertySource(Settings);
 
-	Settings->MapPreview = BaseColorTexParamName.ToString(); // We always bake the base color
+	Settings->MapPreview = BaseColorTexParamName; // We always bake the base color
 	Settings->WatchProperty(Settings->MapPreview, [this](FString) { UpdateVisualization(); GetToolManager()->PostInvalidation(); });
 	Settings->WatchProperty(Settings->SamplesPerPixel, [this](EBakeTextureSamplesPerPixel) { OpState |= EBakeOpState::Evaluate; });
-	Settings->WatchProperty(Settings->RenderCaptureResolution, [this](int32) {  OpState |= EBakeOpState::Evaluate; });
-	Settings->WatchProperty(Settings->MaterialSettings, [this](FMaterialProxySettingsRC) { OpState |= EBakeOpState::Evaluate; });
-	Settings->WatchProperty(Settings->CaptureFieldOfView, [this](float) { OpState |= EBakeOpState::Evaluate; });
-	Settings->WatchProperty(Settings->NearPlaneDist, [this](float) { OpState |= EBakeOpState::Evaluate; });
+	Settings->WatchProperty(Settings->TextureSize, [this](EBakeTextureResolution) { OpState |= EBakeOpState::Evaluate; });
+
+	RenderCaptureProperties = NewObject<URenderCaptureProperties>(this);
+	RenderCaptureProperties->RestoreProperties(this);
+	AddToolPropertySource(RenderCaptureProperties);
+
+	RenderCaptureProperties->WatchProperty(RenderCaptureProperties->Resolution, [this](EBakeTextureResolution) { OpState |= EBakeOpState::Evaluate; });
+	RenderCaptureProperties->WatchProperty(RenderCaptureProperties->bNormalMap, [this](bool) { OpState |= EBakeOpState::Evaluate; });
+	RenderCaptureProperties->WatchProperty(RenderCaptureProperties->bMetallicMap, [this](bool) { OpState |= EBakeOpState::Evaluate; });
+	RenderCaptureProperties->WatchProperty(RenderCaptureProperties->bRoughnessMap, [this](bool) { OpState |= EBakeOpState::Evaluate; });
+	RenderCaptureProperties->WatchProperty(RenderCaptureProperties->bSpecularMap, [this](bool) { OpState |= EBakeOpState::Evaluate; });
+	RenderCaptureProperties->WatchProperty(RenderCaptureProperties->bPackedMRSMap, [this](bool) { OpState |= EBakeOpState::Evaluate; });
+	RenderCaptureProperties->WatchProperty(RenderCaptureProperties->bEmissiveMap, [this](bool) { OpState |= EBakeOpState::Evaluate; });
+	RenderCaptureProperties->WatchProperty(RenderCaptureProperties->CaptureFieldOfView, [this](float) { OpState |= EBakeOpState::Evaluate; });
+	RenderCaptureProperties->WatchProperty(RenderCaptureProperties->NearPlaneDist, [this](float) { OpState |= EBakeOpState::Evaluate; });
 	
 	InputMeshSettings = NewObject<UBakeRenderCaptureInputToolProperties>(this);
 	InputMeshSettings->RestoreProperties(this);
@@ -668,14 +681,13 @@ void UBakeRenderCaptureTool::Setup()
 	InputMeshSettings->WatchProperty(InputMeshSettings->TargetUVLayer, [this](FString) { OpState |= EBakeOpState::Evaluate; });
 	
 	{
-		FBakeRenderCaptureOptions::FOptions Options = FBakeRenderCaptureOptions::ConstructOptions(*Settings, *InputMeshSettings);
-		Settings->MapPreviewNamesList.Add(NormalTexParamName.ToString());
-		Settings->MapPreviewNamesList.Add(BaseColorTexParamName.ToString());
-		Settings->MapPreviewNamesList.Add(RoughnessTexParamName.ToString());
-		Settings->MapPreviewNamesList.Add(MetallicTexParamName.ToString());
-		Settings->MapPreviewNamesList.Add(SpecularTexParamName.ToString());
-		Settings->MapPreviewNamesList.Add(EmissiveTexParamName.ToString());
-		Settings->MapPreviewNamesList.Add(PackedMRSTexParamName.ToString());
+		Settings->MapPreviewNamesList.Add(NormalTexParamName);
+		Settings->MapPreviewNamesList.Add(BaseColorTexParamName);
+		Settings->MapPreviewNamesList.Add(RoughnessTexParamName);
+		Settings->MapPreviewNamesList.Add(MetallicTexParamName);
+		Settings->MapPreviewNamesList.Add(SpecularTexParamName);
+		Settings->MapPreviewNamesList.Add(EmissiveTexParamName);
+		Settings->MapPreviewNamesList.Add(PackedMRSTexParamName);
 	}
 
 	ResultSettings = NewObject<UBakeRenderCaptureResults>(this);
@@ -683,8 +695,8 @@ void UBakeRenderCaptureTool::Setup()
 	AddToolPropertySource(ResultSettings);
 	SetToolPropertySourceEnabled(ResultSettings, true);
 
-	// Used to implmement SceneCapture cancellation
-	ComputedSettings = NewObject<UBakeRenderCaptureToolProperties>(this);
+	// Used to implement SceneCapture cancellation
+	ComputedRenderCaptureProperties = NewObject<URenderCaptureProperties>(this);
 
 	// Hide the render capture meshes since this baker operates solely in world space which will occlude the preview of
 	// the target mesh.
@@ -693,7 +705,9 @@ void UBakeRenderCaptureTool::Setup()
 		UE::ToolTarget::HideSourceObject(Targets[Idx]);
 	}
 
+	// Make sure we trigger SceneCapture computation in UpdateResult
 	OpState |= EBakeOpState::Evaluate;
+	ComputedRenderCaptureProperties->NearPlaneDist = 0.f; // Arbitrary invalid value
 
 	SetToolDisplayName(LOCTEXT("ToolName", "Bake Render Capture"));
 	GetToolManager()->DisplayMessage(
@@ -711,6 +725,7 @@ void UBakeRenderCaptureTool::OnShutdown(EToolShutdownType ShutdownType)
 	Super::OnShutdown(ShutdownType);
 	
 	Settings->SaveProperties(this);
+	RenderCaptureProperties->SaveProperties(this);
 	InputMeshSettings->SaveProperties(this);
 
 	if (ComputeRC)
@@ -763,44 +778,44 @@ void UBakeRenderCaptureTool::CreateTextureAssetsRC(UWorld* SourceWorld, UObject*
 
 	if (ResultSettings->BaseColorMap != nullptr)
 	{
-		const FString TexName = FString::Printf(TEXT("%s_%s"), *BaseName, *BaseColorTexParamName.ToString());
+		const FString TexName = FString::Printf(TEXT("%s_%s"), *BaseName, *BaseColorTexParamName);
 		CreateTextureAsset(TexName, FTexture2DBuilder::ETextureType::Color, ResultSettings->BaseColorMap);
 	}
 
-	if (Settings->MaterialSettings.bNormalMap && ResultSettings->NormalMap != nullptr)
+	if (RenderCaptureProperties->bNormalMap && ResultSettings->NormalMap != nullptr)
 	{
-		const FString TexName = FString::Printf(TEXT("%s_%s"), *BaseName, *NormalTexParamName.ToString());
+		const FString TexName = FString::Printf(TEXT("%s_%s"), *BaseName, *NormalTexParamName);
 		CreateTextureAsset(TexName, FTexture2DBuilder::ETextureType::NormalMap, ResultSettings->NormalMap);
 	}
 
-	if (Settings->MaterialSettings.bEmissiveMap && ResultSettings->EmissiveMap != nullptr)
+	if (RenderCaptureProperties->bEmissiveMap && ResultSettings->EmissiveMap != nullptr)
 	{
-		const FString TexName = FString::Printf(TEXT("%s_%s"), *BaseName, *EmissiveTexParamName.ToString());
+		const FString TexName = FString::Printf(TEXT("%s_%s"), *BaseName, *EmissiveTexParamName);
 		CreateTextureAsset(TexName, FTexture2DBuilder::ETextureType::EmissiveHDR, ResultSettings->EmissiveMap);
 	}
 
-	if (Settings->MaterialSettings.bPackedMRSMap && ResultSettings->PackedMRSMap != nullptr)
+	if (RenderCaptureProperties->bPackedMRSMap && ResultSettings->PackedMRSMap != nullptr)
 	{
-		const FString TexName = FString::Printf(TEXT("%s_%s"), *BaseName, *PackedMRSTexParamName.ToString());
+		const FString TexName = FString::Printf(TEXT("%s_%s"), *BaseName, *PackedMRSTexParamName);
 		CreateTextureAsset(TexName, FTexture2DBuilder::ETextureType::ColorLinear, ResultSettings->PackedMRSMap);
 	}
 	else
 	{
-		if (Settings->MaterialSettings.bMetallicMap && ResultSettings->MetallicMap != nullptr)
+		if (RenderCaptureProperties->bMetallicMap && ResultSettings->MetallicMap != nullptr)
 		{
-			const FString TexName = FString::Printf(TEXT("%s_%s"), *BaseName, *MetallicTexParamName.ToString());
+			const FString TexName = FString::Printf(TEXT("%s_%s"), *BaseName, *MetallicTexParamName);
 			CreateTextureAsset(TexName, FTexture2DBuilder::ETextureType::Metallic, ResultSettings->MetallicMap);
 		}
 
-		if (Settings->MaterialSettings.bRoughnessMap && ResultSettings->RoughnessMap != nullptr)
+		if (RenderCaptureProperties->bRoughnessMap && ResultSettings->RoughnessMap != nullptr)
 		{
-			const FString TexName = FString::Printf(TEXT("%s_%s"), *BaseName, *RoughnessTexParamName.ToString());
+			const FString TexName = FString::Printf(TEXT("%s_%s"), *BaseName, *RoughnessTexParamName);
 			CreateTextureAsset(TexName, FTexture2DBuilder::ETextureType::Roughness, ResultSettings->RoughnessMap);
 		}
 
-		if (Settings->MaterialSettings.bSpecularMap && ResultSettings->SpecularMap != nullptr)
+		if (RenderCaptureProperties->bSpecularMap && ResultSettings->SpecularMap != nullptr)
 		{
-			const FString TexName = FString::Printf(TEXT("%s_%s"), *BaseName, *SpecularTexParamName.ToString());
+			const FString TexName = FString::Printf(TEXT("%s_%s"), *BaseName, *SpecularTexParamName);
 			CreateTextureAsset(TexName, FTexture2DBuilder::ETextureType::Specular, ResultSettings->SpecularMap);
 		}
 	}
@@ -839,15 +854,15 @@ bool UBakeRenderCaptureTool::CanAccept() const
 	{
 		return false;
 	}
-	if (Settings->MaterialSettings.bNormalMap && ResultSettings->NormalMap == nullptr)
+	if (RenderCaptureProperties->bNormalMap && ResultSettings->NormalMap == nullptr)
 	{
 		return false;
 	}
-	if (Settings->MaterialSettings.bEmissiveMap && ResultSettings->EmissiveMap == nullptr)
+	if (RenderCaptureProperties->bEmissiveMap && ResultSettings->EmissiveMap == nullptr)
 	{
 		return false;
 	}
-	if (Settings->MaterialSettings.bPackedMRSMap)
+	if (RenderCaptureProperties->bPackedMRSMap)
 	{
 		if (ResultSettings->PackedMRSMap == nullptr)
 		{
@@ -856,15 +871,15 @@ bool UBakeRenderCaptureTool::CanAccept() const
 	}
 	else
 	{
-		if (Settings->MaterialSettings.bMetallicMap && ResultSettings->MetallicMap == nullptr)
+		if (RenderCaptureProperties->bMetallicMap && ResultSettings->MetallicMap == nullptr)
 		{
 			return false;
 		}
-		if (Settings->MaterialSettings.bRoughnessMap && ResultSettings->RoughnessMap == nullptr)
+		if (RenderCaptureProperties->bRoughnessMap && ResultSettings->RoughnessMap == nullptr)
 		{
 			return false;
 		}
-		if (Settings->MaterialSettings.bSpecularMap && ResultSettings->SpecularMap == nullptr)
+		if (RenderCaptureProperties->bSpecularMap && ResultSettings->SpecularMap == nullptr)
 		{
 			return false;
 		}
@@ -880,15 +895,17 @@ TUniquePtr<UE::Geometry::TGenericDataOperator<FBakeRenderCaptureResultsBuilder>>
 	Op->Actors = &Tool->Actors;
 	Op->BaseMesh = &Tool->TargetMesh;
 	Op->BaseMeshTangents = Tool->TargetMeshTangents;
-	Op->Options = FBakeRenderCaptureOptions::ConstructOptions(*Tool->Settings, *Tool->InputMeshSettings);
-	Op->SceneCapture = MoveTemp(Tool->SceneCapture);
+	Op->Options = FRenderCaptureSettings::ConstructOptions(*Tool->RenderCaptureProperties, *Tool->InputMeshSettings);
+	Op->TextureImageSize = static_cast<int32>(Tool->Settings->TextureSize);
+	Op->SamplesPerPixel = Tool->Settings->SamplesPerPixel;
+	Op->SceneCapture = Tool->SceneCapture.Get();
 	return Op;
 }
 
 
 void UBakeRenderCaptureTool::OnMapsUpdatedRC(const TUniquePtr<FBakeRenderCaptureResultsBuilder>& NewResult)
 {
-	FBakeRenderCaptureOptions::FOptions Options = FBakeRenderCaptureOptions::ConstructOptions(*Settings, *InputMeshSettings);
+	FRenderCaptureSettings::FOptions Options = FRenderCaptureSettings::ConstructOptions(*RenderCaptureProperties, *InputMeshSettings);
 
 	check(IsInGameThread());
 
@@ -1035,52 +1052,61 @@ void UBakeRenderCaptureTool::UpdateResult()
 
 	//
 	// create a set of spatially located render captures of the scene ("photo set").
-	// @Refactor SceneCapture currently must happen on the game thread, we should try to fix that
 	//
-	FBakeRenderCaptureOptions::FOptions Options = FBakeRenderCaptureOptions::ConstructOptions(*Settings, *InputMeshSettings);
-
-	for (int Idx = 1; Idx < Targets.Num(); ++Idx)
+	if (*RenderCaptureProperties != *ComputedRenderCaptureProperties)
 	{
-		UE::ToolTarget::ShowSourceObject(Targets[Idx]);
-	}
-
-	// Do not allow user-cancellation on the call that occurs when the Render Capture Tool starts up
-	const bool bAllowCancel = (bFirstEverSceneCapture == false);
+		for (int Idx = 1; Idx < Targets.Num(); ++Idx)
+		{
+			UE::ToolTarget::ShowSourceObject(Targets[Idx]);
+		}
 
-	SceneCapture.Reset();
-	SceneCapture = CapturePhotoSet(Actors, Options, bAllowCancel);
+		// Do not allow user-cancellation on the call that occurs when the Render Capture Tool starts up
+		const bool bAllowCancel = (bFirstEverSceneCapture == false);
 
-	for (int Idx = 1; Idx < Targets.Num(); ++Idx)
-	{
-		UE::ToolTarget::HideSourceObject(Targets[Idx]);
-	}
+		SceneCapture.Reset();
+		FRenderCaptureSettings::FOptions Options = FRenderCaptureSettings::ConstructOptions(*RenderCaptureProperties, *InputMeshSettings);
+		SceneCapture = CapturePhotoSet(Actors, Options, bAllowCancel);
 
-	if (SceneCapture->Cancelled())
-	{
-		// Restore the settings present before the change that invoked the scene capture recompute
-		Settings->SamplesPerPixel = ComputedSettings->SamplesPerPixel;
-		Settings->RenderCaptureResolution = ComputedSettings->RenderCaptureResolution;
-		Settings->MaterialSettings = ComputedSettings->MaterialSettings;
-		Settings->CaptureFieldOfView = ComputedSettings->CaptureFieldOfView;
-		Settings->NearPlaneDist = ComputedSettings->NearPlaneDist;
+		for (int Idx = 1; Idx < Targets.Num(); ++Idx)
+		{
+			UE::ToolTarget::HideSourceObject(Targets[Idx]);
+		}
 
-		// Silently make the above Settings updates so we don't overwrite the change to OpState below and call this function again
-		Settings->SilentUpdateWatched();
+		if (SceneCapture->Cancelled())
+		{
+			// Restore the settings present before the change that invoked the scene capture recompute
+			RenderCaptureProperties->Resolution         = ComputedRenderCaptureProperties->Resolution;
+			RenderCaptureProperties->bNormalMap         = ComputedRenderCaptureProperties->bNormalMap;
+			RenderCaptureProperties->bMetallicMap       = ComputedRenderCaptureProperties->bMetallicMap;
+			RenderCaptureProperties->bRoughnessMap      = ComputedRenderCaptureProperties->bRoughnessMap;
+			RenderCaptureProperties->bSpecularMap       = ComputedRenderCaptureProperties->bSpecularMap;
+			RenderCaptureProperties->bPackedMRSMap      = ComputedRenderCaptureProperties->bPackedMRSMap;
+			RenderCaptureProperties->bEmissiveMap       = ComputedRenderCaptureProperties->bEmissiveMap;
+			RenderCaptureProperties->CaptureFieldOfView = ComputedRenderCaptureProperties->CaptureFieldOfView;
+			RenderCaptureProperties->NearPlaneDist      = ComputedRenderCaptureProperties->NearPlaneDist;
+
+			// Silently make the above updates so we don't overwrite the change to OpState below and call this function again
+			RenderCaptureProperties->SilentUpdateWatched();
+
+			OpState = EBakeOpState::Clean;
 
-		OpState = EBakeOpState::Clean;
+			return;
+		}
 
-		return;
+		// Cache Settings used to compute this SceneCapture so we can restore them if the user cancels a SceneCapture recompute
+		ComputedRenderCaptureProperties->Resolution         = RenderCaptureProperties->Resolution;
+		ComputedRenderCaptureProperties->bNormalMap         = RenderCaptureProperties->bNormalMap;
+		ComputedRenderCaptureProperties->bMetallicMap       = RenderCaptureProperties->bMetallicMap;
+		ComputedRenderCaptureProperties->bRoughnessMap      = RenderCaptureProperties->bRoughnessMap;
+		ComputedRenderCaptureProperties->bSpecularMap       = RenderCaptureProperties->bSpecularMap;
+		ComputedRenderCaptureProperties->bPackedMRSMap      = RenderCaptureProperties->bPackedMRSMap;
+		ComputedRenderCaptureProperties->bEmissiveMap       = RenderCaptureProperties->bEmissiveMap;
+		ComputedRenderCaptureProperties->CaptureFieldOfView = RenderCaptureProperties->CaptureFieldOfView;
+		ComputedRenderCaptureProperties->NearPlaneDist      = RenderCaptureProperties->NearPlaneDist;
+
+		bFirstEverSceneCapture = false;
 	}
 
-	bFirstEverSceneCapture = false;
-
-	// Cache Settings used to compute this SceneCapture so we can restore them if the user cancels a SceneCapture recompute
-	ComputedSettings->SamplesPerPixel = Settings->SamplesPerPixel;
-	ComputedSettings->RenderCaptureResolution = Settings->RenderCaptureResolution;
-	ComputedSettings->MaterialSettings = Settings->MaterialSettings;
-	ComputedSettings->CaptureFieldOfView = Settings->CaptureFieldOfView;
-	ComputedSettings->NearPlaneDist = Settings->NearPlaneDist;
-
 	// @Incomplete Clear our invalid bitflag to check for valid inputs.
 	//OpState &= ~EBakeOpState::Invalid;
 	//OpState |= CheckValidInputs();
@@ -1104,7 +1130,7 @@ void UBakeRenderCaptureTool::UpdateVisualization()
 		return;
 	}
 
-	FBakeRenderCaptureOptions::FOptions Options = FBakeRenderCaptureOptions::ConstructOptions(*Settings, *InputMeshSettings);
+	FRenderCaptureSettings::FOptions Options = FRenderCaptureSettings::ConstructOptions(*RenderCaptureProperties, *InputMeshSettings);
 	
 	if (Options.bUsePackedMRS)
 	{
@@ -1114,35 +1140,35 @@ void UBakeRenderCaptureTool::UpdateVisualization()
 		if (VisualizationProps->bPreviewAsMaterial)
 		{
 			// We set all textures which were computed in the corresponding texture channels
-			Material->SetTextureParameterValue(TEXT("BaseColor"), Options.bBakeBaseColor ? ResultSettings->BaseColorMap : EmptyColorMapWhite);
-			Material->SetTextureParameterValue(TEXT("Emissive"),  Options.bBakeEmissive  ? ResultSettings->EmissiveMap  : EmptyEmissiveMap);
-			Material->SetTextureParameterValue(TEXT("NormalMap"), Options.bBakeNormalMap ? ResultSettings->NormalMap    : EmptyNormalMap);
-			Material->SetTextureParameterValue(TEXT("PackedMRS"), ResultSettings->PackedMRSMap);
+			Material->SetTextureParameterValue(FName(BaseColorTexParamName), Options.bBakeBaseColor ? ResultSettings->BaseColorMap : EmptyColorMapWhite);
+			Material->SetTextureParameterValue(FName(EmissiveTexParamName),  Options.bBakeEmissive  ? ResultSettings->EmissiveMap  : EmptyEmissiveMap);
+			Material->SetTextureParameterValue(FName(NormalTexParamName), Options.bBakeNormalMap ? ResultSettings->NormalMap    : EmptyNormalMap);
+			Material->SetTextureParameterValue(FName(PackedMRSTexParamName), ResultSettings->PackedMRSMap);
 		}
 		else
 		{
 			// The BaseColor texture channel will be set according to the selected MapPreview
 			TObjectPtr<UTexture2D> BaseColorMap = EmptyColorMapWhite;
-			if (Options.bBakeBaseColor && Settings->MapPreview == BaseColorTexParamName.ToString())
+			if (Options.bBakeBaseColor && Settings->MapPreview == BaseColorTexParamName)
 			{
 				BaseColorMap = ResultSettings->BaseColorMap;
 			}
-			else if (Options.bBakeEmissive && Settings->MapPreview == EmissiveTexParamName.ToString())
+			else if (Options.bBakeEmissive && Settings->MapPreview == EmissiveTexParamName)
 			{
 				BaseColorMap = ResultSettings->EmissiveMap;
 			}
-			else if (Options.bBakeNormalMap && Settings->MapPreview == NormalTexParamName.ToString())
+			else if (Options.bBakeNormalMap && Settings->MapPreview == NormalTexParamName)
 			{
 				BaseColorMap = ResultSettings->NormalMap;
 			}
-			else if (Settings->MapPreview == PackedMRSTexParamName.ToString())
+			else if (Settings->MapPreview == PackedMRSTexParamName)
 			{
 				BaseColorMap = ResultSettings->PackedMRSMap;
 			}
-			Material->SetTextureParameterValue(TEXT("BaseColor"), BaseColorMap);
-			Material->SetTextureParameterValue(TEXT("Emissive"),  EmptyEmissiveMap);
-			Material->SetTextureParameterValue(TEXT("NormalMap"), EmptyNormalMap);
-			Material->SetTextureParameterValue(TEXT("PackedMRS"), EmptyPackedMRSMap);
+			Material->SetTextureParameterValue(FName(BaseColorTexParamName), BaseColorMap);
+			Material->SetTextureParameterValue(FName(EmissiveTexParamName),  EmptyEmissiveMap);
+			Material->SetTextureParameterValue(FName(NormalTexParamName), EmptyNormalMap);
+			Material->SetTextureParameterValue(FName(PackedMRSTexParamName), EmptyPackedMRSMap);
 		}
 
 		Material->SetScalarParameterValue(TEXT("UVChannel"), Options.TargetUVLayer);
@@ -1155,38 +1181,38 @@ void UBakeRenderCaptureTool::UpdateVisualization()
 		if (VisualizationProps->bPreviewAsMaterial)
 		{
 			// We set all textures which were computed in the corresponding texture channels
-			Material->SetTextureParameterValue(TEXT("BaseColor"), Options.bBakeBaseColor ? ResultSettings->BaseColorMap : EmptyColorMapWhite);
-			Material->SetTextureParameterValue(TEXT("Roughness"), Options.bBakeRoughness ? ResultSettings->RoughnessMap : EmptyRoughnessMap);
-			Material->SetTextureParameterValue(TEXT("Metallic"),  Options.bBakeMetallic  ? ResultSettings->MetallicMap  : EmptyMetallicMap);
-			Material->SetTextureParameterValue(TEXT("Specular"),  Options.bBakeSpecular  ? ResultSettings->SpecularMap  : EmptySpecularMap);
-			Material->SetTextureParameterValue(TEXT("Emissive"),  Options.bBakeEmissive  ? ResultSettings->EmissiveMap  : EmptyEmissiveMap);
-			Material->SetTextureParameterValue(TEXT("NormalMap"), Options.bBakeNormalMap ? ResultSettings->NormalMap    : EmptyNormalMap);
+			Material->SetTextureParameterValue(FName(BaseColorTexParamName), Options.bBakeBaseColor ? ResultSettings->BaseColorMap : EmptyColorMapWhite);
+			Material->SetTextureParameterValue(FName(RoughnessTexParamName), Options.bBakeRoughness ? ResultSettings->RoughnessMap : EmptyRoughnessMap);
+			Material->SetTextureParameterValue(FName(MetallicTexParamName),  Options.bBakeMetallic  ? ResultSettings->MetallicMap  : EmptyMetallicMap);
+			Material->SetTextureParameterValue(FName(SpecularTexParamName),  Options.bBakeSpecular  ? ResultSettings->SpecularMap  : EmptySpecularMap);
+			Material->SetTextureParameterValue(FName(EmissiveTexParamName),  Options.bBakeEmissive  ? ResultSettings->EmissiveMap  : EmptyEmissiveMap);
+			Material->SetTextureParameterValue(FName(NormalTexParamName), Options.bBakeNormalMap ? ResultSettings->NormalMap    : EmptyNormalMap);
 		}
 		else
 		{
 			// The BaseColor texture channel will be set according to the selected MapPreview
 			TObjectPtr<UTexture2D> BaseColorMap = EmptyColorMapWhite;
-			if (Options.bBakeBaseColor && Settings->MapPreview == BaseColorTexParamName.ToString())
+			if (Options.bBakeBaseColor && Settings->MapPreview == BaseColorTexParamName)
 			{
 				BaseColorMap = ResultSettings->BaseColorMap;
 			}
-			else if (Options.bBakeRoughness && Settings->MapPreview == RoughnessTexParamName.ToString())
+			else if (Options.bBakeRoughness && Settings->MapPreview == RoughnessTexParamName)
 			{
 				BaseColorMap = ResultSettings->RoughnessMap;
 			}
-			else if (Options.bBakeMetallic && Settings->MapPreview == MetallicTexParamName.ToString())
+			else if (Options.bBakeMetallic && Settings->MapPreview == MetallicTexParamName)
 			{
 				BaseColorMap = ResultSettings->MetallicMap;
 			}
-			else if (Options.bBakeSpecular && Settings->MapPreview == SpecularTexParamName.ToString())
+			else if (Options.bBakeSpecular && Settings->MapPreview == SpecularTexParamName)
 			{
 				BaseColorMap = ResultSettings->SpecularMap;
 			}
-			else if (Options.bBakeEmissive && Settings->MapPreview == EmissiveTexParamName.ToString())
+			else if (Options.bBakeEmissive && Settings->MapPreview == EmissiveTexParamName)
 			{
 				BaseColorMap = ResultSettings->EmissiveMap;
 			}
-			else if (Options.bBakeNormalMap && Settings->MapPreview == NormalTexParamName.ToString())
+			else if (Options.bBakeNormalMap && Settings->MapPreview == NormalTexParamName)
 			{
 				BaseColorMap = ResultSettings->NormalMap;
 			}
@@ -1221,4 +1247,4 @@ void UBakeRenderCaptureTool::InvalidateResults()
 	ResultSettings->NormalMap = nullptr;
 }
 
-#undef LOCTEXT_NAMESPACE
+#undef LOCTEXT_NAMESPACE
\ No newline at end of file
diff --git a/Engine/Plugins/Experimental/MeshModelingToolsetExp/Source/MeshModelingToolsExp/Public/BakeRenderCaptureTool.h b/Engine/Plugins/Experimental/MeshModelingToolsetExp/Source/MeshModelingToolsExp/Public/BakeRenderCaptureTool.h
--- a/Engine/Plugins/Experimental/MeshModelingToolsetExp/Source/MeshModelingToolsExp/Public/BakeRenderCaptureTool.h
+++ b/Engine/Plugins/Experimental/MeshModelingToolsetExp/Source/MeshModelingToolsExp/Public/BakeRenderCaptureTool.h
@@ -87,63 +87,59 @@ protected:
 //
 
 
-USTRUCT()
-struct FMaterialProxySettingsRC
+UCLASS()
+class MESHMODELINGTOOLSEXP_API URenderCaptureProperties : public UInteractiveToolPropertySet
 {
-	GENERATED_USTRUCT_BODY()
+	GENERATED_BODY()
 
-	// Size of generated BaseColor map
-	UPROPERTY(Category = Material, EditAnywhere, meta=(ClampMin="1", UIMin="1"))
-	int32 TextureSize;
+public:
+	UPROPERTY(Category = RenderCaptureOptions, EditAnywhere, DisplayName="Render Capture Resolution", meta = (ClampMin = "1", UIMin= "1"))
+	EBakeTextureResolution Resolution = EBakeTextureResolution::Resolution512;
 
 	// Whether to generate a texture for the World Normal property
-	UPROPERTY(Category = Material, EditAnywhere)
-	uint8 bNormalMap:1;
+	UPROPERTY(Category = RenderCaptureOptions, EditAnywhere)
+	bool bNormalMap = true;
 
 	// Whether to generate a texture for the Metallic property
-	UPROPERTY(Category = Material, EditAnywhere, meta=(EditCondition="bPackedMRSMap == false"))
-	uint8 bMetallicMap:1;
+	UPROPERTY(Category = RenderCaptureOptions, EditAnywhere, meta=(EditCondition="bPackedMRSMap == false"))
+	bool bMetallicMap = true;
 
 	// Whether to generate a texture for the Roughness property
-	UPROPERTY(Category = Material, EditAnywhere, meta=(EditCondition="bPackedMRSMap == false"))
-	uint8 bRoughnessMap:1;
+	UPROPERTY(Category = RenderCaptureOptions, EditAnywhere, meta=(EditCondition="bPackedMRSMap == false"))
+	bool bRoughnessMap = true;
 
 	// Whether to generate a texture for the Specular property
-	UPROPERTY(Category = Material, EditAnywhere, meta=(EditCondition="bPackedMRSMap == false"))
-	uint8 bSpecularMap:1;
+	UPROPERTY(Category = RenderCaptureOptions, EditAnywhere, meta=(EditCondition="bPackedMRSMap == false"))
+	bool bSpecularMap = true;
 	
 	// Whether to generate a packed texture with Metallic, Roughness and Specular properties
-	UPROPERTY(Category = Material, EditAnywhere)
-	uint8 bPackedMRSMap:1;
+	UPROPERTY(Category = RenderCaptureOptions, EditAnywhere, DisplayName="Packed MRS Map")
+	bool bPackedMRSMap = true;
 
 	// Whether to generate a texture for the Emissive property
-	UPROPERTY(Category = Material, EditAnywhere)
-	uint8 bEmissiveMap:1;
-
-	FMaterialProxySettingsRC()
-		: TextureSize(512)
-		, bNormalMap(true)
-		, bMetallicMap(true)
-		, bRoughnessMap(true)
-		, bSpecularMap(true)
-		, bPackedMRSMap(true)
-		, bEmissiveMap(true)
-	{
-	}
+	UPROPERTY(Category = RenderCaptureOptions, EditAnywhere)
+	bool bEmissiveMap = true;
 
-	bool operator == (const FMaterialProxySettingsRC& Other) const
+	UPROPERTY(Category = RenderCaptureOptions, EditAnywhere, AdvancedDisplay, DisplayName="Field of View", meta = (ClampMin = "5.0", ClampMax = "160.0"))
+	float CaptureFieldOfView = 30.0f;
+
+	UPROPERTY(Category = RenderCaptureOptions, EditAnywhere, AdvancedDisplay, meta = (ClampMin = "0.001", ClampMax = "1000.0"))
+	float NearPlaneDist = 1.0f;
+
+	bool operator==(const URenderCaptureProperties& Other) const
 	{
-		// @Incomplete use the uproperty compare thing
-		return TextureSize == Other.TextureSize
+		return Resolution == Other.Resolution
 			&& bNormalMap == Other.bNormalMap
 			&& bMetallicMap == Other.bMetallicMap
 			&& bRoughnessMap == Other.bRoughnessMap
 			&& bSpecularMap == Other.bSpecularMap
 			&& bPackedMRSMap == Other.bPackedMRSMap
-			&& bEmissiveMap == Other.bEmissiveMap;
+			&& bEmissiveMap == Other.bEmissiveMap
+			&& CaptureFieldOfView == Other.CaptureFieldOfView
+			&& NearPlaneDist == Other.NearPlaneDist;
 	}
 
-	bool operator != (const FMaterialProxySettingsRC& Other) const
+	bool operator != (const URenderCaptureProperties& Other) const
 	{
 		return !(*this == Other);
 	}
@@ -168,27 +164,13 @@ public:
 	UPROPERTY(meta = (TransientToolProperty))
 	TArray<FString> MapPreviewNamesList;
 
-	//
-	// Material Baking Settings
-	//
-
 	/** Number of samples per pixel */
-	UPROPERTY(EditAnywhere, Category = MaterialSettings)
+	UPROPERTY(EditAnywhere, Category = BakeOutput)
 	EBakeTextureSamplesPerPixel SamplesPerPixel = EBakeTextureSamplesPerPixel::Sample1;
 
-	/** If Value is zero, use MaterialSettings resolution, otherwise override the render capture resolution */
-	UPROPERTY(EditAnywhere, Category = MaterialSettings, meta = (ClampMin = "0"))
-	int32 RenderCaptureResolution = 512;
-
-	/** Material generation settings */
-	UPROPERTY(EditAnywhere, Category = MaterialSettings)
-	FMaterialProxySettingsRC MaterialSettings; // TODO make a separate struct and make it work like the baking tool (pick texture widgets etc)
-
-	UPROPERTY(EditAnywhere, AdvancedDisplay, Category = MaterialSettings, meta = (ClampMin = "5.0", ClampMax = "160.0"))
-	float CaptureFieldOfView = 30.0f;
-
-	UPROPERTY(EditAnywhere, AdvancedDisplay, Category = MaterialSettings, meta = (ClampMin = "0.001", ClampMax = "1000.0"))
-	float NearPlaneDist = 1.0f;
+	/* Size of generated textures */
+	UPROPERTY(EditAnywhere, Category = BakeOutput, DisplayName="Results Resolution", meta=(ClampMin="1", UIMin="1"))
+	EBakeTextureResolution TextureSize = EBakeTextureResolution::Resolution512;
 };
 
 
@@ -250,6 +232,9 @@ protected:
 	UPROPERTY()
 	TObjectPtr<UBakeRenderCaptureToolProperties> Settings;
 
+	UPROPERTY()
+	TObjectPtr<URenderCaptureProperties> RenderCaptureProperties;
+
 	UPROPERTY()
 	TObjectPtr<UBakeRenderCaptureInputToolProperties> InputMeshSettings;
 
@@ -326,11 +311,12 @@ protected:
 
 	// If the user cancels a scene capture before the computation completes then the settings which changed to invoke
 	// the capture are reverted to these values
-	TObjectPtr<UBakeRenderCaptureToolProperties> ComputedSettings;
+	UPROPERTY()
+	TObjectPtr<URenderCaptureProperties> ComputedRenderCaptureProperties;
 
 	// Analytics
 	virtual FString GetAnalyticsEventName() const override
 	{
 		return TEXT("BakeRC");
 	}
-};
+};
\ No newline at end of file
