Render Capture Baking: Refactor render capture evaluator, to better support analytics and address some tech debt

#rnx

diff --git a/Engine/Plugins/Experimental/MeshModelingToolsetExp/Source/MeshModelingToolsExp/Private/BakeRenderCaptureTool.cpp b/Engine/Plugins/Experimental/MeshModelingToolsetExp/Source/MeshModelingToolsExp/Private/BakeRenderCaptureTool.cpp
--- a/Engine/Plugins/Experimental/MeshModelingToolsetExp/Source/MeshModelingToolsExp/Private/BakeRenderCaptureTool.cpp
+++ b/Engine/Plugins/Experimental/MeshModelingToolsetExp/Source/MeshModelingToolsExp/Private/BakeRenderCaptureTool.cpp
@@ -16,7 +16,7 @@
 
 #include "Sampling/MeshImageBakingCache.h"
 #include "Sampling/MeshMapBaker.h"
-#include "Sampling/GenericEvaluator.h"
+#include "Sampling/RenderCaptureMapEvaluator.h"
 #include "Image/ImageInfilling.h"
 #include "Algo/NoneOf.h"
 #include "Misc/ScopedSlowTask.h"
@@ -223,10 +223,34 @@ static TUniquePtr<FSceneCapturePhotoSet> CapturePhotoSet(
 }
 
 template <ERenderCaptureType CaptureType>
-TSharedPtr<FGenericEvaluator<FVector4f>>
+TSharedPtr<FRenderCaptureMapEvaluator<FVector4f>>
 MakeColorEvaluator(const FSceneCapturePhotoSet::FSceneSample& DefaultSample, const FSceneCapturePhotoSet* SceneCapture)
 {
-	TSharedPtr<FGenericEvaluator<FVector4f>> Evaluator = MakeShared<FGenericEvaluator<FVector4f>>();
+	TSharedPtr<FRenderCaptureMapEvaluator<FVector4f>> Evaluator = MakeShared<FRenderCaptureMapEvaluator<FVector4f>>();
+
+	switch (CaptureType) {
+	case ERenderCaptureType::BaseColor:
+		Evaluator->Channel = ERenderCaptureChannel::BaseColor;
+		break;
+	case ERenderCaptureType::Roughness:
+		Evaluator->Channel = ERenderCaptureChannel::Roughness;
+		break;
+	case ERenderCaptureType::Metallic:
+		Evaluator->Channel = ERenderCaptureChannel::Metallic;
+		break;
+	case ERenderCaptureType::Specular:
+		Evaluator->Channel = ERenderCaptureChannel::Specular;
+		break;
+	case ERenderCaptureType::Emissive:
+		Evaluator->Channel = ERenderCaptureChannel::Emissive;
+		break;
+	case ERenderCaptureType::WorldNormal:
+		Evaluator->Channel = ERenderCaptureChannel::WorldNormal;
+		break;
+	case ERenderCaptureType::CombinedMRS:
+		Evaluator->Channel = ERenderCaptureChannel::CombinedMRS;
+		break;
+	}
 
 	Evaluator->DefaultResult = DefaultSample.GetValue4f(CaptureType);
 
@@ -254,8 +278,8 @@ static void ImageBuildersFromPhotoSet(
 	const int32 TextureImageSize,
 	const EBakeTextureSamplesPerPixel SamplesPerPixel,
 	const FDynamicMesh3* BaseMesh,
-	const TSharedPtr<UE::Geometry::FMeshTangentsd, ESPMode::ThreadSafe>& BaseMeshTangents,
-	TUniquePtr<FBakeRenderCaptureResultsBuilder>& Results)
+	const TSharedPtr<FMeshTangentsd, ESPMode::ThreadSafe>& BaseMeshTangents,
+	TUniquePtr<FMeshMapBaker>& Result)
 {
 	TRACE_CPUPROFILER_EVENT_SCOPE(ImageBuildersFromPhotoSet);
 
@@ -377,21 +401,18 @@ static void ImageBuildersFromPhotoSet(
 		}
 	};
 
-	FMeshMapBaker Baker;
-	Baker.SetTargetMesh(BaseMesh);
-	Baker.SetTargetMeshTangents(BaseMeshTangents);
-	Baker.SetDimensions(FImageDimensions(TextureImageSize, TextureImageSize));
-	Baker.SetSamplesPerPixel(static_cast<int32>(SamplesPerPixel));
-	Baker.SetFilter(FMeshMapBaker::EBakeFilterType::BSpline);
-	Baker.SetTargetMeshUVLayer(Options.TargetUVLayer);
-	Baker.InteriorSampleCallback = RegisterSampleStats;
-	Baker.PostWriteToImageCallback = ComputeAndApplyInfill;
+	Result->SetTargetMesh(BaseMesh);
+	Result->SetTargetMeshTangents(BaseMeshTangents);
+	Result->SetDimensions(FImageDimensions(TextureImageSize, TextureImageSize));
+	Result->SetSamplesPerPixel(static_cast<int32>(SamplesPerPixel));
+	Result->SetFilter(FMeshMapBaker::EBakeFilterType::BSpline);
+	Result->SetTargetMeshUVLayer(Options.TargetUVLayer);
+	Result->InteriorSampleCallback = RegisterSampleStats;
+	Result->PostWriteToImageCallback = ComputeAndApplyInfill;
 
 	FSceneCapturePhotoSetSampler Sampler(SceneCapture, VisibilityFunction, BaseMesh, &BaseMeshSpatial, BaseMeshTangents.Get());
-	Baker.SetDetailSampler(&Sampler);
-	Baker.SetCorrespondenceStrategy(FMeshMapBaker::ECorrespondenceStrategy::Custom);
-
-	TMap<ERenderCaptureType, int32> CaptureTypeToEvaluatorIndexMap;
+	Result->SetDetailSampler(&Sampler);
+	Result->SetCorrespondenceStrategy(FMeshMapBaker::ECorrespondenceStrategy::Custom);
 
 	// Pixels in the output textures which don't map onto the mesh have a light grey color (except the normal map which
 	// will show a color corresponding to a unit z tangent space normal)
@@ -406,13 +427,12 @@ static void ImageBuildersFromPhotoSet(
 	DefaultColorSample.Emissive = FVector3f(InvalidColor.X, InvalidColor.Y, InvalidColor.Z);
 	DefaultColorSample.WorldNormal = FVector4f((DefaultNormal + FVector3f::One()) * .5f, InvalidColor.W);
 
-	auto AddColorEvaluator = [&Baker, &InfillData, &CaptureTypeToEvaluatorIndexMap] (
-		const TSharedPtr<FGenericEvaluator<FVector4f>>& Evaluator,
+	auto AddColorEvaluator = [&Result, &InfillData] (
+		const TSharedPtr<FRenderCaptureMapEvaluator<FVector4f>>& Evaluator,
 		ERenderCaptureType CaptureType)
 	{
-		int32 EvaluatorIndex = Baker.AddEvaluator(Evaluator);
+		Result->AddEvaluator(Evaluator);
 		InfillData.EvaluatorNeedsInfill.Add(true);
-		CaptureTypeToEvaluatorIndexMap.Emplace(CaptureType, EvaluatorIndex);
 	};
 
 	AddColorEvaluator(
@@ -455,8 +475,9 @@ static void ImageBuildersFromPhotoSet(
 	}
 
 	if (Options.bBakeNormalMap) {
-		TSharedPtr<FGenericEvaluator<FVector3f, FMeshMapEvaluator::EComponents::Float3>> Evaluator =
-			MakeShared<FGenericEvaluator<FVector3f, FMeshMapEvaluator::EComponents::Float3>>();
+		TSharedPtr<FRenderCaptureMapEvaluator<FVector3f>> Evaluator = MakeShared<FRenderCaptureMapEvaluator<FVector3f>>();
+
+		Evaluator->Channel = ERenderCaptureChannel::WorldNormal;
 
 		Evaluator->DefaultResult = DefaultNormal;
 
@@ -498,9 +519,7 @@ static void ImageBuildersFromPhotoSet(
 			return Out;
 		};
 
-		int32 EvaluatorIndex = Baker.AddEvaluator(Evaluator);
-
-		CaptureTypeToEvaluatorIndexMap.Emplace(ERenderCaptureType::WorldNormal, EvaluatorIndex);
+		Result->AddEvaluator(Evaluator);
 
 		// Note: No infill on normal map for now, doesn't make sense to do after mapping to tangent space!
 		//  (should we build baked normal map in world space, and then resample to tangent space??)
@@ -509,39 +528,7 @@ static void ImageBuildersFromPhotoSet(
 
 	{
 		TRACE_CPUPROFILER_EVENT_SCOPE(ImageBuildersFromPhotoSet_Bake);
-		Baker.Bake();
-	}
-
-	for (TTuple<ERenderCaptureType, int32>& Item : CaptureTypeToEvaluatorIndexMap)
-	{
-		if (Item.Key == ERenderCaptureType::BaseColor)
-		{
-			Results->ColorImage = MoveTemp(Baker.GetBakeResults(Item.Value)[0]);
-		}
-		else if (Item.Key == ERenderCaptureType::CombinedMRS)
-		{
-			Results->PackedMRSImage = MoveTemp(Baker.GetBakeResults(Item.Value)[0]);
-		}
-		else if (Item.Key == ERenderCaptureType::Roughness)
-		{
-			Results->RoughnessImage = MoveTemp(Baker.GetBakeResults(Item.Value)[0]);
-		}
-		else if (Item.Key == ERenderCaptureType::Specular)
-		{
-			Results->SpecularImage = MoveTemp(Baker.GetBakeResults(Item.Value)[0]);
-		}
-		else if (Item.Key == ERenderCaptureType::Metallic)
-		{
-			Results->MetallicImage = MoveTemp(Baker.GetBakeResults(Item.Value)[0]);
-		}
-		else if (Item.Key == ERenderCaptureType::Emissive)
-		{
-			Results->EmissiveImage = MoveTemp(Baker.GetBakeResults(Item.Value)[0]);
-		}
-		else if (Item.Key == ERenderCaptureType::WorldNormal)
-		{
-			Results->NormalImage = MoveTemp(Baker.GetBakeResults(Item.Value)[0]);
-		}
+		Result->Bake();
 	}
 }
 
@@ -549,11 +536,9 @@ static void ImageBuildersFromPhotoSet(
 // Tool Operators
 //
 
-
-class FRenderCaptureMapBakerOp : public TGenericDataOperator<FBakeRenderCaptureResultsBuilder>
+class FRenderCaptureMapBakerOp : public TGenericDataOperator<FMeshMapBaker>
 {
 public:
-	TArray<TObjectPtr<AActor>>* Actors;
 	UE::Geometry::FDynamicMesh3* BaseMesh = nullptr;
 	TSharedPtr<UE::Geometry::FMeshTangentsd, ESPMode::ThreadSafe> BaseMeshTangents;
 	FRenderCaptureSettings::FOptions Options;
@@ -561,12 +546,11 @@ public:
 	EBakeTextureSamplesPerPixel SamplesPerPixel;
 	TObjectPtr<UBakeRenderCaptureInputToolProperties> InputMeshSettings;
 	FSceneCapturePhotoSet* SceneCapture = nullptr;
-	
+
 	// Begin TGenericDataOperator interface
 	// @Incomplete Use the Progress thing
 	virtual void CalculateResult(FProgressCancel* Progress) override
 	{
-		check(Actors != nullptr);
 		check(BaseMesh != nullptr);
 		check(BaseMeshTangents.IsValid());
 		check(SceneCapture != nullptr);
@@ -728,9 +712,9 @@ void UBakeRenderCaptureTool::OnShutdown(EToolShutdownType ShutdownType)
 	RenderCaptureProperties->SaveProperties(this);
 	InputMeshSettings->SaveProperties(this);
 
-	if (ComputeRC)
+	if (Compute)
 	{
-		ComputeRC->Shutdown();
+		Compute->Shutdown();
 	}
 
 	// Restore visibility of source meshes
@@ -826,22 +810,6 @@ void UBakeRenderCaptureTool::CreateTextureAssetsRC(UWorld* SourceWorld, UObject*
 }
 
 
-void UBakeRenderCaptureTool::OnTick(float DeltaTime)
-{
-	if (ComputeRC)
-	{
-		ComputeRC->Tick(DeltaTime);
-
-		const float ElapsedComputeTime = ComputeRC->GetElapsedComputeTime();
-		if (!CanAccept() && ElapsedComputeTime > SecondsBeforeWorkingMaterial)
-		{
-			UMaterialInstanceDynamic* ProgressMaterial =
-				static_cast<bool>(OpState & EBakeOpState::Invalid) ? ErrorPreviewMaterial : WorkingPreviewMaterial;
-			PreviewMesh->SetOverrideRenderMaterial(ProgressMaterial);
-		}
-	}
-}
-
 
 bool UBakeRenderCaptureTool::CanAccept() const
 {
@@ -889,64 +857,129 @@ bool UBakeRenderCaptureTool::CanAccept() const
 }
 
 
-TUniquePtr<UE::Geometry::TGenericDataOperator<FBakeRenderCaptureResultsBuilder>> UBakeRenderCaptureTool::FComputeFactory::MakeNewOperator()
+
+TUniquePtr<TGenericDataOperator<FMeshMapBaker>> UBakeRenderCaptureTool::MakeNewOperator()
 {
 	TUniquePtr<FRenderCaptureMapBakerOp> Op = MakeUnique<FRenderCaptureMapBakerOp>();
-	Op->Actors = &Tool->Actors;
-	Op->BaseMesh = &Tool->TargetMesh;
-	Op->BaseMeshTangents = Tool->TargetMeshTangents;
-	Op->Options = FRenderCaptureSettings::ConstructOptions(*Tool->RenderCaptureProperties, *Tool->InputMeshSettings);
-	Op->TextureImageSize = static_cast<int32>(Tool->Settings->TextureSize);
-	Op->SamplesPerPixel = Tool->Settings->SamplesPerPixel;
-	Op->SceneCapture = Tool->SceneCapture.Get();
+	Op->BaseMesh = &TargetMesh;
+	Op->BaseMeshTangents = TargetMeshTangents;
+	Op->Options = FRenderCaptureSettings::ConstructOptions(*RenderCaptureProperties, *InputMeshSettings);
+	Op->TextureImageSize = static_cast<int32>(Settings->TextureSize);
+	Op->SamplesPerPixel = Settings->SamplesPerPixel;
+	Op->SceneCapture = SceneCapture.Get();
 	return Op;
 }
 
 
-void UBakeRenderCaptureTool::OnMapsUpdatedRC(const TUniquePtr<FBakeRenderCaptureResultsBuilder>& NewResult)
-{
-	FRenderCaptureSettings::FOptions Options = FRenderCaptureSettings::ConstructOptions(*RenderCaptureProperties, *InputMeshSettings);
-
-	check(IsInGameThread());
 
+void UBakeRenderCaptureTool::OnMapsUpdatedRC(const TUniquePtr<FMeshMapBaker>& NewResult)
+{
 	// We do this to defer work I guess, it was like this in the original ApproximateActors implementation :DeferredPopulateSourceData 
 	constexpr bool bPopulateSourceData = false;
 
+	TRACE_CPUPROFILER_EVENT_SCOPE(BakeRenderCaptureTool_Textures_BuildTextures);
+
+	const int32 NumEval = NewResult->NumEvaluators();
+	for (int32 EvalIdx = 0; EvalIdx < NumEval; ++EvalIdx)
 	{
-		TRACE_CPUPROFILER_EVENT_SCOPE(BakeRenderCaptureTool_Textures_BuildTextures);
+		FMeshMapEvaluator* BaseEval = NewResult->GetEvaluator(EvalIdx);
 
-		if (Options.bBakeBaseColor && NewResult->ColorImage.IsValid())
-		{
-			ResultSettings->BaseColorMap = FTexture2DBuilder::BuildTextureFromImage(*NewResult->ColorImage, FTexture2DBuilder::ETextureType::Color, true, bPopulateSourceData);
-		}
-		if (Options.bBakeEmissive && NewResult->EmissiveImage.IsValid())
-		{
-			ResultSettings->EmissiveMap = FTexture2DBuilder::BuildTextureFromImage(*NewResult->EmissiveImage, FTexture2DBuilder::ETextureType::EmissiveHDR, false, bPopulateSourceData);
-			ResultSettings->EmissiveMap->CompressionSettings = TC_HDR_Compressed;
-		}
-		if (Options.bBakeNormalMap && NewResult->NormalImage.IsValid())
-		{
-			ResultSettings->NormalMap = FTexture2DBuilder::BuildTextureFromImage(*NewResult->NormalImage, FTexture2DBuilder::ETextureType::NormalMap, false, bPopulateSourceData);
-		}
+		check(BaseEval->DataLayout().Num() == 1);
 
-		if (Options.bUsePackedMRS && NewResult->PackedMRSImage.IsValid())
+		switch (BaseEval->DataLayout()[0])
 		{
-			ResultSettings->PackedMRSMap = FTexture2DBuilder::BuildTextureFromImage(*NewResult->PackedMRSImage, FTexture2DBuilder::ETextureType::ColorLinear, false, bPopulateSourceData);
-		}
-		else
-		{ 
-			if (Options.bBakeRoughness && NewResult->RoughnessImage.IsValid())
-			{
-				ResultSettings->RoughnessMap = FTexture2DBuilder::BuildTextureFromImage(*NewResult->RoughnessImage, FTexture2DBuilder::ETextureType::Roughness, false, bPopulateSourceData);
-			}
-			if (Options.bBakeMetallic && NewResult->MetallicImage.IsValid())
+		case FMeshMapEvaluator::EComponents::Float4:
+
 			{
-				ResultSettings->MetallicMap = FTexture2DBuilder::BuildTextureFromImage(*NewResult->MetallicImage, FTexture2DBuilder::ETextureType::Metallic, false, bPopulateSourceData);
-			}
-			if (Options.bBakeSpecular && NewResult->SpecularImage.IsValid())
+				FRenderCaptureMapEvaluator<FVector4f>* Eval = static_cast<FRenderCaptureMapEvaluator<FVector4f>*>(BaseEval);
+				TUniquePtr<TImageBuilder<FVector4f>> ImageBuilder = MoveTemp(NewResult->GetBakeResults(EvalIdx)[0]);
+				
+				if (ensure(ImageBuilder.IsValid()) == false) return;
+
+				switch (Eval->Channel)
+				{
+				case ERenderCaptureChannel::BaseColor:
+
+					ResultSettings->BaseColorMap = FTexture2DBuilder::BuildTextureFromImage(
+						*ImageBuilder,
+						FTexture2DBuilder::ETextureType::Color,
+						true,
+						bPopulateSourceData);
+					break;
+
+				case ERenderCaptureChannel::Roughness:
+
+					ResultSettings->RoughnessMap = FTexture2DBuilder::BuildTextureFromImage(
+						*ImageBuilder,
+						FTexture2DBuilder::ETextureType::Roughness,
+						false,
+						bPopulateSourceData);
+					break;
+
+				case ERenderCaptureChannel::Metallic:
+
+					ResultSettings->MetallicMap = FTexture2DBuilder::BuildTextureFromImage(
+						*ImageBuilder,
+						FTexture2DBuilder::ETextureType::Metallic,
+						false,
+						bPopulateSourceData);
+					break;
+
+				case ERenderCaptureChannel::Specular:
+
+					ResultSettings->SpecularMap = FTexture2DBuilder::BuildTextureFromImage(
+						*ImageBuilder,
+						FTexture2DBuilder::ETextureType::Specular,
+						false,
+						bPopulateSourceData);
+					break;
+
+				case ERenderCaptureChannel::Emissive:
+
+					ResultSettings->EmissiveMap = FTexture2DBuilder::BuildTextureFromImage(
+						*ImageBuilder,
+						FTexture2DBuilder::ETextureType::EmissiveHDR,
+						false,
+						bPopulateSourceData);
+					ResultSettings->EmissiveMap->CompressionSettings = TC_HDR_Compressed;
+					break;
+
+				case ERenderCaptureChannel::CombinedMRS:
+
+					ResultSettings->PackedMRSMap = FTexture2DBuilder::BuildTextureFromImage(
+						*ImageBuilder,
+						FTexture2DBuilder::ETextureType::ColorLinear,
+						false,
+						bPopulateSourceData);
+					break;
+
+				default:
+					ensure(false);
+					return;
+				}
+
+			} break; // Float4
+
+		case FMeshMapEvaluator::EComponents::Float3:
+
 			{
-				ResultSettings->SpecularMap = FTexture2DBuilder::BuildTextureFromImage(*NewResult->SpecularImage, FTexture2DBuilder::ETextureType::Specular, false, bPopulateSourceData);
-			}
+				FRenderCaptureMapEvaluator<FVector3f>* Eval = static_cast<FRenderCaptureMapEvaluator<FVector3f>*>(BaseEval);
+				TUniquePtr<TImageBuilder<FVector4f>> ImageBuilder = MoveTemp(NewResult->GetBakeResults(EvalIdx)[0]);
+				
+				if (ensure(ImageBuilder.IsValid()) == false) return;
+				if (ensure(Eval->Channel == ERenderCaptureChannel::WorldNormal) == false) return;
+
+				ResultSettings->NormalMap = FTexture2DBuilder::BuildTextureFromImage(
+					*ImageBuilder,
+					FTexture2DBuilder::ETextureType::NormalMap,
+					false,
+					bPopulateSourceData);
+
+			} break; // Float3
+			
+		default:
+			ensure(false);
+			return;
 		}
 	}
 
@@ -1026,18 +1059,18 @@ void UBakeRenderCaptureTool::InitializePreviewMaterials()
 }
 
 
-
 void UBakeRenderCaptureTool::InvalidateComputeRC()
 {
-	if (!ComputeRC)
+	// Note: This implementation is identical to UBakeMeshAttributeMapsToolBase::InvalidateCompute but calls
+	// OnMapsUpdatedRC rather than OnMapsUpdated
+	if (!Compute)
 	{
 		// Initialize background compute
-		ComputeRC = MakeUnique<TGenericDataBackgroundCompute<FBakeRenderCaptureResultsBuilder>>();
-		ComputeFactory.Tool = this;
-		ComputeRC->Setup(&ComputeFactory);
-		ComputeRC->OnResultUpdated.AddLambda([this](const TUniquePtr<FBakeRenderCaptureResultsBuilder>& NewResult) { OnMapsUpdatedRC(NewResult); });
+		Compute = MakeUnique<TGenericDataBackgroundCompute<FMeshMapBaker>>();
+		Compute->Setup(this);
+		Compute->OnResultUpdated.AddLambda([this](const TUniquePtr<FMeshMapBaker>& NewResult) { OnMapsUpdatedRC(NewResult); });
 	}
-	ComputeRC->InvalidateResult();
+	Compute->InvalidateResult();
 }
 
 
@@ -1045,7 +1078,7 @@ void UBakeRenderCaptureTool::UpdateResult()
 {
 	if (OpState == EBakeOpState::Clean)
 	{
-		// Evaluation already launched/complete. Note that the ComputeRC background compute updates ResultSettings when
+		// Evaluation already launched/complete. Note that the Compute background compute updates ResultSettings when
 		// they are available by calling OnMapsUpdatedRC in its OnResultUpdated delegate.
 		return;
 	}
diff --git a/Engine/Plugins/Experimental/MeshModelingToolsetExp/Source/MeshModelingToolsExp/Public/BakeRenderCaptureTool.h b/Engine/Plugins/Experimental/MeshModelingToolsetExp/Source/MeshModelingToolsExp/Public/BakeRenderCaptureTool.h
--- a/Engine/Plugins/Experimental/MeshModelingToolsetExp/Source/MeshModelingToolsExp/Public/BakeRenderCaptureTool.h
+++ b/Engine/Plugins/Experimental/MeshModelingToolsetExp/Source/MeshModelingToolsExp/Public/BakeRenderCaptureTool.h
@@ -217,12 +217,14 @@ public:
 	// Begin UInteractiveTool interface
 	virtual void Setup() override;
 	virtual void OnShutdown(EToolShutdownType ShutdownType) override;
-	virtual void OnTick(float DeltaTime) override;
 	virtual bool HasCancel() const override { return true; }
 	virtual bool HasAccept() const override { return true; }
 	virtual bool CanAccept() const override;
 	// End UInteractiveTool interface
 
+	// Begin IGenericDataOperatorFactory interface
+	virtual TUniquePtr<UE::Geometry::TGenericDataOperator<UE::Geometry::FMeshMapBaker>> MakeNewOperator() override;
+	// End IGenericDataOperatorFactory interface
 
 protected:
 
@@ -249,6 +251,7 @@ protected:
 	virtual void UpdateResult() override;
 	virtual void UpdateVisualization() override;
 	virtual void GatherAnalytics(FBakeAnalytics::FMeshSettings& Data) override;
+	virtual FString GetAnalyticsEventName() const override { return TEXT("BakeRC"); }
 	// End UBakeMeshAttributeMapsToolBase interface
 	
 	void InvalidateResults();
@@ -263,23 +266,11 @@ protected:
 		// End IGenericDataOperatorFactory interface
 	};
 
-	// TODO :PortRenderCaptureToolToNewBakingFramework
-	// This tool has been ported to modeling mode after the new baking framework was implemented,
-	// since there is quite a lot of code to move the plan is to get the tool working in modeling
-	// mode using the old framework and then port it to the new one. This means that some aspects
-	// of the UBakeMeshAttributeMapsToolBase cannot yet be used, so we work around these and point
-	// out where we're going to need to do future @Refactor-ing. -- matija.kecman, 4 March 2022
-
-	// :PortRenderCaptureToolToNewBakingFramework
-	// When we use the new baking framework, we should be able to use:
-	// - UBakeMeshAttributeMapsToolBase::Compute
-	// - UBakeMeshAttributeMapsToolBase::InvalidateCompute
-	// - UBakeMeshAttributeMapsToolBase::OnMapsUpdated
-	// - UBakeMeshAttributeMapsToolBase::OnTick
- 	FComputeFactory ComputeFactory;
-	TUniquePtr<TGenericDataBackgroundCompute<FBakeRenderCaptureResultsBuilder>> ComputeRC = nullptr;
+	// In this tool we don't call UBakeMeshAttributeMapsToolBase::OnMapsUpdated because it would require e.g, adding
+	// the render capture channels to EBakeMapType.  The implementation is simpler and leads to less coupling if we
+	// just implement custom versions of the following functions.
 	void InvalidateComputeRC();
-	void OnMapsUpdatedRC(const TUniquePtr<FBakeRenderCaptureResultsBuilder>& NewResult);
+	void OnMapsUpdatedRC(const TUniquePtr<UE::Geometry::FMeshMapBaker>& NewResult);
 	
 	/**
 	 * Create texture assets from our result map of Texture2D
@@ -313,10 +304,4 @@ protected:
 	// the capture are reverted to these values
 	UPROPERTY()
 	TObjectPtr<URenderCaptureProperties> ComputedRenderCaptureProperties;
-
-	// Analytics
-	virtual FString GetAnalyticsEventName() const override
-	{
-		return TEXT("BakeRC");
-	}
 };
\ No newline at end of file
diff --git a/Engine/Plugins/Runtime/GeometryProcessing/Source/DynamicMesh/Public/Sampling/GenericEvaluator.h b/Engine/Plugins/Runtime/GeometryProcessing/Source/DynamicMesh/Public/Sampling/GenericEvaluator.h
deleted file mode 100644
diff --git a/Engine/Plugins/Runtime/GeometryProcessing/Source/DynamicMesh/Public/Sampling/MeshMapEvaluator.h b/Engine/Plugins/Runtime/GeometryProcessing/Source/DynamicMesh/Public/Sampling/MeshMapEvaluator.h
--- a/Engine/Plugins/Runtime/GeometryProcessing/Source/DynamicMesh/Public/Sampling/MeshMapEvaluator.h
+++ b/Engine/Plugins/Runtime/GeometryProcessing/Source/DynamicMesh/Public/Sampling/MeshMapEvaluator.h
@@ -20,7 +20,7 @@ enum class EMeshMapEvaluatorType
 	Property,
 	ResampleImage,
 	MultiResampleImage,
-	Generic
+	RenderCapture
 };
 
 class FMeshMapEvaluator
diff --git a/Engine/Plugins/Runtime/GeometryProcessing/Source/DynamicMesh/Public/Sampling/RenderCaptureMapEvaluator.h b/Engine/Plugins/Runtime/GeometryProcessing/Source/DynamicMesh/Public/Sampling/RenderCaptureMapEvaluator.h
new file mode 100644
--- /dev/null
+++ b/Engine/Plugins/Runtime/GeometryProcessing/Source/DynamicMesh/Public/Sampling/RenderCaptureMapEvaluator.h
@@ -0,0 +1,84 @@
+﻿// Copyright Epic Games, Inc. All Rights Reserved.
+
+#pragma once
+
+#include "Sampling/MeshMapEvaluator.h"
+
+namespace UE
+{
+namespace Geometry
+{
+
+// TODO Consolidate this with ERenderCaptureType
+enum class ERenderCaptureChannel
+{
+	BaseColor = 1,
+	Roughness = 2,
+	Metallic = 4,
+	Specular = 8,
+	Emissive = 16,
+	WorldNormal = 32,
+	CombinedMRS = 128
+};
+
+template <typename DataType>
+class FRenderCaptureMapEvaluator : public FMeshMapEvaluator
+{
+public:
+
+	ERenderCaptureChannel Channel = ERenderCaptureChannel::BaseColor;
+
+	// Analytics
+	// TODO Add some analytics to collect
+
+	DataType DefaultResult;
+	TFunction<DataType(const FCorrespondenceSample& Sample)> EvaluateSampleCallback;
+	TFunction<FVector4f(const int DataIdx, float*& In)> EvaluateColorCallback;
+
+public:
+	virtual ~FRenderCaptureMapEvaluator() {}
+
+	// Begin FMeshMapEvaluator interface
+	virtual void Setup(const FMeshBaseBaker& Baker, FEvaluationContext& Context) override
+	{
+		Context.Evaluate = &EvaluateSample;
+		Context.EvaluateDefault = &EvaluateDefault;
+		Context.EvaluateColor = &EvaluateColor;
+		Context.EvalData = this;
+		Context.AccumulateMode = EAccumulateMode::Add;
+		Context.DataLayout = DataLayout();
+	}
+
+	virtual const TArray<EComponents>& DataLayout() const override {
+		static const TArray<EComponents> Layout{(Channel == ERenderCaptureChannel::WorldNormal
+			? EComponents::Float3
+			: EComponents::Float4) };
+		return Layout;
+	}
+
+	virtual EMeshMapEvaluatorType Type() const override { return EMeshMapEvaluatorType::RenderCapture; }
+	// End FMeshMapEvaluator interface
+
+	static void EvaluateSample(float*& Out, const FCorrespondenceSample& Sample, void* EvalData)
+	{
+		const FRenderCaptureMapEvaluator* Eval = static_cast<FRenderCaptureMapEvaluator*>(EvalData);
+		const DataType SampleResult = Eval->EvaluateSampleCallback(Sample);
+		WriteToBuffer(Out, SampleResult);
+	}
+
+	static void EvaluateDefault(float*& Out, void* EvalData)
+	{
+		const FRenderCaptureMapEvaluator* Eval = static_cast<FRenderCaptureMapEvaluator*>(EvalData);
+		WriteToBuffer(Out, Eval->DefaultResult);
+	}
+
+	static void EvaluateColor(const int DataIdx, float*& In, FVector4f& Out, void* EvalData)
+	{
+		const FRenderCaptureMapEvaluator* Eval = static_cast<FRenderCaptureMapEvaluator*>(EvalData);
+		Out = Eval->EvaluateColorCallback(DataIdx, In);
+	}
+};
+
+
+} // end namespace UE::Geometry
+} // end namespace UE
