Render Capture Baking: Add analytics to render capture baking tool

#rnx

diff --git a/Engine/Plugins/Experimental/MeshModelingToolsetExp/Source/MeshModelingToolsExp/Private/BakeRenderCaptureTool.cpp b/Engine/Plugins/Experimental/MeshModelingToolsetExp/Source/MeshModelingToolsExp/Private/BakeRenderCaptureTool.cpp
--- a/Engine/Plugins/Experimental/MeshModelingToolsetExp/Source/MeshModelingToolsExp/Private/BakeRenderCaptureTool.cpp
+++ b/Engine/Plugins/Experimental/MeshModelingToolsetExp/Source/MeshModelingToolsExp/Private/BakeRenderCaptureTool.cpp
@@ -14,6 +14,8 @@
 
 #include "ModelingObjectsCreationAPI.h"
 
+#include "EngineAnalytics.h"
+
 #include "Sampling/MeshImageBakingCache.h"
 #include "Sampling/MeshMapBaker.h"
 #include "Sampling/RenderCaptureMapEvaluator.h"
@@ -806,7 +808,7 @@ void UBakeRenderCaptureTool::CreateTextureAssetsRC(UWorld* SourceWorld, UObject*
 
 	ensure(bCreatedAssetOK);
 
-	//RecordAnalytics(BakeAnalytics, GetAnalyticsEventName());
+	RecordAnalytics();
 }
 
 
@@ -983,7 +985,7 @@ void UBakeRenderCaptureTool::OnMapsUpdatedRC(const TUniquePtr<FMeshMapBaker>& Ne
 		}
 	}
 
-	//GatherAnalytics(*NewResult, CachedBakeSettings, BakeAnalytics);
+	GatherAnalytics(*NewResult);
 	UpdateVisualization();
 	GetToolManager()->PostInvalidation();
 }
@@ -1263,11 +1265,6 @@ void UBakeRenderCaptureTool::UpdateVisualization()
 }
 
 
-void UBakeRenderCaptureTool::GatherAnalytics(FBakeAnalytics::FMeshSettings& Data)
-{
-	
-}
-
 
 void UBakeRenderCaptureTool::InvalidateResults()
 {
@@ -1280,4 +1277,102 @@ void UBakeRenderCaptureTool::InvalidateResults()
 	ResultSettings->NormalMap = nullptr;
 }
 
-#undef LOCTEXT_NAMESPACE
\ No newline at end of file
+
+
+void UBakeRenderCaptureTool::RecordAnalytics() const
+{
+	if (FEngineAnalytics::IsAvailable() == false)
+	{
+		return;
+	}
+
+	TArray<FAnalyticsEventAttribute> Attributes;
+
+	// General
+	Attributes.Add(FAnalyticsEventAttribute(TEXT("Bake.Duration.Total.Seconds"), BakeAnalytics.TotalBakeDuration));
+	Attributes.Add(FAnalyticsEventAttribute(TEXT("Bake.Duration.WriteToImage.Seconds"), BakeAnalytics.WriteToImageDuration));
+	Attributes.Add(FAnalyticsEventAttribute(TEXT("Bake.Duration.WriteToGutter.Seconds"), BakeAnalytics.WriteToGutterDuration));
+	Attributes.Add(FAnalyticsEventAttribute(TEXT("Bake.Stats.NumSamplePixels"), BakeAnalytics.NumSamplePixels));
+	Attributes.Add(FAnalyticsEventAttribute(TEXT("Bake.Stats.NumGutterPixels"), BakeAnalytics.NumGutterPixels));
+
+	// Input mesh data
+	Attributes.Add(FAnalyticsEventAttribute(TEXT("Input.TargetMesh.NumTriangles"), BakeAnalytics.MeshSettings.NumTargetMeshTris));
+	Attributes.Add(FAnalyticsEventAttribute(TEXT("Input.RenderCapture.NumMeshes"), BakeAnalytics.MeshSettings.NumDetailMesh));
+	Attributes.Add(FAnalyticsEventAttribute(TEXT("Input.RenderCapture.NumTriangles"), BakeAnalytics.MeshSettings.NumDetailMeshTris));
+
+	// Bake settings
+	Attributes.Add(FAnalyticsEventAttribute(TEXT("Settings.Image.Width"), static_cast<int32>(Settings->TextureSize)));
+	Attributes.Add(FAnalyticsEventAttribute(TEXT("Settings.Image.Height"), static_cast<int32>(Settings->TextureSize)));
+	Attributes.Add(FAnalyticsEventAttribute(TEXT("Settings.SamplesPerPixel"), static_cast<int32>(Settings->SamplesPerPixel)));
+	Attributes.Add(FAnalyticsEventAttribute(TEXT("Settings.TargetUVLayer"), InputMeshSettings->TargetUVLayerNamesList.IndexOfByKey(InputMeshSettings->TargetUVLayer)));
+
+	// Render Capture settings
+	Attributes.Add(FAnalyticsEventAttribute(TEXT("Settings.RenderCapture.Image.Width"), static_cast<int32>(RenderCaptureProperties->Resolution)));
+	Attributes.Add(FAnalyticsEventAttribute(TEXT("Settings.RenderCapture.Image.Height"), static_cast<int32>(RenderCaptureProperties->Resolution)));
+	Attributes.Add(FAnalyticsEventAttribute(TEXT("Settings.RenderCapture.NormalMap.Enabled"), RenderCaptureProperties->bNormalMap));
+	Attributes.Add(FAnalyticsEventAttribute(TEXT("Settings.RenderCapture.MetallicMap.Enabled"), RenderCaptureProperties->bMetallicMap));
+	Attributes.Add(FAnalyticsEventAttribute(TEXT("Settings.RenderCapture.RoughnessMap.Enabled"), RenderCaptureProperties->bRoughnessMap));
+	Attributes.Add(FAnalyticsEventAttribute(TEXT("Settings.RenderCapture.SpecularMap.Enabled"), RenderCaptureProperties->bSpecularMap));
+	Attributes.Add(FAnalyticsEventAttribute(TEXT("Settings.RenderCapture.PackedMRSMap.Enabled"), RenderCaptureProperties->bPackedMRSMap));
+	Attributes.Add(FAnalyticsEventAttribute(TEXT("Settings.RenderCapture.EmissiveMap.Enabled"), RenderCaptureProperties->bEmissiveMap));
+	Attributes.Add(FAnalyticsEventAttribute(TEXT("Settings.RenderCapture.CaptureFieldOfView"), RenderCaptureProperties->CaptureFieldOfView));
+	Attributes.Add(FAnalyticsEventAttribute(TEXT("Settings.RenderCapture.NearPlaneDistance"), RenderCaptureProperties->NearPlaneDist));
+
+	FEngineAnalytics::GetProvider().RecordEvent(FString(TEXT("Editor.Usage.MeshModelingMode.")) + GetAnalyticsEventName(), Attributes);
+
+	constexpr bool bDebugLogAnalytics = false;
+	if constexpr (bDebugLogAnalytics)
+	{
+		for (const FAnalyticsEventAttribute& Attr : Attributes)
+		{
+			UE_LOG(LogGeometry, Log, TEXT("[%s] %s = %s"), *GetAnalyticsEventName(), *Attr.GetName(), *Attr.GetValue());
+		}
+	}
+}
+
+
+void UBakeRenderCaptureTool::GatherAnalytics(const FMeshMapBaker& Result)
+{
+	if (!FEngineAnalytics::IsAvailable())
+	{
+		return;
+	}
+
+	BakeAnalytics.TotalBakeDuration = Result.BakeAnalytics.TotalBakeDuration;
+	BakeAnalytics.WriteToImageDuration = Result.BakeAnalytics.WriteToImageDuration;
+	BakeAnalytics.WriteToGutterDuration = Result.BakeAnalytics.WriteToGutterDuration;
+	BakeAnalytics.NumSamplePixels = Result.BakeAnalytics.NumSamplePixels;
+	BakeAnalytics.NumGutterPixels = Result.BakeAnalytics.NumGutterPixels;
+}
+
+
+void UBakeRenderCaptureTool::GatherAnalytics(FBakeAnalytics::FMeshSettings& Data)
+{
+	if (FEngineAnalytics::IsAvailable() == false)
+	{
+		return;
+	}
+
+	Data.NumTargetMeshTris = TargetMesh.TriangleCount();
+	Data.NumDetailMesh = Actors.Num();
+	Data.NumDetailMeshTris = 0;
+	for (AActor* Actor : Actors)
+	{
+		check(Actor != nullptr);
+		TInlineComponentArray<UPrimitiveComponent*> PrimitiveComponents;
+		Actor->GetComponents(PrimitiveComponents);
+		for (UPrimitiveComponent* PrimitiveComponent : PrimitiveComponents)
+		{
+			if (UStaticMeshComponent* StaticMeshComponent = Cast<UStaticMeshComponent>(PrimitiveComponent))
+			{
+				if (StaticMeshComponent->GetStaticMesh() != nullptr)
+				{
+					// TODO We could also check GetNumNaniteTriangles here and use the maximum
+					Data.NumDetailMeshTris += StaticMeshComponent->GetStaticMesh()->GetNumTriangles(0);
+				}
+			} 
+		}
+	}
+}
+
+#undef LOCTEXT_NAMESPACE
diff --git a/Engine/Plugins/Experimental/MeshModelingToolsetExp/Source/MeshModelingToolsExp/Public/BakeRenderCaptureTool.h b/Engine/Plugins/Experimental/MeshModelingToolsetExp/Source/MeshModelingToolsExp/Public/BakeRenderCaptureTool.h
--- a/Engine/Plugins/Experimental/MeshModelingToolsetExp/Source/MeshModelingToolsExp/Public/BakeRenderCaptureTool.h
+++ b/Engine/Plugins/Experimental/MeshModelingToolsetExp/Source/MeshModelingToolsExp/Public/BakeRenderCaptureTool.h
@@ -304,4 +304,12 @@ protected:
 	// the capture are reverted to these values
 	UPROPERTY()
 	TObjectPtr<URenderCaptureProperties> ComputedRenderCaptureProperties;
+
+	//
+	// Analytics
+	//
+
+	void RecordAnalytics() const;
+	void GatherAnalytics(const UE::Geometry::FMeshMapBaker& Result);
+
 };
\ No newline at end of file
