GeometryProcessing: add FMeshResolveTJunctions, tries to fix cases where a boundary vertex lies on the interior of some other boundary edge, via edge splits, so that a Weld will work on the new set of boundary edges.
ModelingTools: Add option to Weld Tool to try to resolve T-junctions. Add feedback to report back the initial and final number of open boundary edges.

diff --git a/Engine/Plugins/Experimental/MeshModelingToolsetExp/Source/MeshModelingToolsExp/Private/WeldMeshEdgesTool.cpp b/Engine/Plugins/Experimental/MeshModelingToolsetExp/Source/MeshModelingToolsExp/Private/WeldMeshEdgesTool.cpp
--- a/Engine/Plugins/Experimental/MeshModelingToolsetExp/Source/MeshModelingToolsExp/Private/WeldMeshEdgesTool.cpp
+++ b/Engine/Plugins/Experimental/MeshModelingToolsetExp/Source/MeshModelingToolsExp/Private/WeldMeshEdgesTool.cpp
@@ -10,6 +10,7 @@
 
 #include "DynamicMesh/DynamicMesh3.h"
 #include "DynamicMesh/Operations/MergeCoincidentMeshEdges.h"
+#include "Operations/MeshResolveTJunctions.h"
 
 using namespace UE::Geometry;
 
@@ -35,6 +36,10 @@ public:
 	TSharedPtr<FDynamicMesh3, ESPMode::ThreadSafe> SourceMesh;
 	double Tolerance;
 	bool bOnlyUnique;
+	bool bResolveTJunctions;
+
+	int32 InitialNumBoundaryEdges = 0;
+	int32 FinalNumBoundaryEdges = 0;
 
 public:
 	FWeldMeshEdgesOp() : FDynamicMeshOperator() {}
@@ -52,19 +57,53 @@ public:
 		}
 
 		ResultMesh->Copy(*SourceMesh, true, true, true, true);
-	
+
+		// If we had bowties or broken normals on input, we will probably still have them on output...
+		bool bWasCleanMesh = ResultMesh->CheckValidity(FDynamicMesh3::FValidityOptions(), EValidityCheckFailMode::ReturnOnly);
+
+		InitialNumBoundaryEdges = 0;
+		for (int eid : SourceMesh->BoundaryEdgeIndicesItr())
+		{
+			InitialNumBoundaryEdges++;
+		}
+		FinalNumBoundaryEdges = InitialNumBoundaryEdges;
+
 		FMergeCoincidentMeshEdges Merger(ResultMesh.Get());
 		Merger.MergeVertexTolerance = Tolerance;
 		Merger.MergeSearchTolerance = 2 * Merger.MergeVertexTolerance;
 		Merger.OnlyUniquePairs = bOnlyUnique;
 
-		if (Merger.Apply() == false)
+		bool bOK = Merger.Apply();
+		FinalNumBoundaryEdges = Merger.FinalNumBoundaryEdges;
+
+		if (bOK && Merger.FinalNumBoundaryEdges > 0 && bResolveTJunctions)
+		{
+			FMeshResolveTJunctions Resolver(ResultMesh.Get());
+			Resolver.DistanceTolerance = 2 * Tolerance;
+			bool bResolveOK = Resolver.Apply();
+
+			if (bResolveOK && Resolver.NumSplitEdges > 0)
+			{
+				FMergeCoincidentMeshEdges SecondPassMerger(ResultMesh.Get());
+				SecondPassMerger.MergeVertexTolerance = Tolerance;
+				SecondPassMerger.MergeSearchTolerance = 2 * SecondPassMerger.MergeVertexTolerance;
+				SecondPassMerger.OnlyUniquePairs = bOnlyUnique;
+				bOK = SecondPassMerger.Apply();
+				FinalNumBoundaryEdges = SecondPassMerger.FinalNumBoundaryEdges;
+			}
+		}
+
+		FDynamicMesh3::FValidityOptions UseValidityCheck = (bWasCleanMesh) ?
+			FDynamicMesh3::FValidityOptions() : FDynamicMesh3::FValidityOptions::Permissive();
+		if (bOK == false)
 		{
 			ResultMesh->Copy(*SourceMesh, true, true, true, true);
+			FinalNumBoundaryEdges = InitialNumBoundaryEdges;
 		}
-		else if (ResultMesh->CheckValidity(true, EValidityCheckFailMode::ReturnOnly) == false)
+		else if (ResultMesh->CheckValidity(UseValidityCheck, EValidityCheckFailMode::ReturnOnly) == false)
 		{
 			ResultMesh->Copy(*SourceMesh, true, true, true, true);
+			FinalNumBoundaryEdges = InitialNumBoundaryEdges;
 		}
 	}
 
@@ -133,6 +172,7 @@ void UWeldMeshEdgesTool::Setup()
 	AddToolPropertySource(Settings);
 	Settings->WatchProperty(Settings->Tolerance, [this](float) { PreviewCompute->InvalidateResult(); });
 	Settings->WatchProperty(Settings->bOnlyUnique, [this](bool) { PreviewCompute->InvalidateResult(); });
+	Settings->WatchProperty(Settings->bResolveTJunctions, [this](bool) { PreviewCompute->InvalidateResult(); });
 
 	// create mesh display
 	MeshElementsDisplay = NewObject<UMeshElementsVisualizer>(this);
@@ -158,6 +198,11 @@ void UWeldMeshEdgesTool::Setup()
 		});
 	});
 
+	PreviewCompute->OnOpCompleted.AddLambda([this](const FDynamicMeshOperator* Op)
+	{
+		Settings->InitialEdges = ((FWeldMeshEdgesOp*)(Op))->InitialNumBoundaryEdges;
+		Settings->RemainingEdges = ((FWeldMeshEdgesOp*)(Op))->FinalNumBoundaryEdges;
+	});
 
 	PreviewCompute->InvalidateResult();
 
@@ -218,6 +263,7 @@ void UWeldMeshEdgesTool::OnTick(float DeltaTime)
 void UWeldMeshEdgesTool::UpdateOpParameters(FWeldMeshEdgesOp& Op) const
 {
 	Op.bOnlyUnique = Settings->bOnlyUnique;
+	Op.bResolveTJunctions = Settings->bResolveTJunctions;
 	Op.Tolerance = Settings->Tolerance;
 	Op.SourceMesh = SourceMesh;
 
diff --git a/Engine/Plugins/Experimental/MeshModelingToolsetExp/Source/MeshModelingToolsExp/Public/WeldMeshEdgesTool.h b/Engine/Plugins/Experimental/MeshModelingToolsetExp/Source/MeshModelingToolsExp/Public/WeldMeshEdgesTool.h
--- a/Engine/Plugins/Experimental/MeshModelingToolsetExp/Source/MeshModelingToolsExp/Public/WeldMeshEdgesTool.h
+++ b/Engine/Plugins/Experimental/MeshModelingToolsetExp/Source/MeshModelingToolsExp/Public/WeldMeshEdgesTool.h
@@ -43,6 +43,18 @@ public:
 	/** Only merge unambiguous pairs that have unique duplicate-edge matches */
 	UPROPERTY(EditAnywhere, Category = Options)
 	bool bOnlyUnique = false;
+
+	/** If enabled, after an initial attempt at Welding, attempt to resolve remaining open edges in T-junction configurations via edge splits, and then retry Weld */
+	UPROPERTY(EditAnywhere, Category = Options)
+	bool bResolveTJunctions = false;
+
+	/** Initial number of open boundary edges */
+	UPROPERTY(VisibleAnywhere, Category = Statistics)
+	int32 InitialEdges;
+
+	/** Number of remaining open boundary edges */
+	UPROPERTY(VisibleAnywhere, Category = Statistics)
+	int32 RemainingEdges;
 };
 
 
diff --git a/Engine/Plugins/Runtime/GeometryProcessing/Source/DynamicMesh/Private/Operations/MeshResolveTJunctions.cpp b/Engine/Plugins/Runtime/GeometryProcessing/Source/DynamicMesh/Private/Operations/MeshResolveTJunctions.cpp
new file mode 100644
--- /dev/null
+++ b/Engine/Plugins/Runtime/GeometryProcessing/Source/DynamicMesh/Private/Operations/MeshResolveTJunctions.cpp
@@ -0,0 +1,119 @@
+// Copyright Epic Games, Inc. All Rights Reserved.
+
+
+#include "Operations/MeshResolveTJunctions.h"
+#include "DynamicMesh/DynamicMesh3.h"
+#include "Spatial/PointSetHashTable.h"
+#include "SegmentTypes.h"
+
+using namespace UE::Geometry;
+
+const double FMeshResolveTJunctions::DEFAULT_TOLERANCE = FMathf::ZeroTolerance;
+
+bool FMeshResolveTJunctions::Apply()
+{
+	NumSplitEdges = 0;
+
+	// make clean boundary edge set, either for the whole mesh, or based on input set
+	if (BoundaryEdges.Num() == 0)
+	{
+		for (int32 eid : Mesh->BoundaryEdgeIndicesItr())
+		{
+			BoundaryEdges.Add(eid);
+		}
+	}
+	else
+	{
+		TArray<int32> ToRemove;
+		for (int32 eid : Mesh->BoundaryEdgeIndicesItr())
+		{
+			if (Mesh->IsBoundaryEdge(eid) == false)
+			{
+				ToRemove.Add(eid);
+			}
+		}
+		for (int32 eid : ToRemove)
+		{
+			BoundaryEdges.Remove(eid);
+		}
+	}
+
+
+	TSet<int32> BoundaryVertices;
+	for (int32 eid : BoundaryEdges)
+	{
+		FIndex2i EdgeV = Mesh->GetEdgeV(eid);
+		BoundaryVertices.Add(EdgeV.A);
+		BoundaryVertices.Add(EdgeV.B);
+	}
+
+	// TODO: this should probably use a hash table like MeshMergeCoincidentEdges.
+	// However in this case adding the midpoints won't work, the entire edge boxes
+	// needed to be added. And when edges are split, the hash will need to
+	// be updated. So for now this is just O(N*M).
+	// (idea: a heuristic to avoid removing from the hash might be to just not
+	//  bother updating the bounds of the pre-split edges. It will mean some extra
+	//  checking but may be cheaper than a hash-remove)
+
+	for (int32 vid : BoundaryVertices)
+	{
+		FVector3d Position = Mesh->GetVertex(vid);
+		int e0, e1;
+		Mesh->GetVtxBoundaryEdges(vid, e0, e1);
+
+		int32 OnEdgeID = -1;
+		double MinDistSqr = TNumericLimits<double>::Max();
+		FSegment3d OnEdgeSegment;
+
+		// find the edge in the boundary edge set that is closest to this vertex, 
+		// and not connected to this vertex
+		for (int32 eid : BoundaryEdges)
+		{
+			if (eid == e0 || eid == e1) continue;
+
+			FVector3d A, B;
+			Mesh->GetEdgeV(eid, A, B);
+			FSegment3d EdgeSegment(A, B);
+			double DistSqr = EdgeSegment.DistanceSquared(Position);
+			if (DistSqr < MinDistSqr)
+			{
+				MinDistSqr = DistSqr;
+				OnEdgeID = eid;
+				OnEdgeSegment = EdgeSegment;
+			}
+		}
+
+		// if we did not find an edge, or we are too far from any edge, give up
+		if (OnEdgeID == -1 || MinDistSqr > DistanceTolerance*DistanceTolerance)
+		{
+			continue;
+		}
+		// if we are within tolerance of either edge endpoint, we do not need to split
+		if (Distance(OnEdgeSegment.StartPoint(), Position) < DistanceTolerance
+			|| Distance(OnEdgeSegment.EndPoint(), Position) < DistanceTolerance)
+		{
+			continue;
+		}
+
+		// check that the position is within the span of the edge
+		double SegmentT = OnEdgeSegment.Project(Position);
+		if (FMathd::Abs(SegmentT) > FMathd::Max(OnEdgeSegment.Extent - DistanceTolerance, 0.0) )
+		{
+			continue;		// at/on endpoint, should be able to weld, can skip
+		}
+
+		// now get edge parameter
+		double SplitParameter = OnEdgeSegment.ProjectUnitRange(Position);
+
+		// split boundary edge and add new boundary edge to the active edge set
+		FDynamicMesh3::FEdgeSplitInfo SplitInfo;
+		if (Mesh->SplitEdge(OnEdgeID, SplitInfo, SplitParameter) == EMeshResult::Ok)
+		{
+			BoundaryEdges.Add(SplitInfo.NewEdges.A);
+			NumSplitEdges++;
+		}
+
+	}
+
+	return true;
+}
\ No newline at end of file
diff --git a/Engine/Plugins/Runtime/GeometryProcessing/Source/DynamicMesh/Public/Operations/MeshResolveTJunctions.h b/Engine/Plugins/Runtime/GeometryProcessing/Source/DynamicMesh/Public/Operations/MeshResolveTJunctions.h
new file mode 100644
--- /dev/null
+++ b/Engine/Plugins/Runtime/GeometryProcessing/Source/DynamicMesh/Public/Operations/MeshResolveTJunctions.h
@@ -0,0 +1,76 @@
+// Copyright Epic Games, Inc. All Rights Reserved.
+
+#pragma once
+
+#include "MathUtil.h"
+#include "VectorTypes.h"
+
+namespace UE
+{
+namespace Geometry
+{
+
+class FDynamicMesh3;
+
+/**
+ * FMeshResolveTJunctions splits edges to create matching vertices at T-Junctions in the mesh.
+ * T-Junctions are edge configurations where a border vertex lies on an opposing border edge. 
+ * This frequently occurs in (eg) CAD meshes where the tessellator had bugs or simply did not 
+ * bother to ensure that tessellations match on adjacent spline patches. 
+ * 
+ * The simplest case would be, one side has two collinear edges [A,B] and [B,C],
+ * and the other side has edge [F,G] such that F/A and G/C are coincident. 
+ * The implementation works by adding a matching vertex to [F,G], via an edge split
+ * at the projected location of B, to create a new vertex X that could be welded with B.
+ * 
+ *                    F  A                         F  A
+ * So basically it    |  |                         |  |
+ * turns this         |  B         into this one   X  B
+ * configuration      |  |                         |  |
+ *                    G  C                         G  C
+ *
+ * If everything is within-tolerance, then after resolving all the T-Junctions, 
+ * a FMeshMergeCoincidentEdges would successfully weld the new set of border 
+ * edges back together (note: caller must do this, FMeshResolveTJunctions only splits, it does not weld!)
+ * 
+ * Caller can provide a subset of edges via BoundaryEdges, otherwise all boundary edges
+ * in Mesh will be used.
+ * 
+ * Current implementation is O(N*M) in the number of boundary edges (N) and boundary vertices (M).
+ * Could be improved with a spatial data structure. 
+ */
+class DYNAMICMESH_API FMeshResolveTJunctions
+{
+public:
+	/** default tolerance is float ZeroTolerance */
+	static const double DEFAULT_TOLERANCE;  // = FMathf::ZeroTolerance;
+
+	/** The mesh that we are modifying */
+	FDynamicMesh3* Mesh;
+
+	/** Subset of mesh boundary edges (otherwise all boundary edges are processed) */
+	TSet<int32> BoundaryEdges;
+
+	/** Distance threshold used for various checks (eg is vertex on edge, end endpoint tolerance, etc) */
+	double DistanceTolerance = DEFAULT_TOLERANCE;
+
+	/** Number of edges that were split to resolve T-junctions */
+	int32 NumSplitEdges = 0;
+
+
+public:
+	FMeshResolveTJunctions(FDynamicMesh3* MeshIn) : Mesh(MeshIn)
+	{
+	}
+
+	/**
+	 * Run the resolve operation and modify .Mesh
+	 * @return true if the algorithm succeeds
+	 */
+	bool Apply();
+
+};
+
+
+} // end namespace UE::Geometry
+} // end namespace UE
\ No newline at end of file
diff --git a/Engine/Source/Runtime/GeometryCore/Private/DynamicMesh/Operations/MergeCoincidentMeshEdges.cpp b/Engine/Source/Runtime/GeometryCore/Private/DynamicMesh/Operations/MergeCoincidentMeshEdges.cpp
--- a/Engine/Source/Runtime/GeometryCore/Private/DynamicMesh/Operations/MergeCoincidentMeshEdges.cpp
+++ b/Engine/Source/Runtime/GeometryCore/Private/DynamicMesh/Operations/MergeCoincidentMeshEdges.cpp
@@ -54,8 +54,11 @@ bool FMergeCoincidentMeshEdges::Apply()
 
 	// @todo equivalence sets should be symmetric. this neither enforces that,
 	// nor takes advantage of it.
+	InitialNumBoundaryEdges = 0;
 	for (int eid : Mesh->BoundaryEdgeIndicesItr()) 
 	{
+		InitialNumBoundaryEdges++;
+
 		FVector3d midpt = Mesh->GetEdgePoint(eid, 0.5);
 
 		// find all other edges with same midpoint in query sphere
@@ -197,6 +200,12 @@ bool FMergeCoincidentMeshEdges::Apply()
 
 	}
 
+	FinalNumBoundaryEdges = 0;
+	for (int eid : Mesh->BoundaryEdgeIndicesItr())
+	{
+		FinalNumBoundaryEdges++;
+	}
+
 	return true;
 }
 
diff --git a/Engine/Source/Runtime/GeometryCore/Public/DynamicMesh/Operations/MergeCoincidentMeshEdges.h b/Engine/Source/Runtime/GeometryCore/Public/DynamicMesh/Operations/MergeCoincidentMeshEdges.h
--- a/Engine/Source/Runtime/GeometryCore/Public/DynamicMesh/Operations/MergeCoincidentMeshEdges.h
+++ b/Engine/Source/Runtime/GeometryCore/Public/DynamicMesh/Operations/MergeCoincidentMeshEdges.h
@@ -45,6 +45,12 @@ public:
 	  */
 	double MergeSearchTolerance = 0;
 
+
+	/** Number of mesh boundary edges before merging */
+	int32 InitialNumBoundaryEdges = 0;
+	/** Number of mesh boundary edges after merging */
+	int32 FinalNumBoundaryEdges = 0;
+
 public:
 	FMergeCoincidentMeshEdges(FDynamicMesh3* mesh) : Mesh(mesh)
 	{
