Render Capture Baking: Add messages warning users about missing tangents or UVs

#rnx

diff --git a/Engine/Plugins/Experimental/MeshModelingToolsetExp/Source/MeshModelingToolsExp/Private/BakeRenderCaptureTool.cpp b/Engine/Plugins/Experimental/MeshModelingToolsetExp/Source/MeshModelingToolsExp/Private/BakeRenderCaptureTool.cpp
--- a/Engine/Plugins/Experimental/MeshModelingToolsetExp/Source/MeshModelingToolsExp/Private/BakeRenderCaptureTool.cpp
+++ b/Engine/Plugins/Experimental/MeshModelingToolsetExp/Source/MeshModelingToolsExp/Private/BakeRenderCaptureTool.cpp
@@ -155,11 +155,11 @@ public:
 		//
 		FOptions Options;
 
-		Options.TargetUVLayer = InputMeshSettings.TargetUVLayerNamesList.IndexOfByKey(InputMeshSettings.TargetUVLayer);
+		Options.TargetUVLayer = InputMeshSettings.GetTargetUVLayerIndex();
 
 		Options.RenderCaptureImageSize = static_cast<int32>(RenderCaptureProperties.Resolution);
 
-		//Options.bBakeBaseColor // This is always true
+		Options.bBakeBaseColor = RenderCaptureProperties.bBaseColorMap;
 		Options.bBakeRoughness = RenderCaptureProperties.bRoughnessMap;
 		Options.bBakeMetallic = RenderCaptureProperties.bMetallicMap;
 		Options.bBakeSpecular = RenderCaptureProperties.bSpecularMap;
@@ -274,16 +274,33 @@ MakeColorEvaluator(const FSceneCapturePhotoSet::FSceneSample& DefaultSample, con
 	return Evaluator;
 }
 
-static void ImageBuildersFromPhotoSet(
-	FSceneCapturePhotoSet* SceneCapture,
-	const FRenderCaptureSettings::FOptions& Options,
-	const int32 TextureImageSize,
-	const EBakeTextureSamplesPerPixel SamplesPerPixel,
-	const FDynamicMesh3* BaseMesh,
-	const TSharedPtr<FMeshTangentsd, ESPMode::ThreadSafe>& BaseMeshTangents,
-	TUniquePtr<FMeshMapBaker>& Result)
+//
+// Tool Operator
+//
+
+class FRenderCaptureMapBakerOp : public TGenericDataOperator<FMeshMapBaker>
+{
+public:
+	UE::Geometry::FDynamicMesh3* BaseMesh = nullptr;
+	TSharedPtr<UE::Geometry::FMeshTangentsd, ESPMode::ThreadSafe> BaseMeshTangents;
+	FRenderCaptureSettings::FOptions Options;
+	int32 TextureImageSize;
+	EBakeTextureSamplesPerPixel SamplesPerPixel;
+	FSceneCapturePhotoSet* SceneCapture = nullptr;
+
+	// Begin TGenericDataOperator interface
+	virtual void CalculateResult(FProgressCancel* Progress) override;
+	// End TGenericDataOperator interface
+};
+
+// Bake textures onto the base/target mesh by projecting/sampling the set of captured photos
+void FRenderCaptureMapBakerOp::CalculateResult(FProgressCancel*)
 {
-	TRACE_CPUPROFILER_EVENT_SCOPE(ImageBuildersFromPhotoSet);
+	TRACE_CPUPROFILER_EVENT_SCOPE(FRenderCaptureMapBakerOp_CalculateResult);
+
+	check(BaseMesh != nullptr);
+	check(BaseMeshTangents.IsValid());
+	check(SceneCapture != nullptr);
 
 	FDynamicMeshAABBTree3 BaseMeshSpatial(BaseMesh, true);
 
@@ -429,51 +446,40 @@ static void ImageBuildersFromPhotoSet(
 	DefaultColorSample.Emissive = FVector3f(InvalidColor.X, InvalidColor.Y, InvalidColor.Z);
 	DefaultColorSample.WorldNormal = FVector4f((DefaultNormal + FVector3f::One()) * .5f, InvalidColor.W);
 
-	auto AddColorEvaluator = [&Result, &InfillData] (
-		const TSharedPtr<FRenderCaptureMapEvaluator<FVector4f>>& Evaluator,
-		ERenderCaptureType CaptureType)
+	auto AddColorEvaluator = [this, &InfillData] (const TSharedPtr<FRenderCaptureMapEvaluator<FVector4f>>& Evaluator)
 	{
 		Result->AddEvaluator(Evaluator);
 		InfillData.EvaluatorNeedsInfill.Add(true);
 	};
 
-	AddColorEvaluator(
-		MakeColorEvaluator<ERenderCaptureType::BaseColor>(DefaultColorSample, SceneCapture),
-		ERenderCaptureType::BaseColor);
+	if (Options.bBakeBaseColor)
+	{
+		AddColorEvaluator(MakeColorEvaluator<ERenderCaptureType::BaseColor>(DefaultColorSample, SceneCapture));
+	}
 
 	if (Options.bUsePackedMRS)
 	{
-		AddColorEvaluator(
-			MakeColorEvaluator<ERenderCaptureType::CombinedMRS>(DefaultColorSample, SceneCapture),
-			ERenderCaptureType::CombinedMRS);
+		AddColorEvaluator(MakeColorEvaluator<ERenderCaptureType::CombinedMRS>(DefaultColorSample, SceneCapture));
 	}
 	else
 	{
 		if (Options.bBakeRoughness)
 		{
-			AddColorEvaluator(
-				MakeColorEvaluator<ERenderCaptureType::Roughness>(DefaultColorSample, SceneCapture),
-				ERenderCaptureType::Roughness);
+			AddColorEvaluator(MakeColorEvaluator<ERenderCaptureType::Roughness>(DefaultColorSample, SceneCapture));
 		}
 		if (Options.bBakeMetallic)
 		{
-			AddColorEvaluator(
-				MakeColorEvaluator<ERenderCaptureType::Metallic>(DefaultColorSample, SceneCapture),
-				ERenderCaptureType::Metallic);
+			AddColorEvaluator(MakeColorEvaluator<ERenderCaptureType::Metallic>(DefaultColorSample, SceneCapture));
 		}
 		if (Options.bBakeSpecular)
 		{
-			AddColorEvaluator(
-				MakeColorEvaluator<ERenderCaptureType::Specular>(DefaultColorSample, SceneCapture),
-				ERenderCaptureType::Specular);
+			AddColorEvaluator(MakeColorEvaluator<ERenderCaptureType::Specular>(DefaultColorSample, SceneCapture));
 		}
 	}
 
 	if (Options.bBakeEmissive)
 	{
-		AddColorEvaluator(
-			MakeColorEvaluator<ERenderCaptureType::Emissive>(DefaultColorSample, SceneCapture),
-			ERenderCaptureType::Emissive);
+		AddColorEvaluator(MakeColorEvaluator<ERenderCaptureType::Emissive>(DefaultColorSample, SceneCapture));
 	}
 
 	if (Options.bBakeNormalMap) {
@@ -483,7 +489,7 @@ static void ImageBuildersFromPhotoSet(
 
 		Evaluator->DefaultResult = DefaultNormal;
 
-		Evaluator->EvaluateSampleCallback = [&DefaultColorSample, &BaseMeshTangents, SceneCapture](const FMeshMapEvaluator::FCorrespondenceSample& Sample)
+		Evaluator->EvaluateSampleCallback = [this, &DefaultColorSample](const FMeshMapEvaluator::FCorrespondenceSample& Sample)
 		{
 			const int32 TriangleID = Sample.BaseSample.TriangleIndex;
 			const FVector3d BaryCoords = Sample.BaseSample.BaryCoords;
@@ -529,40 +535,11 @@ static void ImageBuildersFromPhotoSet(
 	}
 
 	{
-		TRACE_CPUPROFILER_EVENT_SCOPE(ImageBuildersFromPhotoSet_Bake);
+		TRACE_CPUPROFILER_EVENT_SCOPE(FRenderCaptureMapBakerOp_CalculateResult_Bake);
 		Result->Bake();
 	}
 }
 
-//
-// Tool Operators
-//
-
-class FRenderCaptureMapBakerOp : public TGenericDataOperator<FMeshMapBaker>
-{
-public:
-	UE::Geometry::FDynamicMesh3* BaseMesh = nullptr;
-	TSharedPtr<UE::Geometry::FMeshTangentsd, ESPMode::ThreadSafe> BaseMeshTangents;
-	FRenderCaptureSettings::FOptions Options;
-	int32 TextureImageSize;
-	EBakeTextureSamplesPerPixel SamplesPerPixel;
-	TObjectPtr<UBakeRenderCaptureInputToolProperties> InputMeshSettings;
-	FSceneCapturePhotoSet* SceneCapture = nullptr;
-
-	// Begin TGenericDataOperator interface
-	// @Incomplete Use the Progress thing
-	virtual void CalculateResult(FProgressCancel* Progress) override
-	{
-		check(BaseMesh != nullptr);
-		check(BaseMeshTangents.IsValid());
-		check(SceneCapture != nullptr);
-
-		// Bake textures onto the base/target mesh by projecting/sampling the set of captured photos
-		ImageBuildersFromPhotoSet(SceneCapture, Options, TextureImageSize, SamplesPerPixel, BaseMesh, BaseMeshTangents, Result);
-	}
-	// End TGenericDataOperator interface
-};
-
 
 //
 // Tool Builder
@@ -650,6 +627,7 @@ void UBakeRenderCaptureTool::Setup()
 	AddToolPropertySource(RenderCaptureProperties);
 
 	RenderCaptureProperties->WatchProperty(RenderCaptureProperties->Resolution, [this](EBakeTextureResolution) { OpState |= EBakeOpState::Evaluate; });
+	RenderCaptureProperties->WatchProperty(RenderCaptureProperties->bBaseColorMap, [this](bool) { OpState |= EBakeOpState::Evaluate; });
 	RenderCaptureProperties->WatchProperty(RenderCaptureProperties->bNormalMap, [this](bool) { OpState |= EBakeOpState::Evaluate; });
 	RenderCaptureProperties->WatchProperty(RenderCaptureProperties->bMetallicMap, [this](bool) { OpState |= EBakeOpState::Evaluate; });
 	RenderCaptureProperties->WatchProperty(RenderCaptureProperties->bRoughnessMap, [this](bool) { OpState |= EBakeOpState::Evaluate; });
@@ -681,6 +659,8 @@ void UBakeRenderCaptureTool::Setup()
 	AddToolPropertySource(ResultSettings);
 	SetToolPropertySourceEnabled(ResultSettings, true);
 
+	TargetUVLayerToError.Reset();
+
 	// Used to implement SceneCapture cancellation
 	ComputedRenderCaptureProperties = NewObject<URenderCaptureProperties>(this);
 
@@ -704,6 +684,22 @@ void UBakeRenderCaptureTool::Setup()
 }
 
 
+
+
+
+void UBakeRenderCaptureTool::Render(IToolsContextRenderAPI* RenderAPI)
+{
+	Super::Render(RenderAPI);
+
+	const float Brightness = VisualizationProps->Brightness;
+	const FVector BrightnessColor(Brightness, Brightness, Brightness);
+	PreviewMaterialRC->SetVectorParameterValue(TEXT("Brightness"), BrightnessColor);
+	PreviewMaterialPackedRC->SetVectorParameterValue(TEXT("Brightness"), BrightnessColor);
+}
+
+
+
+
 void UBakeRenderCaptureTool::OnShutdown(EToolShutdownType ShutdownType)
 {
 	TRACE_CPUPROFILER_EVENT_SCOPE(UBakeRenderCaptureTool::Shutdown);
@@ -762,7 +758,7 @@ void UBakeRenderCaptureTool::CreateTextureAssetsRC(UWorld* SourceWorld, UObject*
 				FCreateTextureObjectParams{ 0, SourceWorld, SourceAsset, TexName, Tex }).IsOK();
 	};
 
-	if (ResultSettings->BaseColorMap != nullptr)
+	if (RenderCaptureProperties->bBaseColorMap && ResultSettings->BaseColorMap != nullptr)
 	{
 		const FString TexName = FString::Printf(TEXT("%s_%s"), *BaseName, *BaseColorTexParamName);
 		CreateTextureAsset(TexName, FTexture2DBuilder::ETextureType::Color, ResultSettings->BaseColorMap);
@@ -820,7 +816,7 @@ bool UBakeRenderCaptureTool::CanAccept() const
 		return false;
 	}
 
-	if (ResultSettings->BaseColorMap == nullptr)
+	if (RenderCaptureProperties->bBaseColorMap && ResultSettings->BaseColorMap == nullptr)
 	{
 		return false;
 	}
@@ -1111,6 +1107,7 @@ void UBakeRenderCaptureTool::UpdateResult()
 		{
 			// Restore the settings present before the change that invoked the scene capture recompute
 			RenderCaptureProperties->Resolution         = ComputedRenderCaptureProperties->Resolution;
+			RenderCaptureProperties->bBaseColorMap      = ComputedRenderCaptureProperties->bBaseColorMap;
 			RenderCaptureProperties->bNormalMap         = ComputedRenderCaptureProperties->bNormalMap;
 			RenderCaptureProperties->bMetallicMap       = ComputedRenderCaptureProperties->bMetallicMap;
 			RenderCaptureProperties->bRoughnessMap      = ComputedRenderCaptureProperties->bRoughnessMap;
@@ -1130,6 +1127,7 @@ void UBakeRenderCaptureTool::UpdateResult()
 
 		// Cache Settings used to compute this SceneCapture so we can restore them if the user cancels a SceneCapture recompute
 		ComputedRenderCaptureProperties->Resolution         = RenderCaptureProperties->Resolution;
+		ComputedRenderCaptureProperties->bBaseColorMap      = RenderCaptureProperties->bBaseColorMap;
 		ComputedRenderCaptureProperties->bNormalMap         = RenderCaptureProperties->bNormalMap;
 		ComputedRenderCaptureProperties->bMetallicMap       = RenderCaptureProperties->bMetallicMap;
 		ComputedRenderCaptureProperties->bRoughnessMap      = RenderCaptureProperties->bRoughnessMap;
@@ -1142,16 +1140,68 @@ void UBakeRenderCaptureTool::UpdateResult()
 		bFirstEverSceneCapture = false;
 	}
 
-	// @Incomplete Clear our invalid bitflag to check for valid inputs.
-	//OpState &= ~EBakeOpState::Invalid;
-	//OpState |= CheckValidInputs();
-	//if (static_cast<bool>(OpState & EBakeOpState::Invalid))
-	//{
-	//	// Early exit if op input parameters are invalid.
-	//	return;
-	//}
-	
+	FText ErrorMessage; // Empty message indicates no error
+
+	{
+		const int32 TargetUVLayer = InputMeshSettings->GetTargetUVLayerIndex();
+		if (FText* Message = TargetUVLayerToError.Find(TargetUVLayer); Message)
+		{
+			ErrorMessage = *Message;
+		}
+		else
+		{
+			const auto HasDegenerateUVs = [this]
+			{
+				FDynamicMeshUVOverlay* UVOverlay = TargetMesh.Attributes()->GetUVLayer(InputMeshSettings->GetTargetUVLayerIndex());
+				FAxisAlignedBox2f Bounds = FAxisAlignedBox2f::Empty();
+				for (const int Index : UVOverlay->ElementIndicesItr())
+				{
+					FVector2f UV;
+					UVOverlay->GetElement(Index, UV);
+					Bounds.Contain(UV);
+				}
+				return Bounds.Min == Bounds.Max;
+			};
+
+			if (TargetMesh.Attributes()->GetUVLayer(InputMeshSettings->GetTargetUVLayerIndex()) == nullptr)
+			{
+				ErrorMessage = LOCTEXT("TargetMeshMissingUVs", "The Target Mesh UV layer is missing");
+			}
+			else if (HasDegenerateUVs())
+			{
+				ErrorMessage = LOCTEXT("TargetMeshDegenerateUVs", "The Target Mesh UV layer is degenerate");
+			}
+			else
+			{
+				ErrorMessage = FText(); // No error
+			}
+			TargetUVLayerToError.Add(TargetUVLayer, ErrorMessage);
+		}
+
+		// If there are no UV layer errors check for missing tangent space error
+		if (ErrorMessage.IsEmpty() && RenderCaptureProperties->bNormalMap && ValidTargetMeshTangents() == false)
+		{
+			ErrorMessage = LOCTEXT("TargetMeshMissingTangentSpace", "The Target Mesh is missing a tangent space. Disable Normal Map capture to continue.");
+		}
+	}
+
+	// Calling DisplayMessage with an empty string will clear existing messages
+	GetToolManager()->DisplayMessage(ErrorMessage, EToolMessageLevel::UserWarning);
+
 	InvalidateResults();
+
+	const bool bIsInvalid = (ErrorMessage.IsEmpty() == false);
+	if (bIsInvalid)
+	{
+		const bool bWasValid = static_cast<bool>(OpState & EBakeOpState::Invalid) == false;
+		if (bWasValid)
+		{
+			UpdateVisualization(); // Clear the preview mesh material inputs
+		}
+		OpState = EBakeOpState::Invalid;
+		return;
+	}
+
 	InvalidateComputeRC();
 
 	OpState = EBakeOpState::Clean;
@@ -1165,38 +1215,36 @@ void UBakeRenderCaptureTool::UpdateVisualization()
 		return;
 	}
 
-	FRenderCaptureSettings::FOptions Options = FRenderCaptureSettings::ConstructOptions(*RenderCaptureProperties, *InputMeshSettings);
-	
-	if (Options.bUsePackedMRS)
+	if (ResultSettings->PackedMRSMap)
 	{
 		TObjectPtr<UMaterialInstanceDynamic> Material = PreviewMaterialPackedRC;
 		PreviewMesh->SetOverrideRenderMaterial(Material);
-		
+
 		if (VisualizationProps->bPreviewAsMaterial)
 		{
 			// We set all textures which were computed in the corresponding texture channels
-			Material->SetTextureParameterValue(FName(BaseColorTexParamName), Options.bBakeBaseColor ? ResultSettings->BaseColorMap : EmptyColorMapWhite);
-			Material->SetTextureParameterValue(FName(EmissiveTexParamName),  Options.bBakeEmissive  ? ResultSettings->EmissiveMap  : EmptyEmissiveMap);
-			Material->SetTextureParameterValue(FName(NormalTexParamName), Options.bBakeNormalMap ? ResultSettings->NormalMap    : EmptyNormalMap);
+			Material->SetTextureParameterValue(FName(BaseColorTexParamName), ResultSettings->BaseColorMap ? ResultSettings->BaseColorMap : EmptyColorMapWhite);
+			Material->SetTextureParameterValue(FName(EmissiveTexParamName),  ResultSettings->EmissiveMap  ? ResultSettings->EmissiveMap  : EmptyEmissiveMap);
+			Material->SetTextureParameterValue(FName(NormalTexParamName),    ResultSettings->NormalMap    ? ResultSettings->NormalMap    : EmptyNormalMap);
 			Material->SetTextureParameterValue(FName(PackedMRSTexParamName), ResultSettings->PackedMRSMap);
 		}
 		else
 		{
 			// The BaseColor texture channel will be set according to the selected MapPreview
 			TObjectPtr<UTexture2D> BaseColorMap = EmptyColorMapWhite;
-			if (Options.bBakeBaseColor && Settings->MapPreview == BaseColorTexParamName)
+			if (ResultSettings->BaseColorMap && Settings->MapPreview == BaseColorTexParamName)
 			{
 				BaseColorMap = ResultSettings->BaseColorMap;
 			}
-			else if (Options.bBakeEmissive && Settings->MapPreview == EmissiveTexParamName)
+			else if (ResultSettings->EmissiveMap && Settings->MapPreview == EmissiveTexParamName)
 			{
 				BaseColorMap = ResultSettings->EmissiveMap;
 			}
-			else if (Options.bBakeNormalMap && Settings->MapPreview == NormalTexParamName)
+			else if (ResultSettings->NormalMap && Settings->MapPreview == NormalTexParamName)
 			{
 				BaseColorMap = ResultSettings->NormalMap;
 			}
-			else if (Settings->MapPreview == PackedMRSTexParamName)
+			else if (ResultSettings->PackedMRSMap && Settings->MapPreview == PackedMRSTexParamName)
 			{
 				BaseColorMap = ResultSettings->PackedMRSMap;
 			}
@@ -1206,7 +1254,7 @@ void UBakeRenderCaptureTool::UpdateVisualization()
 			Material->SetTextureParameterValue(FName(PackedMRSTexParamName), EmptyPackedMRSMap);
 		}
 
-		Material->SetScalarParameterValue(TEXT("UVChannel"), Options.TargetUVLayer);
+		Material->SetScalarParameterValue(TEXT("UVChannel"), InputMeshSettings->GetTargetUVLayerIndex());
 	}
 	else
 	{
@@ -1216,38 +1264,38 @@ void UBakeRenderCaptureTool::UpdateVisualization()
 		if (VisualizationProps->bPreviewAsMaterial)
 		{
 			// We set all textures which were computed in the corresponding texture channels
-			Material->SetTextureParameterValue(FName(BaseColorTexParamName), Options.bBakeBaseColor ? ResultSettings->BaseColorMap : EmptyColorMapWhite);
-			Material->SetTextureParameterValue(FName(RoughnessTexParamName), Options.bBakeRoughness ? ResultSettings->RoughnessMap : EmptyRoughnessMap);
-			Material->SetTextureParameterValue(FName(MetallicTexParamName),  Options.bBakeMetallic  ? ResultSettings->MetallicMap  : EmptyMetallicMap);
-			Material->SetTextureParameterValue(FName(SpecularTexParamName),  Options.bBakeSpecular  ? ResultSettings->SpecularMap  : EmptySpecularMap);
-			Material->SetTextureParameterValue(FName(EmissiveTexParamName),  Options.bBakeEmissive  ? ResultSettings->EmissiveMap  : EmptyEmissiveMap);
-			Material->SetTextureParameterValue(FName(NormalTexParamName), Options.bBakeNormalMap ? ResultSettings->NormalMap    : EmptyNormalMap);
+			Material->SetTextureParameterValue(FName(BaseColorTexParamName), ResultSettings->BaseColorMap ? ResultSettings->BaseColorMap : EmptyColorMapWhite);
+			Material->SetTextureParameterValue(FName(RoughnessTexParamName), ResultSettings->RoughnessMap ? ResultSettings->RoughnessMap : EmptyRoughnessMap);
+			Material->SetTextureParameterValue(FName(MetallicTexParamName),  ResultSettings->MetallicMap  ? ResultSettings->MetallicMap  : EmptyMetallicMap);
+			Material->SetTextureParameterValue(FName(SpecularTexParamName),  ResultSettings->SpecularMap  ? ResultSettings->SpecularMap  : EmptySpecularMap);
+			Material->SetTextureParameterValue(FName(EmissiveTexParamName),  ResultSettings->EmissiveMap  ? ResultSettings->EmissiveMap  : EmptyEmissiveMap);
+			Material->SetTextureParameterValue(FName(NormalTexParamName),    ResultSettings->NormalMap    ? ResultSettings->NormalMap    : EmptyNormalMap);
 		}
 		else
 		{
 			// The BaseColor texture channel will be set according to the selected MapPreview
 			TObjectPtr<UTexture2D> BaseColorMap = EmptyColorMapWhite;
-			if (Options.bBakeBaseColor && Settings->MapPreview == BaseColorTexParamName)
+			if (ResultSettings->BaseColorMap && Settings->MapPreview == BaseColorTexParamName)
 			{
 				BaseColorMap = ResultSettings->BaseColorMap;
 			}
-			else if (Options.bBakeRoughness && Settings->MapPreview == RoughnessTexParamName)
+			else if (ResultSettings->RoughnessMap && Settings->MapPreview == RoughnessTexParamName)
 			{
 				BaseColorMap = ResultSettings->RoughnessMap;
 			}
-			else if (Options.bBakeMetallic && Settings->MapPreview == MetallicTexParamName)
+			else if (ResultSettings->MetallicMap && Settings->MapPreview == MetallicTexParamName)
 			{
 				BaseColorMap = ResultSettings->MetallicMap;
 			}
-			else if (Options.bBakeSpecular && Settings->MapPreview == SpecularTexParamName)
+			else if (ResultSettings->SpecularMap && Settings->MapPreview == SpecularTexParamName)
 			{
 				BaseColorMap = ResultSettings->SpecularMap;
 			}
-			else if (Options.bBakeEmissive && Settings->MapPreview == EmissiveTexParamName)
+			else if (ResultSettings->EmissiveMap && Settings->MapPreview == EmissiveTexParamName)
 			{
 				BaseColorMap = ResultSettings->EmissiveMap;
 			}
-			else if (Options.bBakeNormalMap && Settings->MapPreview == NormalTexParamName)
+			else if (ResultSettings->NormalMap && Settings->MapPreview == NormalTexParamName)
 			{
 				BaseColorMap = ResultSettings->NormalMap;
 			}
@@ -1260,7 +1308,7 @@ void UBakeRenderCaptureTool::UpdateVisualization()
 			Material->SetTextureParameterValue(TEXT("NormalMap"), EmptyNormalMap);
 		}
 		
-		Material->SetScalarParameterValue(TEXT("UVChannel"), Options.TargetUVLayer);
+		Material->SetScalarParameterValue(TEXT("UVChannel"), InputMeshSettings->GetTargetUVLayerIndex());
 	}
 }
 
@@ -1304,11 +1352,12 @@ void UBakeRenderCaptureTool::RecordAnalytics() const
 	Attributes.Add(FAnalyticsEventAttribute(TEXT("Settings.Image.Width"), static_cast<int32>(Settings->TextureSize)));
 	Attributes.Add(FAnalyticsEventAttribute(TEXT("Settings.Image.Height"), static_cast<int32>(Settings->TextureSize)));
 	Attributes.Add(FAnalyticsEventAttribute(TEXT("Settings.SamplesPerPixel"), static_cast<int32>(Settings->SamplesPerPixel)));
-	Attributes.Add(FAnalyticsEventAttribute(TEXT("Settings.TargetUVLayer"), InputMeshSettings->TargetUVLayerNamesList.IndexOfByKey(InputMeshSettings->TargetUVLayer)));
+	Attributes.Add(FAnalyticsEventAttribute(TEXT("Settings.TargetUVLayer"), InputMeshSettings->GetTargetUVLayerIndex()));
 
 	// Render Capture settings
 	Attributes.Add(FAnalyticsEventAttribute(TEXT("Settings.RenderCapture.Image.Width"), static_cast<int32>(RenderCaptureProperties->Resolution)));
 	Attributes.Add(FAnalyticsEventAttribute(TEXT("Settings.RenderCapture.Image.Height"), static_cast<int32>(RenderCaptureProperties->Resolution)));
+	Attributes.Add(FAnalyticsEventAttribute(TEXT("Settings.RenderCapture.BaseColorMap.Enabled"), RenderCaptureProperties->bBaseColorMap));
 	Attributes.Add(FAnalyticsEventAttribute(TEXT("Settings.RenderCapture.NormalMap.Enabled"), RenderCaptureProperties->bNormalMap));
 	Attributes.Add(FAnalyticsEventAttribute(TEXT("Settings.RenderCapture.MetallicMap.Enabled"), RenderCaptureProperties->bMetallicMap));
 	Attributes.Add(FAnalyticsEventAttribute(TEXT("Settings.RenderCapture.RoughnessMap.Enabled"), RenderCaptureProperties->bRoughnessMap));
diff --git a/Engine/Plugins/Experimental/MeshModelingToolsetExp/Source/MeshModelingToolsExp/Public/BakeRenderCaptureTool.h b/Engine/Plugins/Experimental/MeshModelingToolsetExp/Source/MeshModelingToolsExp/Public/BakeRenderCaptureTool.h
--- a/Engine/Plugins/Experimental/MeshModelingToolsetExp/Source/MeshModelingToolsExp/Public/BakeRenderCaptureTool.h
+++ b/Engine/Plugins/Experimental/MeshModelingToolsetExp/Source/MeshModelingToolsExp/Public/BakeRenderCaptureTool.h
@@ -29,37 +29,27 @@ public:
 	UPROPERTY(VisibleAnywhere, Category = Results, meta = (TransientToolProperty))
 	TObjectPtr<UTexture2D> BaseColorMap;
 
+	/** World space normal map */
 	UPROPERTY(VisibleAnywhere, Category = Results, meta = (TransientToolProperty))
-	TObjectPtr<UTexture2D> RoughnessMap;
+	TObjectPtr<UTexture2D> NormalMap;
+
+	/** Packed Metallic/Roughness/Specular Map */
+	UPROPERTY(VisibleAnywhere, Category = Results, meta = (TransientToolProperty, DisplayName = "Packed MRS Map"))
+	TObjectPtr<UTexture2D> PackedMRSMap;
 
 	UPROPERTY(VisibleAnywhere, Category = Results, meta = (TransientToolProperty))
 	TObjectPtr<UTexture2D> MetallicMap;
 
 	UPROPERTY(VisibleAnywhere, Category = Results, meta = (TransientToolProperty))
-	TObjectPtr<UTexture2D> SpecularMap;
-
-	/** Packed Metallic/Roughness/Specular Map */
-	UPROPERTY(VisibleAnywhere, Category = Results, meta = (DisplayName = "Packed MRS Map", TransientToolProperty))
-	TObjectPtr<UTexture2D> PackedMRSMap;
+	TObjectPtr<UTexture2D> RoughnessMap;
 
 	UPROPERTY(VisibleAnywhere, Category = Results, meta = (TransientToolProperty))
-	TObjectPtr<UTexture2D> EmissiveMap;
+	TObjectPtr<UTexture2D> SpecularMap;
 
-	/** World space normal map */
 	UPROPERTY(VisibleAnywhere, Category = Results, meta = (TransientToolProperty))
-	TObjectPtr<UTexture2D> NormalMap;
+	TObjectPtr<UTexture2D> EmissiveMap;
 };
 
-struct FBakeRenderCaptureResultsBuilder
-{
-	TUniquePtr<TImageBuilder<FVector4f>> RoughnessImage;
-	TUniquePtr<TImageBuilder<FVector4f>> MetallicImage; 
-	TUniquePtr<TImageBuilder<FVector4f>> SpecularImage; 
-	TUniquePtr<TImageBuilder<FVector4f>> PackedMRSImage; 
-	TUniquePtr<TImageBuilder<FVector4f>> EmissiveImage;
-	TUniquePtr<TImageBuilder<FVector4f>> ColorImage;
-	TUniquePtr<TImageBuilder<FVector4f>> NormalImage; // Tangent-space normal map
-};
 
 //
 // Tool Builder
@@ -93,42 +83,48 @@ class MESHMODELINGTOOLSEXP_API URenderCaptureProperties : public UInteractiveToo
 	GENERATED_BODY()
 
 public:
-	UPROPERTY(Category = RenderCaptureOptions, EditAnywhere, DisplayName="Render Capture Resolution", meta = (ClampMin = "1", UIMin= "1"))
+	UPROPERTY(Category = RenderCaptureOptions, EditAnywhere, meta = (NoResetToDefault, ClampMin = "1", UIMin= "1"), DisplayName="Render Capture Resolution")
 	EBakeTextureResolution Resolution = EBakeTextureResolution::Resolution512;
 
+	// Whether to generate a texture for the Base Color property
+	UPROPERTY(Category = RenderCaptureOptions, EditAnywhere, meta = (NoResetToDefault))
+	bool bBaseColorMap = true;
+
 	// Whether to generate a texture for the World Normal property
-	UPROPERTY(Category = RenderCaptureOptions, EditAnywhere)
+	UPROPERTY(Category = RenderCaptureOptions, EditAnywhere, meta = (NoResetToDefault, EditCondition="bEnableNormalMap == true"))
 	bool bNormalMap = true;
 
+	// Whether to generate a packed texture with Metallic, Roughness and Specular properties
+	UPROPERTY(Category = RenderCaptureOptions, EditAnywhere, meta = (NoResetToDefault), DisplayName="Packed MRS Map")
+	bool bPackedMRSMap = true;
+
 	// Whether to generate a texture for the Metallic property
-	UPROPERTY(Category = RenderCaptureOptions, EditAnywhere, meta=(EditCondition="bPackedMRSMap == false"))
+	UPROPERTY(Category = RenderCaptureOptions, EditAnywhere, meta = (NoResetToDefault, EditCondition="bPackedMRSMap == false"))
 	bool bMetallicMap = true;
 
 	// Whether to generate a texture for the Roughness property
-	UPROPERTY(Category = RenderCaptureOptions, EditAnywhere, meta=(EditCondition="bPackedMRSMap == false"))
+	UPROPERTY(Category = RenderCaptureOptions, EditAnywhere, meta = (NoResetToDefault, EditCondition="bPackedMRSMap == false"))
 	bool bRoughnessMap = true;
 
 	// Whether to generate a texture for the Specular property
-	UPROPERTY(Category = RenderCaptureOptions, EditAnywhere, meta=(EditCondition="bPackedMRSMap == false"))
+	UPROPERTY(Category = RenderCaptureOptions, EditAnywhere, meta = (NoResetToDefault, EditCondition="bPackedMRSMap == false"))
 	bool bSpecularMap = true;
 	
-	// Whether to generate a packed texture with Metallic, Roughness and Specular properties
-	UPROPERTY(Category = RenderCaptureOptions, EditAnywhere, DisplayName="Packed MRS Map")
-	bool bPackedMRSMap = true;
-
 	// Whether to generate a texture for the Emissive property
-	UPROPERTY(Category = RenderCaptureOptions, EditAnywhere)
+	UPROPERTY(Category = RenderCaptureOptions, EditAnywhere, meta = (NoResetToDefault))
 	bool bEmissiveMap = true;
 
-	UPROPERTY(Category = RenderCaptureOptions, EditAnywhere, AdvancedDisplay, DisplayName="Field of View", meta = (ClampMin = "5.0", ClampMax = "160.0"))
+	UPROPERTY(Category = RenderCaptureOptions, EditAnywhere, meta = (ClampMin = "5.0", ClampMax = "160.0"), AdvancedDisplay, DisplayName="Field of View")
 	float CaptureFieldOfView = 30.0f;
 
-	UPROPERTY(Category = RenderCaptureOptions, EditAnywhere, AdvancedDisplay, meta = (ClampMin = "0.001", ClampMax = "1000.0"))
+	UPROPERTY(Category = RenderCaptureOptions, EditAnywhere, meta = (ClampMin = "0.001", ClampMax = "1000.0"), AdvancedDisplay)
 	float NearPlaneDist = 1.0f;
 
 	bool operator==(const URenderCaptureProperties& Other) const
 	{
+		// Intentionally omitted TransientToolProperty
 		return Resolution == Other.Resolution
+			&& bBaseColorMap == Other.bBaseColorMap
 			&& bNormalMap == Other.bNormalMap
 			&& bMetallicMap == Other.bMetallicMap
 			&& bRoughnessMap == Other.bRoughnessMap
@@ -153,7 +149,7 @@ class MESHMODELINGTOOLSEXP_API UBakeRenderCaptureToolProperties : public UIntera
 public:
 	
 	/** The map type to preview */
-	UPROPERTY(EditAnywhere, Category = BakeOutput, meta=(DisplayName="Preview Output Type", TransientToolProperty, GetOptions = GetMapPreviewNamesFunc))
+	UPROPERTY(EditAnywhere, Category = BakeOutput, meta = (DisplayName="Preview Output Type", TransientToolProperty, GetOptions = GetMapPreviewNamesFunc))
 	FString MapPreview;
 	
 	UFUNCTION()
@@ -169,7 +165,7 @@ public:
 	EBakeTextureSamplesPerPixel SamplesPerPixel = EBakeTextureSamplesPerPixel::Sample1;
 
 	/* Size of generated textures */
-	UPROPERTY(EditAnywhere, Category = BakeOutput, DisplayName="Results Resolution", meta=(ClampMin="1", UIMin="1"))
+	UPROPERTY(EditAnywhere, Category = BakeOutput, meta = (ClampMin="1", UIMin="1"), DisplayName="Results Resolution")
 	EBakeTextureResolution TextureSize = EBakeTextureResolution::Resolution512;
 };
 
@@ -185,10 +181,15 @@ public:
 	TObjectPtr<UStaticMesh> TargetStaticMesh = nullptr;
 
 	/** UV channel to use for the target mesh */
-	UPROPERTY(EditAnywhere, Category = BakeInput, meta = (DisplayName = "Target Mesh UV Channel",
-		GetOptions = GetTargetUVLayerNamesFunc, TransientToolProperty, NoResetToDefault))
+	UPROPERTY(EditAnywhere, Category = BakeInput, meta = (DisplayName = "Target Mesh UV Channel", GetOptions = GetTargetUVLayerNamesFunc, NoResetToDefault))
 	FString TargetUVLayer;
 
+	UFUNCTION()
+	int32 GetTargetUVLayerIndex() const
+	{
+		return TargetUVLayerNamesList.IndexOfByKey(TargetUVLayer);
+	}
+
 	UFUNCTION()
 	const TArray<FString>& GetTargetUVLayerNamesFunc() const
 	{
@@ -216,6 +217,7 @@ public:
 
 	// Begin UInteractiveTool interface
 	virtual void Setup() override;
+	virtual void Render(IToolsContextRenderAPI* RenderAPI) override;
 	virtual void OnShutdown(EToolShutdownType ShutdownType) override;
 	virtual bool HasCancel() const override { return true; }
 	virtual bool HasAccept() const override { return true; }
@@ -256,16 +258,6 @@ protected:
 	
 	void InvalidateResults();
 
-	class FComputeFactory : public UE::Geometry::IGenericDataOperatorFactory<FBakeRenderCaptureResultsBuilder>
-	{
-	public:
-		UBakeRenderCaptureTool* Tool;
-
-		// Begin IGenericDataOperatorFactory interface
-		virtual TUniquePtr<UE::Geometry::TGenericDataOperator<FBakeRenderCaptureResultsBuilder>> MakeNewOperator() override;
-		// End IGenericDataOperatorFactory interface
-	};
-
 	// In this tool we don't call UBakeMeshAttributeMapsToolBase::OnMapsUpdated because it would require e.g, adding
 	// the render capture channels to EBakeMapType.  The implementation is simpler and leads to less coupling if we
 	// just implement custom versions of the following functions.
@@ -304,6 +296,7 @@ protected:
 	// the capture are reverted to these values
 	UPROPERTY()
 	TObjectPtr<URenderCaptureProperties> ComputedRenderCaptureProperties;
+	TMap<int, FText> TargetUVLayerToError;
 
 	//
 	// Analytics
diff --git a/Engine/Plugins/Runtime/MeshModelingToolset/Source/ModelingComponents/Private/Scene/SceneCapturePhotoSet.cpp b/Engine/Plugins/Runtime/MeshModelingToolset/Source/ModelingComponents/Private/Scene/SceneCapturePhotoSet.cpp
--- a/Engine/Plugins/Runtime/MeshModelingToolset/Source/ModelingComponents/Private/Scene/SceneCapturePhotoSet.cpp
+++ b/Engine/Plugins/Runtime/MeshModelingToolset/Source/ModelingComponents/Private/Scene/SceneCapturePhotoSet.cpp
@@ -169,6 +169,13 @@ void FSceneCapturePhotoSet::AddExteriorCaptures(
 		ViewFrame.Origin = (FVector3d)RenderSphere.Center;
 		ViewFrame.Origin -= (double)RenderSphere.W * ViewFrame.X();
 
+		FSpatialPhotoParams Params;
+		Params.Frame = ViewFrame;
+		Params.NearPlaneDist = NearPlaneDist;
+		Params.HorzFOVDegrees = HorizontalFOVDegrees;
+		Params.Dimensions = PhotoDimensions;
+		PhotoSetParams.Add(Params);
+
 		FSpatialPhoto3f BasePhoto3f;
 		BasePhoto3f.Frame = ViewFrame;
 		BasePhoto3f.NearPlaneDist = NearPlaneDist;
@@ -306,7 +313,59 @@ bool FSceneCapturePhotoSet::ComputeSampleLocation(
 	int& PhotoIndex,
 	FVector2d& PhotoCoords) const
 {
-	return BaseColorPhotoSet.ComputeSampleLocation(Position, Normal, VisibilityFunction, PhotoIndex, PhotoCoords);
+	double DotTolerance = -0.1;		// dot should be negative for normal pointing towards photo
+
+	PhotoIndex = IndexConstants::InvalidID;
+	PhotoCoords = FVector2d(0., 0.);
+
+	double MinDot = 1.0;
+
+	int32 NumPhotos = PhotoSetParams.Num();
+	for (int32 pi = 0; pi < NumPhotos; ++pi)
+	{
+		const FSpatialPhotoParams& Params = PhotoSetParams[pi];
+		check(Params.Dimensions.IsSquare());
+
+		FVector3d ViewDirection = Params.Frame.X();
+		double ViewDot = ViewDirection.Dot(Normal);
+		if (ViewDot > DotTolerance || ViewDot > MinDot)
+		{
+			continue;
+		}
+
+		FFrame3d ViewPlane = Params.Frame;
+		ViewPlane.Origin += Params.NearPlaneDist * ViewDirection;
+
+		double ViewPlaneWidthWorld = Params.NearPlaneDist * FMathd::Tan(Params.HorzFOVDegrees * 0.5 * FMathd::DegToRad);
+		double ViewPlaneHeightWorld = ViewPlaneWidthWorld;
+
+		FVector3d RayOrigin = Params.Frame.Origin;
+		FVector3d RayDir = Normalized(Position - RayOrigin);
+		FVector3d HitPoint;
+		bool bHit = ViewPlane.RayPlaneIntersection(RayOrigin, RayDir, 0, HitPoint);
+		if (bHit)
+		{
+			bool bVisible = VisibilityFunction(Position, HitPoint);
+			if ( bVisible )
+			{
+				double PlaneX = (HitPoint - ViewPlane.Origin).Dot(ViewPlane.Y());
+				double PlaneY = (HitPoint - ViewPlane.Origin).Dot(ViewPlane.Z());
+
+				//FVector2d PlanePos = ViewPlane.ToPlaneUV(HitPoint, 0);
+				double u = PlaneX / ViewPlaneWidthWorld;
+				double v = -(PlaneY / ViewPlaneHeightWorld);
+				if (FMathd::Abs(u) < 1 && FMathd::Abs(v) < 1)
+				{
+					PhotoCoords.X = (u/2.0 + 0.5) * (double)Params.Dimensions.GetWidth();
+					PhotoCoords.Y = (v/2.0 + 0.5) * (double)Params.Dimensions.GetHeight();
+					PhotoIndex = pi;
+					MinDot = ViewDot;
+				}
+			}
+		}
+	}
+
+	return PhotoIndex != IndexConstants::InvalidID;
 }
 
 bool FSceneCapturePhotoSet::ComputeSample(
diff --git a/Engine/Plugins/Runtime/MeshModelingToolset/Source/ModelingComponents/Public/Scene/SceneCapturePhotoSet.h b/Engine/Plugins/Runtime/MeshModelingToolset/Source/ModelingComponents/Public/Scene/SceneCapturePhotoSet.h
--- a/Engine/Plugins/Runtime/MeshModelingToolset/Source/ModelingComponents/Public/Scene/SceneCapturePhotoSet.h
+++ b/Engine/Plugins/Runtime/MeshModelingToolset/Source/ModelingComponents/Public/Scene/SceneCapturePhotoSet.h
@@ -183,6 +183,8 @@ protected:
 	FSpatialPhotoSet3f WorldNormalPhotoSet;
 	FSpatialPhotoSet3f EmissivePhotoSet;
 
+	TArray<FSpatialPhotoParams> PhotoSetParams;
+
 	bool bWriteDebugImages = false;
 	FString DebugImagesFolderName = TEXT("SceneCapturePhotoSet");
 
diff --git a/Engine/Source/Runtime/GeometryCore/Public/Image/SpatialPhotoSet.h b/Engine/Source/Runtime/GeometryCore/Public/Image/SpatialPhotoSet.h
--- a/Engine/Source/Runtime/GeometryCore/Public/Image/SpatialPhotoSet.h
+++ b/Engine/Source/Runtime/GeometryCore/Public/Image/SpatialPhotoSet.h
@@ -13,6 +13,21 @@ namespace UE
 namespace Geometry
 {
 
+struct FSpatialPhotoParams
+{
+	/** Coordinate system of the view camera - X() is forward, Z() is up */
+	FFrame3d Frame;
+
+	/** Near-plane distance for the camera, image pixels lie on this plane */
+	double NearPlaneDist = 1.0;
+
+	/** Horizontal Field-of-View of the camera in degrees (full FOV, so generally calculations will use half this value) */
+	double HorzFOVDegrees = 90.0;
+
+	/** Pixel dimensions of the photo image */
+	FImageDimensions Dimensions;
+};
+
 /**
  * TSpatialPhoto represents a 2D image located in 3D space, ie the image plus camera parameters, 
  * which is essentially a "Photograph" of some 3D scene (hence the name)
