MeshModelingTools: Fix crash in XForm tool when using hotkeys while snap dragging.

diff --git a/Engine/Plugins/Experimental/MeshModelingToolsetExp/Source/MeshModelingToolsExp/Private/TransformMeshesTool.cpp b/Engine/Plugins/Experimental/MeshModelingToolsetExp/Source/MeshModelingToolsExp/Private/TransformMeshesTool.cpp
--- a/Engine/Plugins/Experimental/MeshModelingToolsetExp/Source/MeshModelingToolsExp/Private/TransformMeshesTool.cpp
+++ b/Engine/Plugins/Experimental/MeshModelingToolsetExp/Source/MeshModelingToolsExp/Private/TransformMeshesTool.cpp
@@ -130,6 +130,12 @@ void UTransformMeshesTool::UpdateSetPivotModes(bool bEnableSetPivot)
 
 void UTransformMeshesTool::RegisterActions(FInteractiveToolActionSet& ActionSet)
 {
+	// Note: we should really disallow hotkeys while dragging, but we are discouraged from adding
+	// members to an object in a hotfix, so we can't add a bCurrentlySnapDragging (and we can't use
+	// ActiveSnapDragIndex directly because it can be set even without getting capture). For now, we
+	// just make sure to call CheckAndUpdateWatched() after any property changes to avoid some invalid
+	// states that can happen between a drag call and the next tick.
+
 	ActionSet.RegisterAction(this, (int32)EStandardToolActions::BaseClientDefinedActionID + 1,
 		TEXT("ToggleSetPivot"),
 		LOCTEXT("TransformToggleSetPivot", "Toggle Set Pivot"),
@@ -137,6 +143,7 @@ void UTransformMeshesTool::RegisterActions(FInteractiveToolActionSet& ActionSet)
 		EModifierKey::None, EKeys::S,
 		[this]() { 
 			TransformProps->bSetPivotMode = !TransformProps->bSetPivotMode; 
+			TransformProps->CheckAndUpdateWatched();
 			NotifyOfPropertyChangeByTool(TransformProps);
 		});
 
@@ -147,6 +154,7 @@ void UTransformMeshesTool::RegisterActions(FInteractiveToolActionSet& ActionSet)
 		EModifierKey::None, EKeys::D,
 		[this]() { 
 			TransformProps->bEnableSnapDragging = !TransformProps->bEnableSnapDragging;
+			TransformProps->CheckAndUpdateWatched();
 			NotifyOfPropertyChangeByTool(TransformProps);
 		});
 
@@ -157,6 +165,7 @@ void UTransformMeshesTool::RegisterActions(FInteractiveToolActionSet& ActionSet)
 		EModifierKey::None, EKeys::A,
 		[this]() { 
 			TransformProps->TransformMode = (ETransformMeshesTransformMode)(((uint8)TransformProps->TransformMode+1) % (uint8)ETransformMeshesTransformMode::LastValue);
+			TransformProps->CheckAndUpdateWatched();
 			NotifyOfPropertyChangeByTool(TransformProps);
 		});
 
@@ -168,6 +177,7 @@ void UTransformMeshesTool::RegisterActions(FInteractiveToolActionSet& ActionSet)
 		EModifierKey::None, EKeys::W,
 		[this]() { 
 			TransformProps->SnapDragSource = (ETransformMeshesSnapDragSource)(((uint8)TransformProps->SnapDragSource + 1) % (uint8)ETransformMeshesSnapDragSource::LastValue);
+			TransformProps->CheckAndUpdateWatched();
 			NotifyOfPropertyChangeByTool(TransformProps);
 		});
 
@@ -178,6 +188,7 @@ void UTransformMeshesTool::RegisterActions(FInteractiveToolActionSet& ActionSet)
 		EModifierKey::None, EKeys::E,
 		[this]() { 
 			TransformProps->RotationMode = (ETransformMeshesSnapDragRotationMode)(((uint8)TransformProps->RotationMode + 1) % (uint8)ETransformMeshesSnapDragRotationMode::LastValue); 
+			TransformProps->CheckAndUpdateWatched();
 			NotifyOfPropertyChangeByTool(TransformProps);
 		});
 
