ModelingTools: add support for symmetric sculpting to MeshVertexSculptTool. Symmetry detection is based on FMeshPlanarSymmetry, detects point-set-symmetry (rather than topological symmetry). Currently limited to symmetry in local X or Y axis (auto-detected, no user control). Current implementation is based on applying stamp and then enforcing symmetry constraints, which can have some undesirable behavior when sculpting near the symmetry plane.

UMeshSculptToolBase::UpdateHoverStamp() now properly updates LocalFrame. Brush Indicator now takes it's position from HoverStamp instead of CurrentStamp during a stroke.

diff --git a/Engine/Plugins/Experimental/MeshModelingToolsetExp/Source/MeshModelingToolsExp/Private/MeshVertexSculptTool.cpp b/Engine/Plugins/Experimental/MeshModelingToolsetExp/Source/MeshModelingToolsExp/Private/MeshVertexSculptTool.cpp
--- a/Engine/Plugins/Experimental/MeshModelingToolsetExp/Source/MeshModelingToolsExp/Private/MeshVertexSculptTool.cpp
+++ b/Engine/Plugins/Experimental/MeshModelingToolsetExp/Source/MeshModelingToolsExp/Private/MeshVertexSculptTool.cpp
@@ -3,6 +3,7 @@
 #include "MeshVertexSculptTool.h"
 #include "InteractiveToolManager.h"
 #include "InteractiveGizmoManager.h"
+#include "ToolDataVisualizer.h"
 #include "Async/ParallelFor.h"
 #include "Async/Async.h"
 #include "Selections/MeshConnectedComponents.h"
@@ -10,6 +11,7 @@
 #include "MeshWeights.h"
 #include "DynamicMesh/MeshNormals.h"
 #include "DynamicMesh/MeshIndexUtil.h"
+#include "Parameterization/MeshPlanarSymmetry.h"
 #include "Util/BufferUtil.h"
 #include "AssetUtils/Texture2DUtil.h"
 #include "ToolSetupUtil.h"
@@ -58,6 +60,20 @@ UMeshSurfacePointTool* UMeshVertexSculptToolBuilder::CreateNewTool(const FToolBu
 	return SculptTool;
 }
 
+
+/*
+ * internal Change classes
+ */
+
+class FVertexSculptNonSymmetricChange : public FToolCommandChange
+{
+public:
+	virtual void Apply(UObject* Object) override;
+	virtual void Revert(UObject* Object) override;
+};
+
+
+
 /*
  * Tool
  */
@@ -127,6 +143,11 @@ void UMeshVertexSculptTool::Setup()
 		}
 	});
 
+	TFuture<void> InitializeSymmetry = Async(VertexSculptToolAsyncExecTarget, [SculptMesh, this]()
+	{
+		TryToInitializeSymmetry();
+	});
+
 	// currently only supporting default polygroup set
 	TFuture<void> InitializeGroups = Async(VertexSculptToolAsyncExecTarget, [SculptMesh, this]()
 	{
@@ -156,6 +177,7 @@ void UMeshVertexSculptTool::Setup()
 	InitializeGroups.Wait();
 	InitializeBaseMesh.Wait();
 	InitializeRenderDecomp.Wait();
+	InitializeSymmetry.Wait();
 
 	// initialize brush radius range interval, brush properties
 	UMeshSculptToolBase::InitializeBrushSizeRange(Bounds);
@@ -181,6 +203,11 @@ void UMeshVertexSculptTool::Setup()
 	AlphaProperties->RestoreProperties(this);
 	AddToolPropertySource(AlphaProperties);
 
+	SymmetryProperties = NewObject<UMeshSymmetryProperties>(this);
+	SymmetryProperties->RestoreProperties(this);
+	SymmetryProperties->bSymmetryCanBeEnabled = false;
+	AddToolPropertySource(SymmetryProperties);
+
 	this->BaseMeshQueryFunc = [&](int32 VertexID, const FVector3d& Position, double MaxDist, FVector3d& PosOut, FVector3d& NormalOut)
 	{
 		return GetBaseMeshNearest(VertexID, Position, MaxDist, PosOut, NormalOut);
@@ -282,6 +309,13 @@ void UMeshVertexSculptTool::Setup()
 	SetActiveSecondaryBrushType((int32)EMeshVertexSculptBrushType::Smooth);
 
 	StampRandomStream = FRandomStream(31337);
+
+	// update symmetry state based on validity, and then update internal apply-symmetry state
+	SymmetryProperties->bSymmetryCanBeEnabled = bMeshSymmetryIsValid;
+	bApplySymmetry = bMeshSymmetryIsValid && SymmetryProperties->bEnableSymmetry;
+
+	SymmetryProperties->WatchProperty(SymmetryProperties->bEnableSymmetry,
+		[this](bool bNewValue) { bApplySymmetry = bMeshSymmetryIsValid && bNewValue; });
 }
 
 void UMeshVertexSculptTool::Shutdown(EToolShutdownType ShutdownType)
@@ -293,6 +327,7 @@ void UMeshVertexSculptTool::Shutdown(EToolShutdownType ShutdownType)
 
 	SculptProperties->SaveProperties(this);
 	AlphaProperties->SaveProperties(this);
+	SymmetryProperties->SaveProperties(this);
 
 	if (PreviewMeshActor != nullptr)
 	{
@@ -410,6 +445,14 @@ void UMeshVertexSculptTool::OnBeginStroke(const FRay& WorldRay)
 	LastStamp.Power = GetActivePressure() * GetCurrentBrushStrength();
 	LastStamp.TimeStamp = FDateTime::Now();
 
+	// If applying symmetry, make sure the stamp is on the "positive" side. 
+	if (bApplySymmetry)
+	{
+		LastStamp.LocalFrame = Symmetry->GetPositiveSideFrame(LastStamp.LocalFrame);
+		LastStamp.WorldFrame = LastStamp.LocalFrame;
+		LastStamp.WorldFrame.Transform(CurTargetTransform);
+	}
+
 	InitialStrokeTriangleID = -1;
 	InitialStrokeTriangleID = GetBrushTriangleID();
 
@@ -536,7 +579,27 @@ void UMeshVertexSculptTool::UpdateROI(const FVector3d& BrushPos)
 				TriangleROIBuilder.Add(tid);
 			}
 		}
+
 		VertexROIBuilder.SwapValuesWith(VertexROI);
+
+		if (bApplySymmetry)
+		{
+			// Find symmetric Vertex ROI. This will overlap with VertexROI in many cases.
+			SymmetricVertexROI.Reset();
+			Symmetry->GetMirrorVertexROI(VertexROI, SymmetricVertexROI, true);
+			// expand the Triangle ROI to include the symmetric vertex one-rings
+			for (int32 VertexID : SymmetricVertexROI)
+			{
+				if (Mesh->IsVertex(VertexID))
+				{
+					Mesh->EnumerateVertexTriangles(VertexID, [&](int32 tid)
+					{
+						TriangleROIBuilder.Add(tid);
+					});
+				}
+			}
+		}
+
 		TriangleROIBuilder.SwapValuesWith(TriangleROIArray);
 	}
 
@@ -544,6 +607,8 @@ void UMeshVertexSculptTool::UpdateROI(const FVector3d& BrushPos)
 	// In this path we combine everything into one loop. Does fewer distance checks
 	// but nothing can be done in parallel (would change if ROIBuilders had atomic-try-add)
 
+	check(bApplySymmetry == false);		// have not implemented this path yet...
+
 	// TODO would need to support these, this branch is likely dead though
 	ensure(SculptProperties->BrushFilter == EMeshVertexSculptBrushFilterType::None);
 
@@ -579,13 +644,28 @@ void UMeshVertexSculptTool::UpdateROI(const FVector3d& BrushPos)
 #endif
 
 	{
+		// set up and populate position buffers for Vertex ROI
 		TRACE_CPUPROFILER_EVENT_SCOPE(DynamicMeshSculptTool_UpdateROI_4ROI);
-		ROIPositionBuffer.SetNum(VertexROI.Num(), false);
-		ROIPrevPositionBuffer.SetNum(VertexROI.Num(), false);
-		ParallelFor(VertexROI.Num(), [&](int i)
+		int32 ROISize = VertexROI.Num();
+		ROIPositionBuffer.SetNum(ROISize, false);
+		ROIPrevPositionBuffer.SetNum(ROISize, false);
+		ParallelFor(ROISize, [&](int i)
 		{
 			ROIPrevPositionBuffer[i] = Mesh->GetVertexRef(VertexROI[i]);
 		});
+		// do the same for the Symmetric Vertex ROI
+		if (bApplySymmetry)
+		{
+			SymmetricROIPositionBuffer.SetNum(ROISize, false);
+			SymmetricROIPrevPositionBuffer.SetNum(ROISize, false);
+			ParallelFor(ROISize, [&](int i)
+			{
+				if ( Mesh->IsVertex(SymmetricVertexROI[i]) )
+				{
+					SymmetricROIPrevPositionBuffer[i] = Mesh->GetVertexRef(SymmetricVertexROI[i]);
+				}
+			});
+		}
 	}
 }
 
@@ -636,6 +716,13 @@ bool UMeshVertexSculptTool::UpdateStampPosition(const FRay& WorldRay)
 		CurrentStamp.LocalFrame.Rotate(FQuaterniond(CurrentStamp.LocalFrame.Z(), UseAngle, true));
 	}
 
+	if (bApplySymmetry)
+	{
+		CurrentStamp.LocalFrame = Symmetry->GetPositiveSideFrame(CurrentStamp.LocalFrame);
+		CurrentStamp.WorldFrame = CurrentStamp.LocalFrame;
+		CurrentStamp.WorldFrame.Transform(CurTargetTransform);
+	}
+
 	CurrentStamp.PrevLocalFrame = LastStamp.LocalFrame;
 	CurrentStamp.PrevWorldFrame = LastStamp.WorldFrame;
 
@@ -680,6 +767,29 @@ TFuture<void> UMeshVertexSculptTool::ApplyStamp()
 	// can discard alpha now
 	CurrentStamp.StampAlphaFunc = nullptr;
 
+	// if we are applying symmetry, we need to update the on-plane positions as they
+	// will not be in the SymmetricVertexROI
+	if (bApplySymmetry)
+	{
+		// update position of vertices that are on the symmetry plane
+		Symmetry->ApplySymmetryPlaneConstraints(VertexROI, ROIPositionBuffer);
+
+		// currently something gross is that VertexROI/ROIPositionBuffer may have both a vertex and it's mirror vertex,
+		// each with a different position. We somehow need to be able to resolve this, but we don't have the mapping 
+		// between the two locations in VertexROI, and we have no way to figure out the 'new' position of that mirror vertex
+		// until we can look it up by VertexID, not array-index. So, we are going to bake in the new vertex positions for now.
+		const int32 NumV = ROIPositionBuffer.Num();
+		ParallelFor(NumV, [&](int32 k)
+		{
+			int VertIdx = VertexROI[k];
+			const FVector3d& NewPos = ROIPositionBuffer[k];
+			Mesh->SetVertex(VertIdx, NewPos, false);
+		});
+
+		// compute all the mirror vertex positions
+		Symmetry->ComputeSymmetryConstrainedPositions(VertexROI, SymmetricVertexROI, ROIPositionBuffer, SymmetricROIPositionBuffer);
+	}
+
 	// once stamp is applied, we can start updating vertex change, which can happen async as we saved all necessary info
 	TFuture<void> SaveVertexFuture;
 	if (ActiveVertexChange != nullptr)
@@ -693,19 +803,53 @@ TFuture<void> UMeshVertexSculptTool::ApplyStamp()
 				int VertIdx = VertexROI[k];
 				ActiveVertexChange->UpdateVertex(VertIdx, ROIPrevPositionBuffer[k], ROIPositionBuffer[k]);
 			}
+
+			if (bApplySymmetry)
+			{
+				int32 NumSymV = SymmetricVertexROI.Num();
+				for (int32 k = 0; k < NumSymV; ++k)
+				{
+					if (SymmetricVertexROI[k] >= 0)
+					{
+						ActiveVertexChange->UpdateVertex(SymmetricVertexROI[k], SymmetricROIPrevPositionBuffer[k], SymmetricROIPositionBuffer[k]);
+					}
+				}
+			}
+
 		});
 	}
 
-	// now actually update the mesh, which happens 
+	// now actually update the mesh, which happens on the game thread
 	{
 		TRACE_CPUPROFILER_EVENT_SCOPE(VtxSculptTool_ApplyStamp_Sync);
 		const int32 NumV = ROIPositionBuffer.Num();
-		ParallelFor(NumV, [&](int32 k)
+
+		// If we are applying symmetry, we already baked these positions in in the branch above and
+		// can skip it now, otherwise we update the mesh (todo: profile ParallelFor here, is it helping or hurting?)
+		if (bApplySymmetry == false)
 		{
-			int VertIdx = VertexROI[k];
-			const FVector3d& NewPos = ROIPositionBuffer[k];
-			Mesh->SetVertex(VertIdx, NewPos, false);
-		});
+			ParallelFor(NumV, [&](int32 k)
+			{
+				int VertIdx = VertexROI[k];
+				const FVector3d& NewPos = ROIPositionBuffer[k];
+				Mesh->SetVertex(VertIdx, NewPos, false);
+			});
+		}
+
+		// if applying symmetry, bake in new symmetric positions
+		if (bApplySymmetry)
+		{
+			ParallelFor(NumV, [&](int32 k)
+			{
+				int VertIdx = SymmetricVertexROI[k];
+				if (Mesh->IsVertex(VertIdx))
+				{
+					const FVector3d& NewPos = SymmetricROIPositionBuffer[k];
+					Mesh->SetVertex(VertIdx, NewPos, false);
+				}
+			});
+		}
+
 		Mesh->UpdateChangeStamps(true, false);
 	}
 
@@ -793,9 +937,9 @@ bool UMeshVertexSculptTool::UpdateBrushPosition(const FRay& WorldRay)
 
 
 
-void UMeshVertexSculptTool::UpdateHoverStamp(const FFrame3d& StampFrame)
+void UMeshVertexSculptTool::UpdateHoverStamp(const FFrame3d& StampFrameWorld)
 {
-	FFrame3d HoverFrame = StampFrame;
+	FFrame3d HoverFrame = StampFrameWorld;
 	if (bHaveBrushAlpha && (AlphaProperties->RotationAngle != 0))
 	{
 		HoverFrame.Rotate(FQuaterniond(HoverFrame.Z(), AlphaProperties->RotationAngle, true));
@@ -861,6 +1005,22 @@ bool UMeshVertexSculptTool::OnUpdateHover(const FInputDeviceRay& DevicePos)
 }
 
 
+void UMeshVertexSculptTool::Render(IToolsContextRenderAPI* RenderAPI)
+{
+	UMeshSculptToolBase::Render(RenderAPI);
+
+	// draw a dot for the symmetric brush stamp position
+	if (bApplySymmetry)
+	{
+		FToolDataVisualizer Visualizer;
+		Visualizer.BeginFrame(RenderAPI);
+		FVector3d MirrorPoint = CurTargetTransform.TransformPosition(
+			Symmetry->GetMirroredPosition(HoverStamp.LocalFrame.Origin));
+		Visualizer.DrawPoint(MirrorPoint, FLinearColor(1.0, 0.1, 0.1, 1), 5.0f, false);
+		Visualizer.EndFrame();
+	}
+}
+
 
 void UMeshVertexSculptTool::OnTick(float DeltaTime)
 {
@@ -1145,6 +1305,46 @@ double UMeshVertexSculptTool::SampleBrushAlpha(const FSculptBrushStamp& Stamp, c
 }
 
 
+void UMeshVertexSculptTool::TryToInitializeSymmetry()
+{
+	// Currently attempting to detect symmetry by checking both X and Y axes, preferring X. 
+	// This is an expensive way to detect symmetry, however it has the advantage of not depending
+	// on mesh topology/connectivity (ie it detects point-set-symmetry, not mesh-symmetry).
+	// In future it may be desirable to provide some control over this...
+
+	FDynamicMeshAABBTree3 SculptMeshAABBTree(GetSculptMesh(), true);
+
+	TUniquePtr<FMeshPlanarSymmetry> SymmetryX = MakeUnique<FMeshPlanarSymmetry>();
+	bool bSymmetryXValid = false;
+	TUniquePtr<FMeshPlanarSymmetry> SymmetryY = MakeUnique<FMeshPlanarSymmetry>();
+	bool bSymmetryYValid = false;
+
+	TFuture<void> ComputeSymmetryX = Async(VertexSculptToolAsyncExecTarget, [this, &SculptMeshAABBTree, &SymmetryX, &bSymmetryXValid]()
+	{
+		bSymmetryXValid = SymmetryX->Initialize( GetSculptMesh(), &SculptMeshAABBTree, FFrame3d(FVector3d::Zero(), FVector3d::UnitX()) );
+	});
+	TFuture<void> ComputeSymmetryY = Async(VertexSculptToolAsyncExecTarget, [this, &SculptMeshAABBTree, &SymmetryY, &bSymmetryYValid]()
+	{
+		bSymmetryYValid = SymmetryY->Initialize( GetSculptMesh(), &SculptMeshAABBTree, FFrame3d(FVector3d::Zero(), FVector3d::UnitY()) );
+	});
+	ComputeSymmetryX.Wait();
+	ComputeSymmetryY.Wait();
+
+	bMeshSymmetryIsValid = false;
+	if (bSymmetryXValid)
+	{
+		Symmetry = MakePimpl<FMeshPlanarSymmetry>();
+		*Symmetry = MoveTemp(*SymmetryX);
+		bMeshSymmetryIsValid = true;
+	}
+	else if (bSymmetryYValid)
+	{
+		Symmetry = MakePimpl<FMeshPlanarSymmetry>();
+		*Symmetry = MoveTemp(*SymmetryY);
+		bMeshSymmetryIsValid = true;
+	}
+}
+
 
 //
 // Change Tracking
@@ -1166,7 +1366,16 @@ void UMeshVertexSculptTool::EndChange()
 		this->WaitForPendingUndoRedo();
 	};
 
+	GetToolManager()->BeginUndoTransaction(LOCTEXT("VertexSculptChange", "Brush Stroke"));
 	GetToolManager()->EmitObjectChange(DynamicMeshComponent, MoveTemp(NewChange), LOCTEXT("VertexSculptChange", "Brush Stroke"));
+	if (bMeshSymmetryIsValid && bApplySymmetry == false)
+	{
+		// if we end a stroke while symmetry is possible but disabled, we now have to assume that symmetry is no longer possible
+		GetToolManager()->EmitObjectChange(this, MakeUnique<FVertexSculptNonSymmetricChange>(), LOCTEXT("DisableSymmetryChange", "Disable Symmetry"));
+		bMeshSymmetryIsValid = false;
+		SymmetryProperties->bSymmetryCanBeEnabled = bMeshSymmetryIsValid;
+	}
+	GetToolManager()->EndUndoTransaction();
 
 	delete ActiveVertexChange;
 	ActiveVertexChange = nullptr;
@@ -1259,4 +1468,29 @@ void UMeshVertexSculptTool::UpdateBrushType(EMeshVertexSculptBrushType BrushType
 
 
 
+void UMeshVertexSculptTool::UndoRedo_RestoreSymmetryPossibleState(bool bSetToValue)
+{
+	bMeshSymmetryIsValid = bSetToValue;
+	SymmetryProperties->bSymmetryCanBeEnabled = bMeshSymmetryIsValid;
+}
+
+
+
+
+
+void FVertexSculptNonSymmetricChange::Apply(UObject* Object)
+{
+	if (Cast<UMeshVertexSculptTool>(Object))
+	{
+		Cast<UMeshVertexSculptTool>(Object)->UndoRedo_RestoreSymmetryPossibleState(false);
+	}
+}
+void FVertexSculptNonSymmetricChange::Revert(UObject* Object)
+{
+	if (Cast<UMeshVertexSculptTool>(Object))
+	{
+		Cast<UMeshVertexSculptTool>(Object)->UndoRedo_RestoreSymmetryPossibleState(true);
+	}
+}
+
 #undef LOCTEXT_NAMESPACE
diff --git a/Engine/Plugins/Experimental/MeshModelingToolsetExp/Source/MeshModelingToolsExp/Private/Sculpting/MeshSculptToolBase.cpp b/Engine/Plugins/Experimental/MeshModelingToolsetExp/Source/MeshModelingToolsExp/Private/Sculpting/MeshSculptToolBase.cpp
--- a/Engine/Plugins/Experimental/MeshModelingToolsetExp/Source/MeshModelingToolsExp/Private/Sculpting/MeshSculptToolBase.cpp
+++ b/Engine/Plugins/Experimental/MeshModelingToolsetExp/Source/MeshModelingToolsExp/Private/Sculpting/MeshSculptToolBase.cpp
@@ -252,7 +252,7 @@ void UMeshSculptToolBase::OnTick(float DeltaTime)
 	if (bInStroke)
 	{
 		BrushIndicator->Update((float)GetCurrentBrushRadius(),
-			CurrentStamp.WorldFrame.ToFTransform(), 1.0f - (float)GetCurrentBrushFalloff());
+			HoverStamp.WorldFrame.ToFTransform(), 1.0f - (float)GetCurrentBrushFalloff());
 	}
 	else
 	{
@@ -734,9 +734,11 @@ void UMeshSculptToolBase::SaveActiveStrokeModifiers()
 }
 
 
-void UMeshSculptToolBase::UpdateHoverStamp(const FFrame3d& StampFrame)
+void UMeshSculptToolBase::UpdateHoverStamp(const FFrame3d& StampFrameWorld)
 {
-	HoverStamp.WorldFrame = StampFrame;
+	HoverStamp.WorldFrame = StampFrameWorld;
+	HoverStamp.LocalFrame = HoverStamp.WorldFrame;
+	HoverStamp.LocalFrame.Transform(CurTargetTransform.InverseUnsafe());
 }
 
 void UMeshSculptToolBase::UpdateStampPendingState()
diff --git a/Engine/Plugins/Experimental/MeshModelingToolsetExp/Source/MeshModelingToolsExp/Public/MeshVertexSculptTool.h b/Engine/Plugins/Experimental/MeshModelingToolsetExp/Source/MeshModelingToolsExp/Public/MeshVertexSculptTool.h
--- a/Engine/Plugins/Experimental/MeshModelingToolsetExp/Source/MeshModelingToolsExp/Public/MeshVertexSculptTool.h
+++ b/Engine/Plugins/Experimental/MeshModelingToolsetExp/Source/MeshModelingToolsExp/Public/MeshVertexSculptTool.h
@@ -17,12 +17,13 @@
 #include "Image/ImageBuilder.h"
 #include "Util/UniqueIndexSet.h"
 #include "Polygroups/PolygroupSet.h"
+#include "Templates/PimplPtr.h"
 #include "MeshVertexSculptTool.generated.h"
 
 class UMaterialInstanceDynamic;
 class FMeshVertexChangeBuilder;
 class UPreviewMesh;
-
+PREDECLARE_GEOMETRY(class FMeshPlanarSymmetry);
 
 
 
@@ -172,6 +173,21 @@ public:
 
 
 
+UCLASS()
+class MESHMODELINGTOOLSEXP_API UMeshSymmetryProperties : public UInteractiveToolPropertySet
+{
+	GENERATED_BODY()
+public:
+	/** Enable/Disable symmetric sculpting. This option will not be available if symmetry cannot be detected, or a non-symmetric edit has been made */
+	UPROPERTY(EditAnywhere, Category = Symmetry, meta = (HideEditConditionToggle, EditCondition = bSymmetryCanBeEnabled))
+	bool bEnableSymmetry = true;
+
+	// this flag is set/updated by the Tool to enable/disable the bEnableSymmetry toggle
+	UPROPERTY(meta = (TransientToolProperty))
+	bool bSymmetryCanBeEnabled = false;
+};
+
+
 
 /**
  * Mesh Vertex Sculpt Tool Class
@@ -183,6 +199,7 @@ class MESHMODELINGTOOLSEXP_API UMeshVertexSculptTool : public UMeshSculptToolBas
 public:
 	virtual void Setup() override;
 	virtual void Shutdown(EToolShutdownType ShutdownType) override;
+	virtual void Render(IToolsContextRenderAPI* RenderAPI) override;
 
 	virtual void OnTick(float DeltaTime) override;
 
@@ -205,6 +222,9 @@ public:
 	UPROPERTY()
 	TObjectPtr<UTexture2D> BrushAlpha;
 
+	UPROPERTY()
+	TObjectPtr<UMeshSymmetryProperties> SymmetryProperties;
+
 public:
 	virtual void IncreaseBrushSpeedAction() override;
 	virtual void DecreaseBrushSpeedAction() override;
@@ -228,7 +248,7 @@ protected:
 	virtual int32 FindHitTargetMeshTriangle(const FRay3d& LocalRay) override;
 	bool IsHitTriangleBackFacing(int32 TriangleID, const FDynamicMesh3* QueryMesh);
 
-	virtual void UpdateHoverStamp(const FFrame3d& StampFrame) override;
+	virtual void UpdateHoverStamp(const FFrame3d& StampFrameWorld) override;
 
 	virtual void OnBeginStroke(const FRay& WorldRay) override;
 	virtual void OnEndStroke() override;
@@ -311,6 +331,17 @@ protected:
 	TArray<FVector3d> ROIPositionBuffer;
 	TArray<FVector3d> ROIPrevPositionBuffer;
 
+	TPimplPtr<UE::Geometry::FMeshPlanarSymmetry> Symmetry;
+	bool bMeshSymmetryIsValid = false;
+	void TryToInitializeSymmetry();
+	friend class FVertexSculptNonSymmetricChange;
+	virtual void UndoRedo_RestoreSymmetryPossibleState(bool bSetToValue);
+
+	bool bApplySymmetry = false;
+	TArray<int> SymmetricVertexROI;
+	TArray<FVector3d> SymmetricROIPositionBuffer;
+	TArray<FVector3d> SymmetricROIPrevPositionBuffer;
+
 	FMeshVertexChangeBuilder* ActiveVertexChange = nullptr;
 	void BeginChange();
 	void EndChange();
@@ -321,4 +352,3 @@ protected:
 };
 
 
-
diff --git a/Engine/Plugins/Experimental/MeshModelingToolsetExp/Source/MeshModelingToolsExp/Public/Sculpting/MeshSculptToolBase.h b/Engine/Plugins/Experimental/MeshModelingToolsetExp/Source/MeshModelingToolsExp/Public/Sculpting/MeshSculptToolBase.h
--- a/Engine/Plugins/Experimental/MeshModelingToolsetExp/Source/MeshModelingToolsExp/Public/Sculpting/MeshSculptToolBase.h
+++ b/Engine/Plugins/Experimental/MeshModelingToolsetExp/Source/MeshModelingToolsExp/Public/Sculpting/MeshSculptToolBase.h
@@ -466,7 +466,7 @@ protected:
 	FSculptBrushStamp HoverStamp;
 	FSculptBrushStamp CurrentStamp;
 	FSculptBrushStamp LastStamp;
-	virtual void UpdateHoverStamp(const FFrame3d& StampFrame);
+	virtual void UpdateHoverStamp(const FFrame3d& StampFrameWorld);
 	virtual bool IsStampPending() const { return bIsStampPending; }
 	virtual const FRay& GetPendingStampRayWorld() const { return PendingStampRay;  }
 
