fix static analysis warning

diff --git a/Engine/Plugins/Experimental/MeshModelingToolsetExp/Source/MeshModelingToolsExp/Private/Sculpting/MeshSculptToolBase.cpp b/Engine/Plugins/Experimental/MeshModelingToolsetExp/Source/MeshModelingToolsExp/Private/Sculpting/MeshSculptToolBase.cpp
--- a/Engine/Plugins/Experimental/MeshModelingToolsetExp/Source/MeshModelingToolsExp/Private/Sculpting/MeshSculptToolBase.cpp
+++ b/Engine/Plugins/Experimental/MeshModelingToolsetExp/Source/MeshModelingToolsExp/Private/Sculpting/MeshSculptToolBase.cpp
@@ -249,16 +249,9 @@ void UMeshSculptToolBase::OnTick(float DeltaTime)
 
 	UpdateHoverStamp(GetBrushFrameWorld());
 
-	if (bInStroke)
-	{
-		BrushIndicator->Update((float)GetCurrentBrushRadius(),
-			HoverStamp.WorldFrame.ToFTransform(), 1.0f - (float)GetCurrentBrushFalloff());
-	}
-	else
-	{
-		BrushIndicator->Update((float)GetCurrentBrushRadius(),
-			HoverStamp.WorldFrame.ToFTransform(), 1.0f - (float)GetCurrentBrushFalloff());
-	}
+	// always using HoverStamp here because it's position should always be up-to-date for the current apply-stamp...
+	BrushIndicator->Update((float)GetCurrentBrushRadius(),
+		HoverStamp.WorldFrame.ToFTransform(), 1.0f - (float)GetCurrentBrushFalloff());
 
 	UpdateWorkPlane();
 }
