ModelingToolset: remove HairStrands plugin dependency from MeshModelingToolset and Exp, as Hair tools were moved to HairModelingToolset (added HairStrands plugin dependency there, which was missing)

diff --git a/Engine/Plugins/Experimental/HairModelingToolset/HairModelingToolset.uplugin b/Engine/Plugins/Experimental/HairModelingToolset/HairModelingToolset.uplugin
--- a/Engine/Plugins/Experimental/HairModelingToolset/HairModelingToolset.uplugin
+++ b/Engine/Plugins/Experimental/HairModelingToolset/HairModelingToolset.uplugin
@@ -35,6 +35,10 @@
 		{
 			"Name": "ModelingToolsEditorMode",
 			"Enabled": true
+		},
+		{
+			"Name": "HairStrands",
+			"Enabled": true
 		}		
 	]
 }
diff --git a/Engine/Plugins/Experimental/MeshModelingToolsetExp/MeshModelingToolsetExp.uplugin b/Engine/Plugins/Experimental/MeshModelingToolsetExp/MeshModelingToolsetExp.uplugin
--- a/Engine/Plugins/Experimental/MeshModelingToolsetExp/MeshModelingToolsetExp.uplugin
+++ b/Engine/Plugins/Experimental/MeshModelingToolsetExp/MeshModelingToolsetExp.uplugin
@@ -49,10 +49,6 @@
 		{
 			"Name": "ProxyLODPlugin",
 			"Enabled": true
-		},
-		{
-			"Name": "HairStrands",
-			"Enabled": true
 		}
 	]
 }
diff --git a/Engine/Plugins/Runtime/MeshModelingToolset/MeshModelingToolset.uplugin b/Engine/Plugins/Runtime/MeshModelingToolset/MeshModelingToolset.uplugin
--- a/Engine/Plugins/Runtime/MeshModelingToolset/MeshModelingToolset.uplugin
+++ b/Engine/Plugins/Runtime/MeshModelingToolset/MeshModelingToolset.uplugin
@@ -54,10 +54,6 @@
 		{
 			"Name": "ProxyLODPlugin",
 			"Enabled": true
-		},
-		{
-			"Name": "HairStrands",
-			"Enabled": true
 		}
 	]
 }
