ModelingTools: add new Pattern Tool and expose in Modeling Mode. Pattern tool tiles input object(s) in linear/grid/radial patterns, with various transform options, and emits new actors in various ways

diff --git a/Engine/Plugins/Editor/ModelingToolsEditorMode/Source/ModelingToolsEditorMode/Private/ModelingToolsEditorMode.cpp b/Engine/Plugins/Editor/ModelingToolsEditorMode/Source/ModelingToolsEditorMode/Private/ModelingToolsEditorMode.cpp
--- a/Engine/Plugins/Editor/ModelingToolsEditorMode/Source/ModelingToolsEditorMode/Private/ModelingToolsEditorMode.cpp
+++ b/Engine/Plugins/Editor/ModelingToolsEditorMode/Source/ModelingToolsEditorMode/Private/ModelingToolsEditorMode.cpp
@@ -90,6 +90,7 @@
 #include "TransferMeshTool.h"
 #include "ConvertMeshesTool.h"
 #include "SplitMeshesTool.h"
+#include "PatternTool.h"
 
 #include "Physics/PhysicsInspectorTool.h"
 #include "Physics/SetCollisionGeometryTool.h"
@@ -536,6 +537,7 @@ void UModelingToolsEditorMode::Enter()
 	RegisterTool(ToolManagerCommands.BeginTransformUVIslandsTool, TEXT("BeginTransformUVIslandsTool"), NewObject<UEditUVIslandsToolBuilder>());
 	RegisterTool(ToolManagerCommands.BeginLatticeDeformerTool, TEXT("BeginLatticeDeformerTool"), NewObject<ULatticeDeformerToolBuilder>());
 	RegisterTool(ToolManagerCommands.BeginSubdividePolyTool, TEXT("BeginSubdividePolyTool"), NewObject<USubdividePolyToolBuilder>());
+	RegisterTool(ToolManagerCommands.BeginPatternTool, TEXT("BeginPatternTool"), NewObject<UPatternToolBuilder>());
 
 	UCombineMeshesToolBuilder* CombineMeshesToolBuilder = NewObject<UCombineMeshesToolBuilder>();
 	RegisterTool(ToolManagerCommands.BeginCombineMeshesTool, TEXT("BeginCombineMeshesTool"), CombineMeshesToolBuilder);
diff --git a/Engine/Plugins/Editor/ModelingToolsEditorMode/Source/ModelingToolsEditorMode/Private/ModelingToolsEditorModeToolkit.cpp b/Engine/Plugins/Editor/ModelingToolsEditorMode/Source/ModelingToolsEditorMode/Private/ModelingToolsEditorModeToolkit.cpp
--- a/Engine/Plugins/Editor/ModelingToolsEditorMode/Source/ModelingToolsEditorMode/Private/ModelingToolsEditorModeToolkit.cpp
+++ b/Engine/Plugins/Editor/ModelingToolsEditorMode/Source/ModelingToolsEditorMode/Private/ModelingToolsEditorModeToolkit.cpp
@@ -648,6 +648,7 @@ void FModelingToolsEditorModeToolkit::BuildToolPalette(FName PaletteIndex, class
 		ToolbarBuilder.AddToolBarButton(Commands.BeginRevolveBoundaryTool);
 		ToolbarBuilder.AddToolBarButton(Commands.BeginCombineMeshesTool);
 		ToolbarBuilder.AddToolBarButton(Commands.BeginDuplicateMeshesTool);
+		ToolbarBuilder.AddToolBarButton(Commands.BeginPatternTool);
 	}
 	else if (PaletteIndex == TransformTabName)
 	{
diff --git a/Engine/Plugins/Editor/ModelingToolsEditorMode/Source/ModelingToolsEditorMode/Private/ModelingToolsManagerActions.cpp b/Engine/Plugins/Editor/ModelingToolsEditorMode/Source/ModelingToolsEditorMode/Private/ModelingToolsManagerActions.cpp
--- a/Engine/Plugins/Editor/ModelingToolsEditorMode/Source/ModelingToolsEditorMode/Private/ModelingToolsManagerActions.cpp
+++ b/Engine/Plugins/Editor/ModelingToolsEditorMode/Source/ModelingToolsEditorMode/Private/ModelingToolsManagerActions.cpp
@@ -97,6 +97,7 @@ void FModelingToolsManagerCommands::RegisterCommands()
 	REGISTER_MODELING_TOOL_COMMAND(BeginTransferMeshTool, "Transfer", "Transfer meshes");
 	REGISTER_MODELING_TOOL_COMMAND(BeginConvertMeshesTool, "Convert", "Convert meshes");
 	REGISTER_MODELING_TOOL_COMMAND(BeginSplitMeshesTool, "Split", "Split meshes");
+	REGISTER_MODELING_TOOL_COMMAND(BeginPatternTool, "Pattern", "Create Patterns of meshes");
 
 	// MeshOps
 	REGISTER_MODELING_TOOL_COMMAND(BeginSimplifyMeshTool, "Simplfy", "Simplify mesh");
diff --git a/Engine/Plugins/Editor/ModelingToolsEditorMode/Source/ModelingToolsEditorMode/Public/ModelingToolsManagerActions.h b/Engine/Plugins/Editor/ModelingToolsEditorMode/Source/ModelingToolsEditorMode/Public/ModelingToolsManagerActions.h
--- a/Engine/Plugins/Editor/ModelingToolsEditorMode/Source/ModelingToolsEditorMode/Public/ModelingToolsManagerActions.h
+++ b/Engine/Plugins/Editor/ModelingToolsEditorMode/Source/ModelingToolsEditorMode/Public/ModelingToolsManagerActions.h
@@ -69,6 +69,8 @@ public:
 	TSharedPtr<FUICommandInfo> BeginTransferMeshTool;
 	TSharedPtr<FUICommandInfo> BeginConvertMeshesTool;
 	TSharedPtr<FUICommandInfo> BeginSplitMeshesTool;
+	TSharedPtr<FUICommandInfo> BeginPatternTool;
+
 
 	TSharedPtr<FUICommandInfo> BeginRemeshSculptMeshTool;
 	TSharedPtr<FUICommandInfo> BeginRemeshMeshTool;
diff --git a/Engine/Plugins/Experimental/MeshModelingToolsetExp/Source/MeshModelingToolsExp/Private/PatternTool.cpp b/Engine/Plugins/Experimental/MeshModelingToolsetExp/Source/MeshModelingToolsExp/Private/PatternTool.cpp
new file mode 100644
--- /dev/null
+++ b/Engine/Plugins/Experimental/MeshModelingToolsetExp/Source/MeshModelingToolsExp/Private/PatternTool.cpp
@@ -0,0 +1,1264 @@
+// Copyright Epic Games, Inc. All Rights Reserved.
+
+#include "PatternTool.h"
+#include "InteractiveToolManager.h"
+#include "InteractiveGizmoManager.h"
+#include "Mechanics/DragAlignmentMechanic.h"
+#include "Mechanics/ConstructionPlaneMechanic.h"
+#include "ToolBuilderUtil.h"
+#include "ToolSetupUtil.h"
+#include "DynamicMesh/DynamicMesh3.h"
+#include "DynamicMeshEditor.h"
+#include "BaseBehaviors/ClickDragBehavior.h"
+#include "ToolSceneQueriesUtil.h"
+#include "ModelingToolTargetUtil.h"
+#include "TransformSequence.h"
+#include "Selection/ToolSelectionUtil.h"
+
+#include "ModelingObjectsCreationAPI.h"
+#include "ModelingComponentsSettings.h"
+
+#include "Drawing/PreviewGeometryActor.h"
+#include "Components/PrimitiveComponent.h"
+#include "Components/StaticMeshComponent.h"
+#include "Components/DynamicMeshComponent.h"
+#include "Components/InstancedStaticMeshComponent.h"
+#include "Engine/StaticMeshActor.h"
+
+#include "BaseGizmos/GizmoComponents.h"
+#include "BaseGizmos/TransformGizmoUtil.h"
+
+#include "Engine/World.h"
+
+#include "TargetInterfaces/PrimitiveComponentBackedTarget.h"
+#include "ToolTargetManager.h"
+
+
+using namespace UE::Geometry;
+
+#define LOCTEXT_NAMESPACE "UPatternTool"
+
+
+/*
+ * ToolBuilder
+ */
+
+const FToolTargetTypeRequirements& UPatternToolBuilder::GetTargetRequirements() const
+{
+	static FToolTargetTypeRequirements TypeRequirements(
+		UPrimitiveComponentBackedTarget::StaticClass()
+		);
+	return TypeRequirements;
+}
+
+UMultiSelectionMeshEditingTool* UPatternToolBuilder::CreateNewTool(const FToolBuilderState& SceneState) const
+{
+	return NewObject<UPatternTool>(SceneState.ToolManager);
+}
+
+
+// these pattern generators should be promoted to GeometryProcessing, however they need some
+// work to clean up the API and make them more correct (ie handling of step size seems a bit
+// flaky, particularly for circle patterns)
+
+class FPatternGenerator
+{
+public:
+	enum class ESpacingMode
+	{
+		ByCount = 0,
+		ByStepSize = 1
+	};
+
+public:
+	// input settings
+	FQuaterniond StartRotation = FQuaterniond::Identity();
+	FQuaterniond EndRotation = FQuaterniond::Identity();
+	bool bInterpolateRotation = false;			// if false, only StartRotation is used
+
+	FVector3d StartTranslation = FVector3d::Zero();
+	FVector3d EndTranslation = FVector3d::Zero();
+	bool bInterpolateTranslation = false;		// if false, only StartTranslation is used
+
+	FVector3d StartScale = FVector3d::One();
+	FVector3d EndScale = FVector3d::One();
+	bool bInterpolateScale = false;				// if false, only StartScale is used
+
+public:
+	// current result pattern
+	TArray<FTransformSRT3d> Pattern;
+
+
+public:
+	void ResetPattern()
+	{
+		Pattern.Reset();
+	}
+
+	void AddPatternElement(FTransformSRT3d NewTransform, double Alpha)
+	{
+		ApplyElementTransforms(NewTransform, Alpha);
+		Pattern.Add(NewTransform);
+	}
+
+
+	void ApplyElementTransforms(FTransformSRT3d& Transform, double Alpha)
+	{
+		if (bInterpolateRotation)
+		{
+			FQuaterniond InterpRotation(StartRotation, EndRotation, Alpha);
+			Transform.SetRotation(Transform.GetRotation() * InterpRotation);
+		}
+		else
+		{
+			Transform.SetRotation(Transform.GetRotation() * StartRotation);
+		}
+
+		if (bInterpolateScale)
+		{
+			Transform.SetScale( Transform.GetScale() * Lerp(StartScale, EndScale, Alpha) );
+		}
+		else
+		{
+			Transform.SetScale( Transform.GetScale() * StartScale );
+		}
+
+		if (bInterpolateTranslation)
+		{
+			Transform.SetTranslation( Transform.GetTranslation() + Lerp(StartTranslation, EndTranslation, Alpha) );
+		}
+		else
+		{
+			Transform.SetTranslation( Transform.GetTranslation() + StartTranslation );
+		}
+	}
+
+
+
+	double GetAlpha(double AlphaStep, int32 k, int32 Iterations, bool bForceLastStepToOne)
+	{
+		if (k == 0)
+		{
+			return 0;
+		}
+		else if (bForceLastStepToOne && k == Iterations-1)
+		{
+			return 1.0;
+		}
+		else
+		{
+			return FMathd::Clamp((double)k * AlphaStep, 0.0, 1.0);
+		}
+	}
+
+	void ComputeSteps(ESpacingMode UseSpacingMode, int32 CountIn, int32 StepSizeIn, int32 LengthIn, int32& Iterations, double& AlphaStep, bool& bForceLastStepToOne)
+	{
+		Iterations = 1;
+		AlphaStep = 0.5;
+		bForceLastStepToOne = true;
+
+		if (UseSpacingMode == ESpacingMode::ByCount)
+		{
+			Iterations = CountIn;
+			AlphaStep = 1.0 / FMath::Max(Iterations-1, 1);
+		}
+		else if (UseSpacingMode == ESpacingMode::ByStepSize)
+		{
+			bForceLastStepToOne = false;
+			AlphaStep = StepSizeIn / LengthIn;
+			Iterations = FMath::Clamp( (int)(1.0 / AlphaStep) + 1, 1, 1000);
+		}
+		else
+		{
+			check(false);
+		}
+	}
+
+};
+
+
+
+
+class FLinearPatternGenerator : public FPatternGenerator
+{
+public:
+
+	FFrame3d StartFrame;
+	FFrame3d EndFrame;
+
+	ESpacingMode SpacingMode = ESpacingMode::ByCount;
+
+	int Axis = 0;
+	int32 Count = 5;
+	double StepSize = 1.0;
+
+	// Not used for LineFill
+	ESpacingMode SpacingModeY = ESpacingMode::ByCount;
+	int AxisY = 1;
+	int32 CountY = 5;
+	double StepSizeY = 1.0;
+
+
+	enum class EFillMode
+	{
+		LineFill,
+		RectangleFill
+	};
+	EFillMode FillMode = EFillMode::LineFill;
+
+
+	void UpdatePattern()
+	{
+		ResetPattern();
+
+		if (FillMode == EFillMode::LineFill)
+		{
+			UpdatePattern_LineFill();
+		}
+		else if (FillMode == EFillMode::RectangleFill)
+		{
+			UpdatePattern_RectangleFill();
+		}
+		else
+		{
+			ensure(false);
+			AddPatternElement(StartFrame.ToTransform(), 1.0);
+		}
+	}
+
+
+	void UpdatePattern_LineFill();
+	void UpdatePattern_RectangleFill();
+
+};
+
+
+
+void FLinearPatternGenerator::UpdatePattern_LineFill()
+{
+	double LineLength = Distance(StartFrame.Origin, EndFrame.Origin);
+
+	int32 Iterations = 1;
+	double AlphaStep = 0.5;
+	bool bForceLastStepToOne = true;
+	ComputeSteps(SpacingMode, Count, StepSize, LineLength, Iterations, AlphaStep, bForceLastStepToOne);
+
+	for (int32 k = 0; k < Iterations; ++k)
+	{
+		FFrame3d PatternFrame = StartFrame;
+		double Alpha = 0;
+		if (bForceLastStepToOne && k == Iterations-1)
+		{
+			Alpha = 1;
+			PatternFrame = EndFrame;
+		}
+		else if (k != 0)
+		{
+			Alpha = GetAlpha(AlphaStep, k, Iterations, bForceLastStepToOne);
+			PatternFrame = Lerp(StartFrame, EndFrame, Alpha);
+		}
+
+		AddPatternElement(PatternFrame.ToTransform(), Alpha);
+	}
+}
+
+
+
+
+
+void FLinearPatternGenerator::UpdatePattern_RectangleFill()
+{
+	FVector3d LocalPt = StartFrame.ToFramePoint(EndFrame.Origin);
+	double ExtentX = LocalPt[Axis];
+	double ExtentY = LocalPt[AxisY];
+
+	int32 IterationsX = 1;
+	double AlphaStepX = 0.5;
+	bool bForceLastStepToOneX = true;
+	ComputeSteps(SpacingMode, Count, StepSize, ExtentX, IterationsX, AlphaStepX, bForceLastStepToOneX);
+
+	int32 IterationsY = 1;
+	double AlphaStepY = 0.5;
+	bool bForceLastStepToOneY = true;
+	ComputeSteps(SpacingModeY, CountY, StepSizeY, ExtentY, IterationsY, AlphaStepY, bForceLastStepToOneY);
+
+	FFrame3d MidFrameOrientation = Lerp(StartFrame, EndFrame, 0.5);
+	FFrame3d Frame00 = StartFrame;
+	FFrame3d Frame11 = EndFrame;
+	FFrame3d Frame10 = StartFrame;
+	Frame10.Origin += ExtentX * StartFrame.GetAxis(Axis);
+	Frame10.Rotation = MidFrameOrientation.Rotation;
+	FFrame3d Frame01 = StartFrame;
+	Frame01.Origin += ExtentY * StartFrame.GetAxis(AxisY);
+	Frame01.Rotation = MidFrameOrientation.Rotation;
+
+	for (int32 yi = 0; yi < IterationsY; ++yi)
+	{
+		double AlphaY = GetAlpha(AlphaStepY, yi, IterationsY, bForceLastStepToOneY);
+		FFrame3d Frame0 = Lerp(Frame00, Frame01, AlphaY);
+		FFrame3d Frame1 = Lerp(Frame10, Frame11, AlphaY);
+
+		for (int32 xi = 0; xi < IterationsX; ++xi)
+		{
+			double AlphaX = GetAlpha(AlphaStepX, xi, IterationsX, bForceLastStepToOneX);
+			FFrame3d PatternFrame = Lerp(Frame0, Frame1, AlphaX);
+
+			AddPatternElement(PatternFrame.ToTransform(), AlphaX*AlphaY);
+		}
+	}
+}
+
+
+
+
+
+class FRadialPatternGenerator : public FPatternGenerator
+{
+public:
+	FFrame3d CenterFrame;
+
+	double StartAngleDeg = 0.0;
+	double EndAngleDeg = 360.0;
+	double AngleShift = 0.0;
+
+	double Radius = 100.0;
+
+	bool bOriented = true;
+
+	ESpacingMode SpacingMode = ESpacingMode::ByCount;
+
+	int32 Count = 5;
+	double StepSizeDeg = 1.0;
+
+	enum class EFillMode
+	{
+		CircleFill,
+	};
+	EFillMode FillMode = EFillMode::CircleFill;
+
+
+	void UpdatePattern()
+	{
+		ResetPattern();
+
+		if (FillMode == EFillMode::CircleFill)
+		{
+			UpdatePattern_CircleFill();
+		}
+	}
+
+	void UpdatePattern_CircleFill();
+
+};
+
+
+
+void FRadialPatternGenerator::UpdatePattern_CircleFill()
+{
+	double ArcLengthDeg = FMathd::Abs(EndAngleDeg - StartAngleDeg);
+
+	int32 Iterations = 1;
+	double AlphaStep = 0.5;
+
+	if (SpacingMode == ESpacingMode::ByCount)
+	{
+		Iterations = Count;
+		AlphaStep = 1.0 / FMath::Max(Iterations, 1);
+	}
+	else if (SpacingMode == ESpacingMode::ByStepSize)
+	{
+		AlphaStep = StepSizeDeg / ArcLengthDeg;
+		Iterations = FMath::Clamp( (int)(1.0 / AlphaStep) + 1, 1, 1000);
+	}
+	else
+	{
+		check(false);
+	}
+
+	for (int32 k = 0; k < Iterations; ++k)
+	{
+		double Alpha = (double)k * AlphaStep;
+		if (k == 0)
+		{
+			Alpha = 0;
+		}
+		Alpha = FMathd::Clamp(Alpha, 0.0, 1.0);
+
+		double Angle = (1.0-Alpha)*StartAngleDeg + (Alpha)*EndAngleDeg;
+		Angle += AngleShift;
+		double FrameX = Radius * FMathd::Cos(Angle * FMathd::DegToRad);
+		double FrameY = Radius * FMathd::Sin(Angle * FMathd::DegToRad);
+
+		FFrame3d PatternFrame = CenterFrame;
+		PatternFrame.Origin = CenterFrame.FromPlaneUV(FVector2d(FrameX, FrameY), 2);
+
+		if (bOriented)
+		{
+			FVector3d Axis = Normalized(PatternFrame.Origin - CenterFrame.Origin);
+			if (IsNormalized(Axis))
+			{
+				PatternFrame.ConstrainedAlignAxis(0, Axis, PatternFrame.Z());
+			}
+		}
+
+		AddPatternElement(PatternFrame.ToTransform(), Alpha);
+	}
+
+}
+
+
+
+
+
+/*
+ * Tool
+ * 
+ * TODO: convert dynamic mesh to temporary static mesh to allow instance rendering (does it work?)
+ * TODO: when outputting a merged dynamic mesh, scale/rotation will be baked in, and so there are no limits
+ *       on non-uniform scaling, but it would need to be baked into the preview dynamicmesh components 
+ *       (and for static meshes being converted, static preview would have to become dynamic preview)
+ * TODO: investigate using temporary ISMComponents instead of multiple SMComponents
+ * 
+ */
+
+UPatternTool::UPatternTool()
+{
+}
+
+
+void UPatternTool::Setup()
+{
+	UInteractiveTool::Setup();
+
+	PreviewGeometry = NewObject<UPreviewGeometry>(this);
+	PreviewGeometry->CreateInWorld(GetTargetWorld(), FTransform::Identity);
+
+	// Must be done before creating gizmos, so that we can bind the mechanic to them.
+	DragAlignmentMechanic = NewObject<UDragAlignmentMechanic>(this);
+	DragAlignmentMechanic->Setup(this);
+
+	Settings = NewObject<UPatternToolSettings>();
+	AddToolPropertySource(Settings);
+	Settings->RestoreProperties(this);
+	Settings->WatchProperty(Settings->Shape, [this](EPatternToolShape) { OnShapeUpdated(); } );
+	Settings->WatchProperty(Settings->SingleAxis, [this](EPatternToolSingleAxis) { OnParametersUpdated(); } );
+	Settings->WatchProperty(Settings->SinglePlane, [this](EPatternToolSinglePlane) { OnParametersUpdated(); } );
+	Settings->WatchProperty(Settings->bHideSources, [this](bool bNewValue) { OnSourceVisibilityToggled(!bNewValue); } );
+
+	LinearSettings = NewObject<UPatternTool_LinearSettings>();
+	AddToolPropertySource(LinearSettings);
+	LinearSettings->RestoreProperties(this);
+	SetToolPropertySourceEnabled(LinearSettings, false);
+
+	GridSettings = NewObject<UPatternTool_GridSettings>();
+	AddToolPropertySource(GridSettings);
+	GridSettings->RestoreProperties(this);
+	SetToolPropertySourceEnabled(GridSettings, false);
+
+	RadialSettings = NewObject<UPatternTool_RadialSettings>();
+	AddToolPropertySource(RadialSettings);
+	RadialSettings->RestoreProperties(this);
+	SetToolPropertySourceEnabled(RadialSettings, false);
+
+	RotationSettings = NewObject<UPatternTool_RotationSettings>();
+	AddToolPropertySource(RotationSettings);
+	RotationSettings->RestoreProperties(this);
+	
+	TranslationSettings = NewObject<UPatternTool_TranslationSettings>();
+	AddToolPropertySource(TranslationSettings);
+	TranslationSettings->RestoreProperties(this);
+
+	ScaleSettings = NewObject<UPatternTool_ScaleSettings>();
+	AddToolPropertySource(ScaleSettings);
+	ScaleSettings->RestoreProperties(this);
+
+	OutputSettings = NewObject<UPatternTool_OutputSettings>();
+	AddToolPropertySource(OutputSettings);
+	OutputSettings->RestoreProperties(this);
+
+	InitializeElements();
+	for (const FPatternElement& Element : Elements)
+	{
+		if (Element.SourceStaticMesh != nullptr)
+		{
+			OutputSettings->bHaveStaticMeshes = true;
+		}
+	}
+
+	CurrentStartFrameWorld = FFrame3d(Elements[0].SourceTransform);
+
+	OnShapeUpdated();
+
+	PlaneMechanic = NewObject<UConstructionPlaneMechanic>(this);
+	PlaneMechanic->Setup(this);
+	PlaneMechanic->Initialize( GetTargetWorld(), CurrentStartFrameWorld );
+	PlaneMechanic->OnPlaneChanged.AddLambda([this]() { OnMainFrameUpdated(); });
+
+	SetToolDisplayName(LOCTEXT("ToolName", "Pattern"));
+	GetToolManager()->DisplayMessage(
+		LOCTEXT("OnStartPatternTool", "Create Patterns for the selected Objects"),
+		EToolMessageLevel::UserNotification);
+
+	if (bHaveNonUniformScaleElements)
+	{
+		// todo: this message only applies in certain contexts like if bSeparateActors=true or if 
+		// emitting an ISMC. Should make it dynamic and/or smarter.
+		GetToolManager()->DisplayMessage(LOCTEXT("NonUniformScaleWarning", 
+			"Source Objects have Non-Uniform Scaling, which may prevent Pattern Transforms from working correctly."), EToolMessageLevel::UserWarning);
+	}
+}
+
+
+
+
+
+
+
+void UPatternTool::OnShutdown(EToolShutdownType ShutdownType)
+{
+	PlaneMechanic->Shutdown();
+	PlaneMechanic = nullptr;
+
+	// destroy all the preview components we created
+	for (UPrimitiveComponent* Component : AllComponents)
+	{
+		Component->UnregisterComponent();
+		Component->DestroyComponent();
+	}
+	AllComponents.Reset();
+	PreviewComponents.Reset();
+	StaticMeshPools.Reset();
+	DynamicMeshPools.Reset();
+
+	PreviewGeometry->Disconnect();
+
+	Settings->SaveProperties(this);
+	LinearSettings->SaveProperties(this);
+	GridSettings->SaveProperties(this);
+	RadialSettings->SaveProperties(this);
+	RotationSettings->SaveProperties(this);
+	TranslationSettings->SaveProperties(this);
+	ScaleSettings->SaveProperties(this);
+	OutputSettings->SaveProperties(this);
+
+	DragAlignmentMechanic->Shutdown();
+
+	GetToolManager()->GetPairedGizmoManager()->DestroyAllGizmosByOwner(this);
+
+	OnSourceVisibilityToggled(true);
+
+	if (ShutdownType == EToolShutdownType::Accept)
+	{
+		EmitResults();
+	}
+}
+
+
+
+
+void UPatternTool::Render(IToolsContextRenderAPI* RenderAPI)
+{
+	DragAlignmentMechanic->Render(RenderAPI);
+	PlaneMechanic->Render(RenderAPI);
+
+	if (bPatternNeedsUpdating)
+	{
+		// throttle update rate as it is somewhat expensive
+		double DeltaTime = (FDateTime::Now() - LastPatternUpdateTime).GetTotalSeconds();
+		if (DeltaTime > 0.05)
+		{
+			UpdatePattern();
+			bPatternNeedsUpdating = false;
+			LastPatternUpdateTime = FDateTime::Now();
+		}
+	}
+}
+
+
+
+void UPatternTool::OnSourceVisibilityToggled(bool bVisible)
+{
+	for (int32 k = 0; k < Targets.Num(); ++k)
+	{
+		UE::ToolTarget::SetSourceObjectVisible(Targets[k], bVisible);
+	}
+}
+
+
+void UPatternTool::MarkPatternDirty()
+{
+	if (bPatternNeedsUpdating == false)
+	{
+		bPatternNeedsUpdating = true;
+		LastPatternUpdateTime = FDateTime::Now();
+	}
+}
+
+
+
+
+void UPatternTool::InitializeElements()
+{
+	int32 NumElements = Targets.Num();
+	Elements.SetNum(NumElements);
+
+	for (int32 TargetIdx = 0; TargetIdx < NumElements; TargetIdx++)
+	{
+		FPatternElement& Element = Elements[TargetIdx];
+		Element.TargetIndex = TargetIdx;
+
+		Element.SourceComponent = UE::ToolTarget::GetTargetComponent(Targets[TargetIdx]);
+		Element.SourceMaterials = UE::ToolTarget::GetMaterialSet(Targets[TargetIdx], false).Materials;
+		Element.SourceTransform = UE::ToolTarget::GetLocalToWorldTransform(Targets[TargetIdx]);
+		Element.BaseRotateScale = Element.SourceTransform;
+		Element.BaseRotateScale.SetTranslation(FVector3d::Zero());		// clear translation from base transform
+		Element.SourceTransform.SetRotation(FQuaterniond::Identity());	// clear rotate/scale from source transform, so only location is used
+		Element.SourceTransform.SetScale(FVector::One());
+
+		bHaveNonUniformScaleElements = bHaveNonUniformScaleElements || Element.BaseRotateScale.HasNonUniformScale();
+
+		if (UStaticMeshComponent* StaticMeshComp = Cast<UStaticMeshComponent>(Element.SourceComponent))
+		{
+			Element.SourceStaticMesh = StaticMeshComp->GetStaticMesh();
+		}
+		else if (UDynamicMeshComponent* DynamicMeshComp = Cast<UDynamicMeshComponent>(Element.SourceComponent))
+		{
+			Element.SourceDynamicMesh = DynamicMeshComp->GetDynamicMesh();
+		}
+		else
+		{
+			Element.bValid = false;
+		}
+	}
+
+	PreviewComponents.SetNum(NumElements);
+}
+
+
+
+
+void UPatternTool::OnPropertyModified(UObject* PropertySet, FProperty* Property)
+{
+	if (PropertySet == RadialSettings 
+		|| PropertySet == LinearSettings 
+		|| PropertySet == GridSettings
+		|| PropertySet == RotationSettings
+		|| PropertySet == TranslationSettings
+		|| PropertySet == ScaleSettings )
+	{
+		OnParametersUpdated();
+	}
+}
+
+
+void UPatternTool::OnMainFrameUpdated()
+{
+	CurrentStartFrameWorld = PlaneMechanic->Plane;
+	MarkPatternDirty();
+}
+
+
+void UPatternTool::OnShapeUpdated()
+{
+	bool bLinearSettings = (Settings->Shape == EPatternToolShape::Line);
+	bool bGridSettings = (Settings->Shape == EPatternToolShape::Grid);
+	bool bRadialSettings = (Settings->Shape == EPatternToolShape::Circle);
+
+	SetToolPropertySourceEnabled(LinearSettings, bLinearSettings);
+	SetToolPropertySourceEnabled(GridSettings, bGridSettings);
+	SetToolPropertySourceEnabled(RadialSettings, bRadialSettings);
+
+	OnParametersUpdated();
+}
+
+
+void UPatternTool::OnParametersUpdated()
+{
+	MarkPatternDirty();
+}
+
+
+
+
+void UPatternTool::ResetPreviews()
+{
+	int32 NumElements = Elements.Num();
+	if (PreviewComponents.Num() != NumElements)
+	{
+		return;
+	}
+
+	for (int32 ElemIdx = 0; ElemIdx < NumElements; ++ElemIdx)
+	{
+		FPatternElement& Element = Elements[ElemIdx];
+		if (Element.bValid)
+		{
+			FComponentSet& ElemComponents = PreviewComponents[ElemIdx];
+			if (Element.SourceStaticMesh)
+			{
+				ReturnStaticMeshes(Element, ElemComponents);
+			}
+			else if (Element.SourceDynamicMesh != nullptr)
+			{
+				ReturnDynamicMeshes(Element, ElemComponents);
+			}
+		}
+	}
+}
+
+
+
+static void InitializeGenerator(FPatternGenerator& Generator, UPatternTool* Tool)
+{
+	Generator.bInterpolateRotation = Tool->RotationSettings->bInterpolate;
+	if (Generator.bInterpolateRotation)
+	{
+		Generator.StartRotation = FQuaterniond(Tool->RotationSettings->StartRotation);
+		Generator.EndRotation = FQuaterniond(Tool->RotationSettings->Rotation);
+	}
+	else
+	{
+		Generator.StartRotation = Generator.EndRotation = FQuaterniond(Tool->RotationSettings->Rotation);
+	}
+
+	Generator.bInterpolateTranslation = Tool->TranslationSettings->bInterpolate;
+	if (Generator.bInterpolateTranslation)
+	{
+		Generator.StartTranslation = Tool->TranslationSettings->StartTranslation;
+		Generator.EndTranslation = Tool->TranslationSettings->Translation;
+	}
+	else
+	{
+		Generator.StartTranslation = Generator.EndTranslation = Tool->TranslationSettings->Translation;
+	}
+
+	Generator.bInterpolateScale = Tool->ScaleSettings->bInterpolate;
+	if (Generator.bInterpolateScale)
+	{
+		Generator.StartScale = (Tool->ScaleSettings->bUniform) ? (FVector3d::One()*Tool->ScaleSettings->StartScale) : Tool->ScaleSettings->StartScaleNonUniform;
+		Generator.EndScale = (Tool->ScaleSettings->bUniform) ? (FVector3d::One()*Tool->ScaleSettings->Scale) : Tool->ScaleSettings->ScaleNonUniform;
+	}
+	else
+	{
+		Generator.StartScale = Generator.EndScale = (Tool->ScaleSettings->bUniform) ? (FVector3d::One()*Tool->ScaleSettings->Scale) : Tool->ScaleSettings->ScaleNonUniform;
+	}
+}
+
+
+
+void UPatternTool::GetPatternTransforms_Linear(TArray<UE::Geometry::FTransformSRT3d>& TransformsOut)
+{
+	FLinearPatternGenerator Generator;
+	InitializeGenerator(Generator, this);
+
+	double ExtentX = LinearSettings->Extent;
+
+	Generator.StartFrame = FFrame3d();
+	Generator.Axis = (int32)Settings->SingleAxis;
+	FVector3d AxisX = Generator.StartFrame.GetAxis(Generator.Axis);
+
+	if (LinearSettings->bCentered)
+	{
+		Generator.StartFrame.Origin -= 0.5 * ExtentX * AxisX;
+	}
+
+	Generator.EndFrame = Generator.StartFrame;
+	Generator.EndFrame.Origin += ExtentX * AxisX;
+
+	Generator.FillMode = FLinearPatternGenerator::EFillMode::LineFill;
+
+	Generator.SpacingMode = (FLinearPatternGenerator::ESpacingMode)(int)LinearSettings->SpacingMode;
+	Generator.Count = LinearSettings->Count;
+	Generator.StepSize = LinearSettings->StepSize;
+
+	Generator.UpdatePattern();
+	TransformsOut = MoveTemp(Generator.Pattern);
+}
+
+
+
+void UPatternTool::GetPatternTransforms_Grid(TArray<UE::Geometry::FTransformSRT3d>& TransformsOut)
+{
+	FLinearPatternGenerator Generator;
+	InitializeGenerator(Generator, this);
+
+	double ExtentX = GridSettings->ExtentX;
+	double ExtentY = GridSettings->ExtentY;
+
+	Generator.StartFrame = FFrame3d();
+	switch (Settings->SinglePlane)
+	{
+	case EPatternToolSinglePlane::XYPlane:
+		Generator.Axis = 0; Generator.AxisY = 1; break;
+	case EPatternToolSinglePlane::XZPlane:
+		Generator.Axis = 0; Generator.AxisY = 2; break;
+	case EPatternToolSinglePlane::YZPlane:
+		Generator.Axis = 1; Generator.AxisY = 2; break;
+	}
+	FVector3d AxisX = Generator.StartFrame.GetAxis(Generator.Axis);
+	FVector3d AxisY = Generator.StartFrame.GetAxis(Generator.AxisY);
+
+	if (GridSettings->bCenteredX)
+	{
+		Generator.StartFrame.Origin -= 0.5 * ExtentX * AxisX;
+	}
+	if (GridSettings->bCenteredY)
+	{
+		Generator.StartFrame.Origin -= 0.5 * ExtentY * AxisY;
+	}
+
+	Generator.EndFrame = Generator.StartFrame;
+	Generator.EndFrame.Origin += ExtentX * AxisX + ExtentY * AxisY;
+
+	Generator.FillMode = FLinearPatternGenerator::EFillMode::RectangleFill;
+
+	Generator.SpacingMode = (FLinearPatternGenerator::ESpacingMode)(int)GridSettings->SpacingX;
+	Generator.Count = GridSettings->CountX;
+	Generator.StepSize = GridSettings->StepSizeX;
+
+	Generator.SpacingModeY = (FLinearPatternGenerator::ESpacingMode)(int)GridSettings->SpacingY;
+	Generator.CountY = GridSettings->CountY;
+	Generator.StepSizeY = GridSettings->StepSizeY;
+
+	Generator.UpdatePattern();
+	TransformsOut = MoveTemp(Generator.Pattern);
+}
+
+
+
+
+void UPatternTool::GetPatternTransforms_Radial(TArray<UE::Geometry::FTransformSRT3d>& TransformsOut)
+{
+	FRadialPatternGenerator Generator;
+	InitializeGenerator(Generator, this);
+
+	Generator.CenterFrame = FFrame3d();
+
+	Generator.StartAngleDeg = RadialSettings->StartAngle;
+	Generator.EndAngleDeg = RadialSettings->EndAngle;
+	Generator.AngleShift = RadialSettings->AngleShift;
+
+	Generator.Radius = RadialSettings->Radius;
+
+	Generator.bOriented = RadialSettings->bOriented;
+
+	Generator.SpacingMode = (FRadialPatternGenerator::ESpacingMode)(int)RadialSettings->SpacingMode;
+
+	Generator.FillMode = FRadialPatternGenerator::EFillMode::CircleFill;
+
+	Generator.Count = RadialSettings->Count;
+	Generator.StepSizeDeg = RadialSettings->StepSize;
+
+	Generator.UpdatePattern();
+	TransformsOut = MoveTemp(Generator.Pattern);
+}
+
+
+
+void UPatternTool::UpdatePattern()
+{
+	TArray<FTransformSRT3d> Pattern;
+
+	if (Settings->Shape == EPatternToolShape::Line)
+	{
+		GetPatternTransforms_Linear(Pattern);
+	}
+	if (Settings->Shape == EPatternToolShape::Grid)
+	{
+		GetPatternTransforms_Grid(Pattern);
+	}
+	else if (Settings->Shape == EPatternToolShape::Circle)
+	{
+		GetPatternTransforms_Radial(Pattern);
+	}
+
+	// Return all current preview components in use to the preview component pool
+	ResetPreviews();
+
+	int32 NumPatternItems = Pattern.Num();
+	int32 NumElements = Elements.Num();
+	check(PreviewComponents.Num() == NumElements);
+
+	for (int32 ElemIdx = 0; ElemIdx < NumElements; ++ElemIdx)
+	{
+		FPatternElement& Element = Elements[ElemIdx];
+		FTransformSRT3d ElementTransform = Element.BaseRotateScale;
+		FComponentSet& ElemComponents = PreviewComponents[ElemIdx];
+
+		for (int32 k = 0; k < NumPatternItems; ++k)
+		{
+			UPrimitiveComponent* Component = nullptr;
+			if (Element.SourceDynamicMesh != nullptr)
+			{
+				Component = GetPreviewDynamicMesh(Element);
+			}
+			else if (Element.SourceStaticMesh != nullptr)
+			{
+				Component = GetPreviewStaticMesh(Element);
+			}
+			if (Component != nullptr)
+			{
+				FTransform PatternTransform = (FTransform)Pattern[k];
+				FTransform WorldTransform = (FTransform)ElementTransform * PatternTransform * CurrentStartFrameWorld.ToFTransform();
+				Component->SetWorldTransform( WorldTransform );
+				Component->SetVisibility(true);
+
+				ElemComponents.Components.Add(Component);
+			}
+
+		}
+	}
+
+	// ResetPreviews() does not hide the preview components because changing their visibility
+	// will destroy their SceneProxies, which is expensive, and in most cases the same set of
+	// objects will immediately be shown, re-creating the proxy. So instead they are just all
+	// left visible, and HideReturnedPreviewMeshes() fixes up the visibility on any Components
+	// returned to the pools that were not immediately re-used
+	HideReturnedPreviewMeshes();
+
+	CurrentPattern = MoveTemp(Pattern);
+}
+
+
+
+
+UStaticMeshComponent* UPatternTool::GetPreviewStaticMesh(FPatternElement& Element)
+{
+	FComponentSet* FoundPool = StaticMeshPools.Find(Element.TargetIndex);
+	if (FoundPool == nullptr)
+	{
+		StaticMeshPools.Add(Element.TargetIndex, FComponentSet{});
+		FoundPool = StaticMeshPools.Find(Element.TargetIndex);
+		check(FoundPool != nullptr);
+	}
+
+	if (FoundPool->Components.Num() > 0)
+	{
+		UStaticMeshComponent* FoundComponent = Cast<UStaticMeshComponent>(FoundPool->Components.Pop(false));
+		check(FoundComponent != nullptr);
+		return FoundComponent;
+	}
+
+	UStaticMeshComponent* StaticMeshComp = NewObject<UStaticMeshComponent>(PreviewGeometry->GetActor());
+
+	StaticMeshComp->SetStaticMesh(Element.SourceStaticMesh);
+	for (int32 j = 0; j < Element.SourceMaterials.Num(); ++j)
+	{
+		StaticMeshComp->SetMaterial(j, Element.SourceMaterials[j]);
+	}
+
+	StaticMeshComp->SetupAttachment(PreviewGeometry->GetActor()->GetRootComponent());
+	StaticMeshComp->RegisterComponent();
+
+	AllComponents.Add(StaticMeshComp);
+
+	return StaticMeshComp;
+}
+
+void UPatternTool::ReturnStaticMeshes(FPatternElement& Element, FComponentSet& ComponentSet)
+{
+	if (ComponentSet.Components.Num() == 0) return;
+
+	FComponentSet* FoundPool = StaticMeshPools.Find(Element.TargetIndex);
+	check(FoundPool != nullptr);		// should never happen as we would have created in GetPreviewStaticMesh()
+
+	for (UPrimitiveComponent* Component : ComponentSet.Components)
+	{
+		if (UStaticMeshComponent* StaticMeshComponent = Cast<UStaticMeshComponent>(Component))
+		{
+			FoundPool->Components.Add(StaticMeshComponent);
+		}
+	}
+
+	ComponentSet.Components.Reset();
+}
+
+
+
+
+
+UDynamicMeshComponent* UPatternTool::GetPreviewDynamicMesh(FPatternElement& Element)
+{
+	FComponentSet* FoundPool = DynamicMeshPools.Find(Element.TargetIndex);
+	if (FoundPool == nullptr)
+	{
+		DynamicMeshPools.Add(Element.TargetIndex, FComponentSet{});
+		FoundPool = DynamicMeshPools.Find(Element.TargetIndex);
+		check(FoundPool != nullptr);
+	}
+
+	if (FoundPool->Components.Num() > 0)
+	{
+		UDynamicMeshComponent* FoundComponent = Cast<UDynamicMeshComponent>(FoundPool->Components.Pop(false));
+		check(FoundComponent != nullptr);
+		return FoundComponent;
+	}
+
+	UDynamicMeshComponent* DynamicMeshComp = NewObject<UDynamicMeshComponent>(PreviewGeometry->GetActor());
+
+	DynamicMeshComp->SetGenerateOverlapEvents(false);
+	DynamicMeshComp->SetCollisionProfileName(UCollisionProfile::NoCollision_ProfileName);
+	DynamicMeshComp->CollisionType = ECollisionTraceFlag::CTF_UseSimpleAsComplex;
+
+	DynamicMeshComp->ConfigureMaterialSet(Element.SourceMaterials);
+	DynamicMeshComp->SetupAttachment(PreviewGeometry->GetActor()->GetRootComponent());
+	DynamicMeshComp->RegisterComponent();
+
+	FDynamicMesh3 ElementMeshCopy(Element.SourceDynamicMesh->GetMeshRef());
+	DynamicMeshComp->SetMesh(MoveTemp(ElementMeshCopy));
+
+	AllComponents.Add(DynamicMeshComp);
+
+	return DynamicMeshComp;
+}
+
+void UPatternTool::ReturnDynamicMeshes(FPatternElement& Element, FComponentSet& ComponentSet)
+{
+	if (ComponentSet.Components.Num() == 0) return;
+
+	FComponentSet* FoundPool = DynamicMeshPools.Find(Element.TargetIndex);
+	check(FoundPool != nullptr);		// should never happen as we would have created in GetPreviewStaticMesh()
+
+	for (UPrimitiveComponent* Component : ComponentSet.Components)
+	{
+		if (UDynamicMeshComponent* DynamicMeshComponent = Cast<UDynamicMeshComponent>(Component))
+		{
+			FoundPool->Components.Add(DynamicMeshComponent);
+		}
+	}
+
+	ComponentSet.Components.Reset();
+}
+
+
+
+void UPatternTool::HideReturnedPreviewMeshes()
+{
+	for (TPair<int32, FComponentSet> Pair : StaticMeshPools)
+	{
+		for (UPrimitiveComponent* Component : Pair.Value.Components)
+		{
+			if (Component->IsVisible())
+			{
+				Component->SetVisibility(false);
+			}
+		}
+	}
+	for (TPair<int32, FComponentSet> Pair : DynamicMeshPools)
+	{
+		for (UPrimitiveComponent* Component : Pair.Value.Components)
+		{
+			if (Component->IsVisible())
+			{
+				Component->SetVisibility(false);
+			}
+		}
+	}
+}
+
+
+
+
+void UPatternTool::EmitResults()
+{
+	GetToolManager()->BeginUndoTransaction(LOCTEXT("CreatePattern", "Create Pattern"));
+
+	bool bSeparateActorPerItem = OutputSettings->bSeparateActors;
+	bool bCreateISMForStaticMeshes = OutputSettings->bCreateISMCs;
+	bool bConvertToDynamic = OutputSettings->bConvertToDynamic;
+	const UModelingComponentsSettings* ModelingSettings = GetDefault<UModelingComponentsSettings>();
+	
+	int32 NumElements = Elements.Num();
+	int32 NumPatternItems = CurrentPattern.Num();
+
+	TArray<AActor*> NewActors;		// set of new actors created by operation
+
+	// TODO: investigate use of CopyPropertiesForUnrelatedObjects to transfer settings from source to target Components/Actors
+
+	for (int32 ElemIdx = 0; ElemIdx < NumElements; ++ElemIdx)
+	{
+		FPatternElement& Element = Elements[ElemIdx];
+		if (Element.bValid == false)
+		{
+			continue;
+		}
+		FTransformSRT3d ElementTransform = Element.BaseRotateScale;
+
+		if (Element.SourceDynamicMesh != nullptr || bConvertToDynamic )
+		{
+			FDynamicMesh3 ElementMesh = UE::ToolTarget::GetDynamicMeshCopy(Targets[ElemIdx], true);
+
+			// this lambda creates a new dynamic mesh actor w/ the materials from the current Element
+			auto EmitDynamicMeshActor = [this, &Element, ModelingSettings](FDynamicMesh3&& MoveMesh, FString BaseName, FTransformSRT3d Transform)
+			{
+				FCreateMeshObjectParams NewMeshObjectParams;
+				NewMeshObjectParams.TargetWorld = GetTargetWorld();
+				NewMeshObjectParams.Transform = (FTransform)Transform;
+				NewMeshObjectParams.BaseName = BaseName;
+				NewMeshObjectParams.Materials = Element.SourceMaterials;
+				NewMeshObjectParams.SetMesh(MoveTemp(MoveMesh));
+				NewMeshObjectParams.TypeHint = ECreateObjectTypeHint::DynamicMeshActor;
+
+				NewMeshObjectParams.bEnableCollision = ModelingSettings->bEnableCollision;
+				NewMeshObjectParams.CollisionMode = ModelingSettings->CollisionMode;
+				NewMeshObjectParams.bEnableRaytracingSupport = ModelingSettings->bEnableRayTracing;
+
+				FCreateMeshObjectResult Result = UE::Modeling::CreateMeshObject(GetToolManager(), MoveTemp(NewMeshObjectParams));
+				return Result;
+			};
+
+
+			if (bSeparateActorPerItem)
+			{
+				for (int32 k = 0; k < NumPatternItems; ++k)
+				{
+					FTransformSRT3d PatternTransform = CurrentPattern[k];
+					FTransform WorldTransform = (FTransform)ElementTransform * (FTransform)PatternTransform * CurrentStartFrameWorld.ToFTransform();
+
+					FDynamicMesh3 PatternItemMesh = ElementMesh;
+					// TODO: may need to bake nonuniform scale in here, if we allow scaling in transform
+					FCreateMeshObjectResult Result = EmitDynamicMeshActor(MoveTemp(PatternItemMesh),
+						FString::Printf(TEXT("Pattern_%d_%d"), ElemIdx, k), WorldTransform);
+					if (Result.IsOK())
+					{
+						NewActors.Add(Result.NewActor);
+					}
+				}
+			}
+			else
+			{
+				FDynamicMesh3 CombinedPatternMesh;
+				CombinedPatternMesh.EnableMatchingAttributes(ElementMesh, true);
+				FDynamicMeshEditor Editor(&CombinedPatternMesh);
+				FMeshIndexMappings Mappings;
+				for (int32 k = 0; k < NumPatternItems; ++k)
+				{
+					FTransformSequence3d TransformSeq;
+					TransformSeq.Append(ElementTransform);
+					TransformSeq.Append(CurrentPattern[k]);
+					Mappings.Reset();
+					Editor.AppendMesh(&ElementMesh, Mappings,
+						[&](int, const FVector3d& Position) { return TransformSeq.TransformPosition(Position); },
+						[&](int, const FVector3d& Normal) { return TransformSeq.TransformNormal(Normal); } );
+				}
+
+				FCreateMeshObjectResult Result = EmitDynamicMeshActor(MoveTemp(CombinedPatternMesh),
+					FString::Printf(TEXT("Pattern_%d"), ElemIdx), CurrentStartFrameWorld.ToFTransform());
+				if (Result.IsOK())
+				{
+					NewActors.Add(Result.NewActor);
+				}
+			}
+		}
+		else if (Element.SourceStaticMesh != nullptr)
+		{
+			UStaticMesh* SetStaticMesh = Element.SourceStaticMesh;
+			UStaticMeshComponent* SourceComponent = Cast<UStaticMeshComponent>(Elements[ElemIdx].SourceComponent);
+			AActor* SourceActor = (SourceComponent != nullptr) ? SourceComponent->GetOwner() : nullptr;
+
+			if (bSeparateActorPerItem)
+			{
+				for (int32 k = 0; k < NumPatternItems; ++k)
+				{
+					FTransformSRT3d PatternTransform = CurrentPattern[k];
+					FTransform WorldTransform = (FTransform)ElementTransform * (FTransform)PatternTransform * CurrentStartFrameWorld.ToFTransform();
+
+					FActorSpawnParameters SpawnInfo;
+					SpawnInfo.Template = SourceActor;
+					AStaticMeshActor* NewActor = GetTargetWorld()->SpawnActor<AStaticMeshActor>(SpawnInfo);
+					NewActor->SetActorTransform(WorldTransform);
+					//NewActor->GetStaticMeshComponent()->SetStaticMesh(SetStaticMesh);
+					if (NewActor != nullptr)
+					{
+						NewActors.Add(NewActor);
+					}
+				}
+			}
+			else if (bCreateISMForStaticMeshes)
+			{
+				FActorSpawnParameters SpawnInfo;
+				AActor* NewActor = GetTargetWorld()->SpawnActor<AActor>(SpawnInfo);
+				if (NewActor != nullptr)
+				{
+					NewActors.Add(NewActor);
+
+					UInstancedStaticMeshComponent* ISMComponent = NewObject<UInstancedStaticMeshComponent>(NewActor);
+					ISMComponent->SetFlags(RF_Transactional);
+					ISMComponent->bHasPerInstanceHitProxies = true;
+
+					ISMComponent->SetStaticMesh(Element.SourceStaticMesh);
+					for (int32 j = 0; j < Element.SourceMaterials.Num(); ++j)
+					{
+						ISMComponent->SetMaterial(j, Element.SourceMaterials[j]);
+					}
+
+					NewActor->SetRootComponent(ISMComponent);
+					ISMComponent->OnComponentCreated();
+					NewActor->AddInstanceComponent(ISMComponent);
+
+					NewActor->SetActorTransform(CurrentStartFrameWorld.ToFTransform());
+
+					for (int32 k = 0; k < NumPatternItems; ++k)
+					{
+						FTransformSRT3d PatternTransform = CurrentPattern[k];
+						ISMComponent->AddInstance((FTransform)ElementTransform * (FTransform)PatternTransform);
+					}
+
+					ISMComponent->RegisterComponent();
+				}
+			}
+			else
+			{
+				// Emit a single StaticMeshActor with multiple StaticMeshComponents
+
+				AStaticMeshActor* ParentActor = nullptr;
+				UStaticMeshComponent* TemplateComponent = nullptr;
+				for (int32 k = 0; k < NumPatternItems; ++k)
+				{
+					FTransformSRT3d PatternTransform = CurrentPattern[k];
+					FTransform WorldTransform = (FTransform)PatternTransform * CurrentStartFrameWorld.ToFTransform();
+					if (k == 0)
+					{
+						FActorSpawnParameters SpawnInfo;
+						SpawnInfo.Template = SourceActor;
+						ParentActor = GetTargetWorld()->SpawnActor<AStaticMeshActor>(SpawnInfo);
+						ParentActor->SetActorTransform(CurrentStartFrameWorld.ToFTransform());
+						ParentActor->GetStaticMeshComponent()->SetWorldTransform(WorldTransform);
+					}
+					else if (k == 1)
+					{
+						TemplateComponent = DuplicateObject<UStaticMeshComponent>(ParentActor->GetStaticMeshComponent(), ParentActor);
+						TemplateComponent->ClearFlags(RF_DefaultSubObject);
+						TemplateComponent->SetupAttachment(ParentActor->GetRootComponent());
+						TemplateComponent->OnComponentCreated();
+						ParentActor->AddInstanceComponent(TemplateComponent);
+						TemplateComponent->RegisterComponent();
+						TemplateComponent->SetWorldTransform( WorldTransform );
+					}
+					else
+					{
+						UStaticMeshComponent* NewCloneComponent = DuplicateObject<UStaticMeshComponent>(TemplateComponent, ParentActor);
+						NewCloneComponent->ClearFlags(RF_DefaultSubObject);
+						NewCloneComponent->SetupAttachment(ParentActor->GetRootComponent());
+						NewCloneComponent->OnComponentCreated();
+						ParentActor->AddInstanceComponent(NewCloneComponent);
+						NewCloneComponent->RegisterComponent();
+						NewCloneComponent->SetWorldTransform( WorldTransform );
+					}
+				}			
+
+			}
+		}
+	}
+
+	if (NewActors.Num() > 0)
+	{
+		ToolSelectionUtil::SetNewActorSelection(GetToolManager(), NewActors);
+	}
+
+	GetToolManager()->EndUndoTransaction();
+}
+
+
+
+#undef LOCTEXT_NAMESPACE
diff --git a/Engine/Plugins/Experimental/MeshModelingToolsetExp/Source/MeshModelingToolsExp/Public/PatternTool.h b/Engine/Plugins/Experimental/MeshModelingToolsetExp/Source/MeshModelingToolsExp/Public/PatternTool.h
new file mode 100644
--- /dev/null
+++ b/Engine/Plugins/Experimental/MeshModelingToolsetExp/Source/MeshModelingToolsExp/Public/PatternTool.h
@@ -0,0 +1,488 @@
+// Copyright Epic Games, Inc. All Rights Reserved.
+
+#pragma once
+
+#include "CoreMinimal.h"
+#include "BaseTools/MultiSelectionMeshEditingTool.h"
+#include "InteractiveToolBuilder.h"
+#include "BaseBehaviors/BehaviorTargetInterfaces.h"
+#include "Changes/TransformChange.h"
+#include "FrameTypes.h"
+#include "BoxTypes.h"
+#include "TransformTypes.h"
+#include "PatternTool.generated.h"
+
+class UBaseAxisTranslationGizmo;
+class UAxisAngleGizmo;
+class UDragAlignmentMechanic;
+class UCombinedTransformGizmo;
+class UTransformProxy;
+class UPreviewGeometry;
+class UDynamicMesh;
+class UDynamicMeshComponent;
+class UStaticMesh;
+class UStaticMeshComponent;
+class UMaterialInterface;
+class UConstructionPlaneMechanic;
+
+
+/**
+ *
+ */
+UCLASS()
+class MESHMODELINGTOOLSEXP_API UPatternToolBuilder : public UMultiSelectionMeshEditingToolBuilder
+{
+	GENERATED_BODY()
+
+public:
+	virtual UMultiSelectionMeshEditingTool* CreateNewTool(const FToolBuilderState& SceneState) const override;
+
+protected:
+	virtual const FToolTargetTypeRequirements& GetTargetRequirements() const override;
+};
+
+
+
+
+UENUM()
+enum class EPatternToolShape : uint8
+{
+	/** Arrange pattern elements along a Line */
+	Line = 0,
+
+	/** Arrange pattern elements in a 2D Grid */
+	Grid = 1,
+
+	/** Arrange pattern elements in a Circle */
+	Circle = 2
+};
+
+
+UENUM()
+enum class EPatternToolSingleAxis : uint8
+{
+	XAxis = 0,
+	YAxis = 1,
+	ZAxis = 2
+};
+
+UENUM()
+enum class EPatternToolSinglePlane : uint8
+{
+	XYPlane = 0,
+	XZPlane = 1,
+	YZPlane = 2
+};
+
+
+
+UENUM()
+enum class EPatternToolAxisSpacingMode : uint8
+{
+	/** Place a specific number of Pattern Elements along the pattern geometry */
+	ByCount = 0,
+
+	/** Place Pattern Elements at regular increments along the Pattern Geometry (on-center) */
+	StepSize = 1
+};
+
+
+/**
+ * Settings for the Pattern Tool
+ */
+UCLASS()
+class MESHMODELINGTOOLSEXP_API UPatternToolSettings : public UInteractiveToolPropertySet
+{
+	GENERATED_BODY()
+public:
+	/** Shape of the underlying Pattern */
+	UPROPERTY(EditAnywhere, Category = Shape)
+	EPatternToolShape Shape = EPatternToolShape::Line;
+
+	/** Axis direction used for the Pattern geometry */
+	UPROPERTY(EditAnywhere, Category = Shape, meta = (DisplayName = "Direction", EditCondition = "Shape == EPatternToolShape::Line", EditConditionHides))
+	EPatternToolSingleAxis SingleAxis = EPatternToolSingleAxis::XAxis;
+
+	/** Plane used for the Pattern geometry */
+	UPROPERTY(EditAnywhere, Category = Shape, meta = (DisplayName = "Plane", EditCondition = "Shape != EPatternToolShape::Line", EditConditionHides))
+	EPatternToolSinglePlane SinglePlane = EPatternToolSinglePlane::XYPlane;
+
+	/** Shape of the underlying Pattern */
+	UPROPERTY(EditAnywhere, Category = Shape)
+	bool bHideSources = true;
+
+};
+
+
+/**
+ * Settings for Linear Patterns in the Pattern Tool
+ */
+UCLASS()
+class MESHMODELINGTOOLSEXP_API UPatternTool_LinearSettings : public UInteractiveToolPropertySet
+{
+	GENERATED_BODY()
+public:
+	/** Spacing Technique used to distribute Pattern Elements */
+	UPROPERTY(EditAnywhere, Category = LinearPattern)
+	EPatternToolAxisSpacingMode SpacingMode = EPatternToolAxisSpacingMode::ByCount;
+
+	/** Number of Pattern Elements to place */
+	UPROPERTY(EditAnywhere, Category = LinearPattern, meta = (UIMin = 1, UIMax = 25, EditCondition = "SpacingMode == EPatternToolAxisSpacingMode::ByCount"))
+	int32 Count = 10;
+
+	/** Fixed Increment used to place Pattern Elements */
+	UPROPERTY(EditAnywhere, Category = LinearPattern, meta = (UIMin = 0, EditCondition = "SpacingMode == EPatternToolAxisSpacingMode::StepSize"))
+	double StepSize = 100.0;
+
+	/** Length of Pattern along the Axis */
+	UPROPERTY(EditAnywhere, Category = LinearPattern, meta = (UIMin = 0))
+	double Extent = 1000.0;
+
+	/** If true, Pattern is centered at the Origin, otherwise Pattern starts at the Origin */
+	UPROPERTY(EditAnywhere, Category = LinearPattern)
+	bool bCentered = true;
+};
+
+
+
+
+/**
+ * Settings for Grid Patterns in the Pattern Tool
+ * TODO: maybe we can just re-use UPatternTool_LinearSettings for this??
+ */
+UCLASS()
+class MESHMODELINGTOOLSEXP_API UPatternTool_GridSettings : public UInteractiveToolPropertySet
+{
+	GENERATED_BODY()
+public:
+	/** Spacing Technique used to distribute Pattern Elements along the Main axis */
+	UPROPERTY(EditAnywhere, Category = GridPatternX)
+	EPatternToolAxisSpacingMode SpacingX = EPatternToolAxisSpacingMode::ByCount;
+
+	/** Number of  Pattern Elements to place along the Main axis */
+	UPROPERTY(EditAnywhere, Category = GridPatternX, meta = (UIMin = 1, UIMax = 25, EditCondition = "SpacingX == EPatternToolAxisSpacingMode::ByCount"))
+	int32 CountX = 10;
+
+	/** Fixed Increment used to place Pattern Elements along the Main axis */
+	UPROPERTY(EditAnywhere, Category = GridPatternX, meta = (UIMin = 0, EditCondition = "SpacingX == EPatternToolAxisSpacingMode::StepSize"))
+	double StepSizeX = 100.0;
+
+	/** Length/Extent of Pattern falong the Main Axis */
+	UPROPERTY(EditAnywhere, Category = GridPatternX, meta = (UIMin = 0))
+	double ExtentX = 1000.0;
+
+	/** If true, Pattern is centered at the Origin along the Main axis, otherwise Pattern starts at the Origin */
+	UPROPERTY(EditAnywhere, Category = GridPatternX)
+	bool bCenteredX = true;
+
+	/** Spacing Technique used to distribute Pattern Elements along the Secondary axis*/
+	UPROPERTY(EditAnywhere, Category = GridPatternY)
+	EPatternToolAxisSpacingMode SpacingY = EPatternToolAxisSpacingMode::ByCount;
+
+	/** Number of  Pattern Elements to place along the Secondary axis */
+	UPROPERTY(EditAnywhere, Category = GridPatternY, meta = (UIMin = 1, UIMax = 25, EditCondition = "SpacingY == EPatternToolAxisSpacingMode::ByCount"))
+	int32 CountY = 10;
+
+	/** Fixed Increment used to place Pattern Elements along the Secondary axis */
+	UPROPERTY(EditAnywhere, Category = GridPatternY, meta = (UIMin = 0, EditCondition = "SpacingY == EPatternToolAxisSpacingMode::StepSize"))
+	double StepSizeY = 100.0;
+
+	/** Length/Extent of Pattern falong the Secondary Axis */
+	UPROPERTY(EditAnywhere, Category = GridPatternY, meta = (UIMin = 0))
+	double ExtentY = 1000.0;
+
+	/** If true, Pattern is centered at the Origin along the Secondary axis, otherwise Pattern starts at the Origin */
+	UPROPERTY(EditAnywhere, Category = GridPatternY)
+	bool bCenteredY = true;
+};
+
+
+
+
+
+/**
+ * Settings for Radial Patterns in the Pattern Tool
+ */
+UCLASS()
+class MESHMODELINGTOOLSEXP_API UPatternTool_RadialSettings : public UInteractiveToolPropertySet
+{
+	GENERATED_BODY()
+public:
+	/** Spacing Technique used to distribute Pattern Elements around the Circle/Arc */
+	UPROPERTY(EditAnywhere, Category = RadialPattern)
+	EPatternToolAxisSpacingMode SpacingMode = EPatternToolAxisSpacingMode::ByCount;
+
+	/** Number of  Pattern Elements to place */
+	UPROPERTY(EditAnywhere, Category = RadialPattern, meta = (UIMin = 1, UIMax = 25, EditCondition = "SpacingMode == EPatternToolAxisSpacingMode::ByCount"))
+	int32 Count = 10;
+
+	/** Fixed Increment (in Degrees) used to position Pattern Elements around the Circle/Arc */
+	UPROPERTY(EditAnywhere, Category = RadialPattern, meta = (Units = "Degrees", UIMin = 0, EditCondition = "SpacingMode == EPatternToolAxisSpacingMode::StepSize"))
+	double StepSize = 100.0;
+
+	/** Radius of the Circle/Arc */
+	UPROPERTY(EditAnywhere, Category = RadialPattern, meta = (UIMin = 0))
+	double Radius = 250;
+
+	/** Start angle of the Circle/Arc */
+	UPROPERTY(EditAnywhere, Category = RadialPattern, meta = (Units = "Degrees", UIMin = -360))
+	double StartAngle = 0.0;
+
+	/** End angle of the Circle/Arc */
+	UPROPERTY(EditAnywhere, Category = RadialPattern, meta = (Units = "Degrees", UIMax = 360))
+	double EndAngle = 360.0;
+
+	/** Fixed offset added to Start/End Angles */
+	UPROPERTY(EditAnywhere, Category = RadialPattern, meta = (Units = "Degrees", UIMin = -180, UIMax = 180))
+	double AngleShift = 0.0;
+
+	/** If true, Pattern elements are rotated to align with the Circle tangen */
+	UPROPERTY(EditAnywhere, Category = RadialPattern)
+	bool bOriented = true;
+};
+
+
+
+
+/**
+ * Settings for Per Element Rotation in the Pattern Tool
+ */
+UCLASS()
+class MESHMODELINGTOOLSEXP_API UPatternTool_RotationSettings : public UInteractiveToolPropertySet
+{
+	GENERATED_BODY()
+public:
+	/** If true, Rotation is linearly interpolated between StartRotation and Rotation values */
+	UPROPERTY(EditAnywhere, Category = Rotation)
+	bool bInterpolate = false;
+
+	/** Rotation at first Pattern Element */
+	UPROPERTY(EditAnywhere, Category = Rotation, meta = (EditCondition = "bInterpolate", EditConditionHides, HideEditConditionToggle))
+	FRotator StartRotation = FRotator::ZeroRotator;
+
+	/** Rotation applied to all Pattern Elements, or at Last Pattern Element for Interpolated rotations */
+	UPROPERTY(EditAnywhere, Category = Rotation)
+	FRotator Rotation = FRotator::ZeroRotator;
+};
+
+
+/**
+ * Settings for Per Element Translation in the Pattern Tool
+ */
+UCLASS()
+class MESHMODELINGTOOLSEXP_API UPatternTool_TranslationSettings : public UInteractiveToolPropertySet
+{
+	GENERATED_BODY()
+public:
+	/** If true, Translation is linearly interpolated between StartTranslation and Translation values */
+	UPROPERTY(EditAnywhere, Category = Translation)
+	bool bInterpolate = false;
+
+	/** Translation at first Pattern Element */
+	UPROPERTY(EditAnywhere, Category = Translation, meta = (EditCondition = "bInterpolate", EditConditionHides, HideEditConditionToggle))
+	FVector StartTranslation = FVector::ZeroVector;
+
+	/** Translation applied to all Pattern Elements, or at Last Pattern Element for Interpolated translations */
+	UPROPERTY(EditAnywhere, Category = Translation)
+	FVector Translation = FVector::ZeroVector;
+};
+
+
+/**
+ * Settings for Per Element Scale in the Pattern Tool
+ */
+UCLASS()
+class MESHMODELINGTOOLSEXP_API UPatternTool_ScaleSettings : public UInteractiveToolPropertySet
+{
+	GENERATED_BODY()
+public:
+	/** If true, Scaling is limited to Uniform Scaling */
+	UPROPERTY(EditAnywhere, Category = Scale)
+	bool bUniform = true;
+
+	/** If true, Scale is linearly interpolated between StartScale and Scale values */
+	UPROPERTY(EditAnywhere, Category = Scale)
+	bool bInterpolate = false;
+
+	/** Uniform Scale at first Pattern Element */
+	UPROPERTY(EditAnywhere, Category = Scale, meta = (EditCondition = "bInterpolate && bUniform", EditConditionHides, HideEditConditionToggle))
+	float StartScale = 1;
+
+	/** Uniform Scale applied to all Pattern Elements, or at Last Pattern Element for Interpolated scales */
+	UPROPERTY(EditAnywhere, Category = Scale, meta = (EditCondition = "bUniform", EditConditionHides, HideEditConditionToggle))
+	float Scale = 1;
+
+	/** Non-Uniform Scale at first Pattern Element */
+	UPROPERTY(EditAnywhere, Category = Scale, meta = (DisplayName = "StartScale", EditCondition = "bInterpolate && bUniform == false", EditConditionHides, HideEditConditionToggle))
+	FVector StartScaleNonUniform = FVector(1,1,1);
+
+	/** Non-Uniform Scale applied to all Pattern Elements, or at Last Pattern Element for Interpolated scales */
+	UPROPERTY(EditAnywhere, Category = Scale, meta = (DisplayName = "Scale", EditCondition = "bUniform == false", EditConditionHides, HideEditConditionToggle))
+	FVector ScaleNonUniform = FVector(1,1,1);
+};
+
+
+
+
+/**
+ * Ouptput Settings for the Pattern Tool
+ */
+UCLASS()
+class MESHMODELINGTOOLSEXP_API UPatternTool_OutputSettings : public UInteractiveToolPropertySet
+{
+	GENERATED_BODY()
+public:
+	/** Emit a separate Actor for each pattern element */
+	UPROPERTY(EditAnywhere, Category = Output)
+	bool bSeparateActors = false;
+
+	/** Emit StaticMesh pattern elements as DynamicMeshes */
+	UPROPERTY(EditAnywhere, Category = Output, meta = (EditCondition = "bHaveStaticMeshes == true", HideEditConditionToggle))
+	bool bConvertToDynamic = false;
+
+	/** Create InstancedStaticMeshComponents instead multiple StaticMeshComponents, for StaticMesh pattern elements */
+	UPROPERTY(EditAnywhere, Category = Output, meta = (EditCondition = "bHaveStaticMeshes == true && bSeparateActors == false && bConvertToDynamic == false", HideEditConditionToggle))
+	bool bCreateISMCs = false;
+
+	// internal, used to control state of Instance settings
+	UPROPERTY(meta = (TransientToolProperty))
+	bool bHaveStaticMeshes = false;
+};
+
+
+
+/**
+ * UPatternTool takes input meshes and generates 3D Patterns of those meshes, by
+ * placing repeated copies along geometric paths like lines, grids, circles, etc.
+ * The output can be a single Actor per pattern Element, or combined into single
+ * Actors in various ways depending on the input mesh type. 
+ */
+UCLASS()
+class MESHMODELINGTOOLSEXP_API UPatternTool : public UMultiSelectionMeshEditingTool
+{
+	GENERATED_BODY()
+
+public:
+	UPatternTool();
+
+	virtual void Setup() override;
+	virtual void OnShutdown(EToolShutdownType ShutdownType) override;
+
+	virtual void Render(IToolsContextRenderAPI* RenderAPI) override;
+
+	virtual void OnPropertyModified(UObject* PropertySet, FProperty* Property) override;
+
+	virtual bool HasCancel() const override { return true; }
+	virtual bool HasAccept() const override { return true; }
+	virtual bool CanAccept() const override { return true; }
+
+public:
+	UPROPERTY()
+	TObjectPtr<UPatternToolSettings> Settings;
+
+	UPROPERTY()
+	TObjectPtr<UPatternTool_LinearSettings> LinearSettings;
+
+	UPROPERTY()
+	TObjectPtr<UPatternTool_GridSettings> GridSettings;
+
+	UPROPERTY()
+	TObjectPtr<UPatternTool_RadialSettings> RadialSettings;
+
+	UPROPERTY()
+	TObjectPtr<UPatternTool_RotationSettings> RotationSettings;
+
+	UPROPERTY()
+	TObjectPtr<UPatternTool_TranslationSettings> TranslationSettings;
+
+	UPROPERTY()
+	TObjectPtr<UPatternTool_ScaleSettings> ScaleSettings;
+
+	
+
+	UPROPERTY()
+	TObjectPtr<UPatternTool_OutputSettings> OutputSettings;
+
+protected:
+
+	UPROPERTY()
+	TObjectPtr<UTransformProxy> TransformProxy_End = nullptr;
+
+	UPROPERTY()
+	TObjectPtr<UCombinedTransformGizmo> TransformGizmo_End = nullptr;
+
+
+	UPROPERTY()
+	TObjectPtr<UDragAlignmentMechanic> DragAlignmentMechanic = nullptr;
+
+
+	UPROPERTY()
+	TObjectPtr<UConstructionPlaneMechanic> PlaneMechanic = nullptr;
+
+	UE::Geometry::FFrame3d CurrentStartFrameWorld;
+	TArray<UE::Geometry::FTransformSRT3d> CurrentPattern;
+
+	bool bPatternNeedsUpdating = false;
+	FDateTime LastPatternUpdateTime;
+	void MarkPatternDirty();
+
+	void OnSourceVisibilityToggled(bool bVisible);
+	void OnMainFrameUpdated();
+	void OnShapeUpdated();
+	void OnParametersUpdated();
+	void UpdatePattern();
+	void GetPatternTransforms_Linear(TArray<UE::Geometry::FTransformSRT3d>& TransformsOut);
+	void GetPatternTransforms_Grid(TArray<UE::Geometry::FTransformSRT3d>& TransformsOut);
+	void GetPatternTransforms_Radial(TArray<UE::Geometry::FTransformSRT3d>& TransformsOut);
+
+
+	struct FPatternElement
+	{
+		int32 TargetIndex = 0;
+		bool bValid = true;
+
+		UPrimitiveComponent* SourceComponent = nullptr;
+		TArray<UMaterialInterface*> SourceMaterials;
+		UE::Geometry::FTransformSRT3d SourceTransform = UE::Geometry::FTransformSRT3d::Identity();
+		UE::Geometry::FTransformSRT3d BaseRotateScale = UE::Geometry::FTransformSRT3d::Identity();
+
+		UDynamicMesh* SourceDynamicMesh = nullptr;
+		UStaticMesh* SourceStaticMesh = nullptr;
+
+		UE::Geometry::FAxisAlignedBox3d Bounds = UE::Geometry::FAxisAlignedBox3d::Empty();
+	};
+	TArray<FPatternElement> Elements;
+
+	bool bHaveNonUniformScaleElements = false;
+
+	struct FComponentSet
+	{
+		TArray<UPrimitiveComponent*> Components;
+	};
+	TArray<FComponentSet> PreviewComponents;
+
+	UPROPERTY()
+	TSet<UPrimitiveComponent*> AllComponents;		// to keep components in FComponentSet alive
+
+	
+	TMap<int32, FComponentSet> StaticMeshPools;
+	UStaticMeshComponent* GetPreviewStaticMesh(FPatternElement& Element);
+	void ReturnStaticMeshes(FPatternElement& Element, FComponentSet& ComponentSet);
+
+	TMap<int32, FComponentSet> DynamicMeshPools;
+	UDynamicMeshComponent* GetPreviewDynamicMesh(FPatternElement& Element);
+	void ReturnDynamicMeshes(FPatternElement& Element, FComponentSet& ComponentSet);
+
+	void HideReturnedPreviewMeshes();
+
+
+	UPROPERTY()
+	TObjectPtr<UPreviewGeometry> PreviewGeometry = nullptr;	// parent actor for all preview components
+
+
+	void InitializeElements();
+	void ResetPreviews();
+	void DestroyPreviews();
+
+	void EmitResults();
+};
