ModelingTools: fix TObjectPtr error in PatternTool

diff --git a/Engine/Plugins/Experimental/MeshModelingToolsetExp/Source/MeshModelingToolsExp/Public/PatternTool.h b/Engine/Plugins/Experimental/MeshModelingToolsetExp/Source/MeshModelingToolsExp/Public/PatternTool.h
--- a/Engine/Plugins/Experimental/MeshModelingToolsetExp/Source/MeshModelingToolsExp/Public/PatternTool.h
+++ b/Engine/Plugins/Experimental/MeshModelingToolsetExp/Source/MeshModelingToolsExp/Public/PatternTool.h
@@ -462,7 +462,7 @@ protected:
 	TArray<FComponentSet> PreviewComponents;
 
 	UPROPERTY()
-	TSet<UPrimitiveComponent*> AllComponents;		// to keep components in FComponentSet alive
+	TSet<TObjectPtr<UPrimitiveComponent>> AllComponents;		// to keep components in FComponentSet alive
 
 	
 	TMap<int32, FComponentSet> StaticMeshPools;
