ModelingTools: fix static analysis issue in PatternTool

diff --git a/Engine/Plugins/Experimental/MeshModelingToolsetExp/Source/MeshModelingToolsExp/Private/PatternTool.cpp b/Engine/Plugins/Experimental/MeshModelingToolsetExp/Source/MeshModelingToolsExp/Private/PatternTool.cpp
--- a/Engine/Plugins/Experimental/MeshModelingToolsetExp/Source/MeshModelingToolsExp/Private/PatternTool.cpp
+++ b/Engine/Plugins/Experimental/MeshModelingToolsetExp/Source/MeshModelingToolsExp/Private/PatternTool.cpp
@@ -1166,10 +1166,9 @@ void UPatternTool::EmitResults()
 					FActorSpawnParameters SpawnInfo;
 					SpawnInfo.Template = SourceActor;
 					AStaticMeshActor* NewActor = GetTargetWorld()->SpawnActor<AStaticMeshActor>(SpawnInfo);
-					NewActor->SetActorTransform(WorldTransform);
-					//NewActor->GetStaticMeshComponent()->SetStaticMesh(SetStaticMesh);
 					if (NewActor != nullptr)
 					{
+						NewActor->SetActorTransform(WorldTransform);
 						NewActors.Add(NewActor);
 					}
 				}
