ModelingTools: Fixed error preview material not applying when the BakeTexture/BakeVertex/BakeAll tools had invalid inputs.

diff --git a/Engine/Plugins/Experimental/MeshModelingToolsetExp/Source/MeshModelingToolsExp/Private/BakeMeshAttributeMapsToolBase.cpp b/Engine/Plugins/Experimental/MeshModelingToolsetExp/Source/MeshModelingToolsExp/Private/BakeMeshAttributeMapsToolBase.cpp
--- a/Engine/Plugins/Experimental/MeshModelingToolsetExp/Source/MeshModelingToolsExp/Private/BakeMeshAttributeMapsToolBase.cpp
+++ b/Engine/Plugins/Experimental/MeshModelingToolsetExp/Source/MeshModelingToolsExp/Private/BakeMeshAttributeMapsToolBase.cpp
@@ -104,12 +104,17 @@ void UBakeMeshAttributeMapsToolBase::OnTick(float DeltaTime)
 	{
 		Compute->Tick(DeltaTime);
 
-		const float ElapsedComputeTime = Compute->GetElapsedComputeTime();
-		if (!CanAccept() && ElapsedComputeTime > SecondsBeforeWorkingMaterial)
+		if (static_cast<bool>(OpState & EBakeOpState::Invalid))
 		{
-			UMaterialInstanceDynamic* ProgressMaterial =
-				static_cast<bool>(OpState & EBakeOpState::Invalid) ? ErrorPreviewMaterial : WorkingPreviewMaterial;
-			PreviewMesh->SetOverrideRenderMaterial(ProgressMaterial);
+			PreviewMesh->SetOverrideRenderMaterial(ErrorPreviewMaterial);
+		}
+		else
+		{
+			const float ElapsedComputeTime = Compute->GetElapsedComputeTime();
+			if (!CanAccept() && ElapsedComputeTime > SecondsBeforeWorkingMaterial)
+			{
+				PreviewMesh->SetOverrideRenderMaterial(WorkingPreviewMaterial);
+			}
 		}
 	}
 }
diff --git a/Engine/Plugins/Experimental/MeshModelingToolsetExp/Source/MeshModelingToolsExp/Private/BakeMeshAttributeVertexTool.cpp b/Engine/Plugins/Experimental/MeshModelingToolsetExp/Source/MeshModelingToolsExp/Private/BakeMeshAttributeVertexTool.cpp
--- a/Engine/Plugins/Experimental/MeshModelingToolsetExp/Source/MeshModelingToolsExp/Private/BakeMeshAttributeVertexTool.cpp
+++ b/Engine/Plugins/Experimental/MeshModelingToolsetExp/Source/MeshModelingToolsExp/Private/BakeMeshAttributeVertexTool.cpp
@@ -416,12 +416,17 @@ void UBakeMeshAttributeVertexTool::OnTick(float DeltaTime)
 	{
 		Compute->Tick(DeltaTime);
 
-		const float ElapsedComputeTime = Compute->GetElapsedComputeTime();
-		if (!CanAccept() && ElapsedComputeTime > SecondsBeforeWorkingMaterial)
+		if (static_cast<bool>(OpState & EBakeOpState::Invalid))
 		{
-			UMaterialInstanceDynamic* ProgressMaterial =
-				static_cast<bool>(OpState & EBakeOpState::Invalid) ? ErrorPreviewMaterial : WorkingPreviewMaterial;
-			PreviewMesh->SetOverrideRenderMaterial(ProgressMaterial);
+			PreviewMesh->SetOverrideRenderMaterial(ErrorPreviewMaterial);
+		}
+		else
+		{
+			const float ElapsedComputeTime = Compute->GetElapsedComputeTime();
+			if (!CanAccept() && ElapsedComputeTime > SecondsBeforeWorkingMaterial)
+			{
+				PreviewMesh->SetOverrideRenderMaterial(WorkingPreviewMaterial);
+			}
 		}
 	}
 }
