ModelingTools: Reset preview material in BakeTexture/BakeAll when output type is modified.

diff --git a/Engine/Plugins/Experimental/MeshModelingToolsetExp/Source/MeshModelingToolsExp/Private/BakeMeshAttributeMapsToolBase.cpp b/Engine/Plugins/Experimental/MeshModelingToolsetExp/Source/MeshModelingToolsExp/Private/BakeMeshAttributeMapsToolBase.cpp
--- a/Engine/Plugins/Experimental/MeshModelingToolsetExp/Source/MeshModelingToolsExp/Private/BakeMeshAttributeMapsToolBase.cpp
+++ b/Engine/Plugins/Experimental/MeshModelingToolsetExp/Source/MeshModelingToolsExp/Private/BakeMeshAttributeMapsToolBase.cpp
@@ -41,9 +41,7 @@ void UBakeMeshAttributeMapsToolBase::Setup()
 	if (Material != nullptr)
 	{
 		PreviewMaterial = UMaterialInstanceDynamic::Create(Material, GetToolManager());
-		PreviewMaterial->SetTextureParameterValue(TEXT("NormalMap"), EmptyNormalMap);
-		PreviewMaterial->SetTextureParameterValue(TEXT("OcclusionMap"), EmptyColorMapWhite);
-		PreviewMaterial->SetTextureParameterValue(TEXT("ColorMap"), EmptyColorMapWhite);
+		ResetPreview();
 	}
 	UMaterial* BentNormalMaterial = LoadObject<UMaterial>(nullptr, TEXT("/MeshModelingToolsetExp/Materials/BakeBentNormalPreviewMaterial"));
 	check(BentNormalMaterial);
@@ -288,6 +286,14 @@ void UBakeMeshAttributeMapsToolBase::UpdatePreview(const EBakeMapType PreviewMap
 }
 
 
+void UBakeMeshAttributeMapsToolBase::ResetPreview()
+{
+	PreviewMaterial->SetTextureParameterValue(TEXT("NormalMap"), EmptyNormalMap);
+	PreviewMaterial->SetTextureParameterValue(TEXT("OcclusionMap"), EmptyColorMapWhite);
+	PreviewMaterial->SetTextureParameterValue(TEXT("ColorMap"), EmptyColorMapWhite);
+}
+
+
 void UBakeMeshAttributeMapsToolBase::UpdatePreviewNames(
 	const EBakeMapType MapTypes,
 	FString& MapPreview,
@@ -330,6 +336,8 @@ void UBakeMeshAttributeMapsToolBase::OnMapTypesUpdated(
 	TArray<FString>& MapPreviewNamesList,
 	TMap<FString, FString>& MapPreviewNamesMap)
 {
+	ResetPreview();
+	
 	// Use the processed bitfield which may contain additional targets
 	// (ex. AO if BentNormal was requested) to preallocate cached map storage.
 	EBakeMapType CachedMapTypes = GetMapTypes(static_cast<int32>(ResultMapTypes));
diff --git a/Engine/Plugins/Experimental/MeshModelingToolsetExp/Source/MeshModelingToolsExp/Public/BakeMeshAttributeMapsToolBase.h b/Engine/Plugins/Experimental/MeshModelingToolsetExp/Source/MeshModelingToolsExp/Public/BakeMeshAttributeMapsToolBase.h
--- a/Engine/Plugins/Experimental/MeshModelingToolsetExp/Source/MeshModelingToolsExp/Public/BakeMeshAttributeMapsToolBase.h
+++ b/Engine/Plugins/Experimental/MeshModelingToolsetExp/Source/MeshModelingToolsExp/Public/BakeMeshAttributeMapsToolBase.h
@@ -237,6 +237,12 @@ protected:
 	 * @param PreviewMapType EBakeMapType to preview
 	 */
 	void UpdatePreview(EBakeMapType PreviewMapType);
+
+
+	/**
+	 * Resets the preview material parameters to their default state.
+	 */
+	void ResetPreview();
 	
 
 	/**
