Render Capture Baking: Use depth capture to eliminate occlusion artefacts which can appear when baking non-convex shapes and look like weird sihouettes/blotches in the base color channel

Also removed artefacts that became visible after the ones fixible with a depth test were resolved. These were related to linearly interpolating data on the photosets, this was fixed since the baking framework is where texture filtering should take place

#rnx

diff --git a/Engine/Plugins/Experimental/MeshModelingToolsetExp/Source/MeshModelingToolsExp/Private/BakeRenderCaptureTool.cpp b/Engine/Plugins/Experimental/MeshModelingToolsetExp/Source/MeshModelingToolsExp/Private/BakeRenderCaptureTool.cpp
--- a/Engine/Plugins/Experimental/MeshModelingToolsetExp/Source/MeshModelingToolsExp/Private/BakeRenderCaptureTool.cpp
+++ b/Engine/Plugins/Experimental/MeshModelingToolsetExp/Source/MeshModelingToolsExp/Private/BakeRenderCaptureTool.cpp
@@ -39,6 +39,7 @@ class FSceneCapturePhotoSetSampler : public FMeshBakerDynamicMeshSampler
 public:
 	FSceneCapturePhotoSetSampler(
 		FSceneCapturePhotoSet* SceneCapture,
+		float ValidSampleDepthThreshold,
 		TFunctionRef<bool(const FVector3d&, const FVector3d&)> VisibilityFunction,
 		const FDynamicMesh3* Mesh,
 		const FDynamicMeshAABBTree3* Spatial,
@@ -46,12 +47,19 @@ public:
 	:
 		FMeshBakerDynamicMeshSampler(Mesh, Spatial, Tangents),
 		SceneCapture(SceneCapture),
+		ValidSampleDepthThreshold(ValidSampleDepthThreshold),
 		VisibilityFunction(VisibilityFunction)
 	{
 		check(SceneCapture != nullptr);
 		check(Mesh != nullptr);
 		check(Spatial != nullptr);
 		check(Tangents != nullptr);
+
+		bool bHasDepth = SceneCapture->GetCaptureTypeEnabled(ERenderCaptureType::DeviceDepth);
+		if (bHasDepth)
+		{
+			ensure(ValidSampleDepthThreshold > 0); // We only need the depth capture if this threshold is positive
+		}
 	}
 
 	virtual bool SupportsCustomCorrespondence() const override
@@ -65,7 +73,13 @@ public:
 		// Perform a ray-cast to determine which photo/coordinate, if any, should be sampled
 		int PhotoIndex;
 		FVector2d PhotoCoords;
-		SceneCapture->ComputeSampleLocation(Sample.BaseSample.SurfacePoint, Sample.BaseNormal, VisibilityFunction, PhotoIndex, PhotoCoords);
+		SceneCapture->ComputeSampleLocation(
+			Sample.BaseSample.SurfacePoint,
+			Sample.BaseNormal,
+			ValidSampleDepthThreshold,
+			VisibilityFunction,
+			PhotoIndex,
+			PhotoCoords);
 
 		// Store the photo coordinates and index in the correspondence sample
 		Sample.DetailMesh = SceneCapture;
@@ -84,6 +98,7 @@ public:
 
 public:
 	FSceneCapturePhotoSet* SceneCapture = nullptr;
+	float ValidSampleDepthThreshold = 0;
 	TFunctionRef<bool(const FVector3d&, const FVector3d&)> VisibilityFunction;
 };
 
@@ -115,10 +130,12 @@ public:
 		// Material approximation settings
 		//
 		int32 RenderCaptureImageSize = 1024;
+		bool bAntiAliasing = false;
 
 		// render capture parameters
 		double FieldOfViewDegrees = 45.0;
 		double NearPlaneDist = 1.0;
+		double ValidSampleDepthThreshold = 0;
 
 		//
 		// Material output settings
@@ -132,6 +149,7 @@ public:
 		bool bBakeSpecular = true;
 		bool bBakeEmissive = true;
 		bool bBakeNormalMap = true;
+		bool bBakeDeviceDepth = true;
 		
 		bool bUsePackedMRS = true;
 
@@ -148,6 +166,7 @@ public:
 	 */
 	static FOptions ConstructOptions(
 		const URenderCaptureProperties& RenderCaptureProperties,
+		const UBakeRenderCaptureToolProperties& ToolProperties,
 		const UBakeRenderCaptureInputToolProperties& InputMeshSettings)
 	{
 		//
@@ -158,15 +177,22 @@ public:
 		Options.TargetUVLayer = InputMeshSettings.GetTargetUVLayerIndex();
 
 		Options.RenderCaptureImageSize = static_cast<int32>(RenderCaptureProperties.Resolution);
+		Options.ValidSampleDepthThreshold = ToolProperties.ValidSampleDepthThreshold;
 
 		Options.bBakeBaseColor = RenderCaptureProperties.bBaseColorMap;
-		Options.bBakeRoughness = RenderCaptureProperties.bRoughnessMap;
-		Options.bBakeMetallic = RenderCaptureProperties.bMetallicMap;
-		Options.bBakeSpecular = RenderCaptureProperties.bSpecularMap;
-		Options.bBakeEmissive = RenderCaptureProperties.bEmissiveMap;
 		Options.bBakeNormalMap = RenderCaptureProperties.bNormalMap;
-		Options.bUsePackedMRS = RenderCaptureProperties.bPackedMRSMap;
-
+		Options.bBakeEmissive =  RenderCaptureProperties.bEmissiveMap;
+		Options.bBakeDeviceDepth = RenderCaptureProperties.bDeviceDepthMap;
+		
+		// Enforce the PackedMRS precondition here so we don't have to check it at each usage site.  Note: We don't
+		// apply this precondition on the RenderCaptureProperties because we don't want the user to have to re-enable
+		// options which enabling PackedMRS disabled.
+		Options.bUsePackedMRS =  RenderCaptureProperties.bPackedMRSMap;
+		Options.bBakeMetallic =  RenderCaptureProperties.bPackedMRSMap ? false : RenderCaptureProperties.bMetallicMap;
+		Options.bBakeRoughness = RenderCaptureProperties.bPackedMRSMap ? false : RenderCaptureProperties.bRoughnessMap;
+		Options.bBakeSpecular =  RenderCaptureProperties.bPackedMRSMap ? false : RenderCaptureProperties.bSpecularMap;
+
+		Options.bAntiAliasing = RenderCaptureProperties.bAntiAliasing;
 		Options.FieldOfViewDegrees = RenderCaptureProperties.CaptureFieldOfView;
 		Options.NearPlaneDist = RenderCaptureProperties.NearPlaneDist;
 
@@ -194,24 +220,24 @@ static TUniquePtr<FSceneCapturePhotoSet> CapturePhotoSet(
 	TUniquePtr<FSceneCapturePhotoSet> SceneCapture = MakeUnique<FSceneCapturePhotoSet>();
 	SceneCapture->SetAllowCancel(bAllowCancel);
 
-	SceneCapture->SetCaptureTypeEnabled(ERenderCaptureType::BaseColor, Options.bBakeBaseColor);
+	SceneCapture->SetCaptureTypeEnabled(ERenderCaptureType::DeviceDepth, Options.bBakeDeviceDepth);
+	SceneCapture->SetCaptureTypeEnabled(ERenderCaptureType::BaseColor,   Options.bBakeBaseColor);
 	SceneCapture->SetCaptureTypeEnabled(ERenderCaptureType::WorldNormal, Options.bBakeNormalMap);
-	SceneCapture->SetCaptureTypeEnabled(ERenderCaptureType::Emissive, Options.bBakeEmissive);
-
-	if (Options.bUsePackedMRS)
-	{
-		SceneCapture->SetCaptureTypeEnabled(ERenderCaptureType::CombinedMRS, true);
-		SceneCapture->SetCaptureTypeEnabled(ERenderCaptureType::Roughness, false);
-		SceneCapture->SetCaptureTypeEnabled(ERenderCaptureType::Metallic, false);
-		SceneCapture->SetCaptureTypeEnabled(ERenderCaptureType::Specular, false);
-	}
-	else
-	{
-		SceneCapture->SetCaptureTypeEnabled(ERenderCaptureType::CombinedMRS, false);
-		SceneCapture->SetCaptureTypeEnabled(ERenderCaptureType::Roughness, Options.bBakeRoughness);
-		SceneCapture->SetCaptureTypeEnabled(ERenderCaptureType::Metallic, Options.bBakeMetallic);
-		SceneCapture->SetCaptureTypeEnabled(ERenderCaptureType::Specular, Options.bBakeSpecular);
-	}
+	SceneCapture->SetCaptureTypeEnabled(ERenderCaptureType::Emissive,    Options.bBakeEmissive);
+	SceneCapture->SetCaptureTypeEnabled(ERenderCaptureType::CombinedMRS, Options.bUsePackedMRS);
+	SceneCapture->SetCaptureTypeEnabled(ERenderCaptureType::Metallic,    Options.bBakeMetallic);
+	SceneCapture->SetCaptureTypeEnabled(ERenderCaptureType::Roughness,   Options.bBakeRoughness);
+	SceneCapture->SetCaptureTypeEnabled(ERenderCaptureType::Specular,    Options.bBakeSpecular);
+
+	FRenderCaptureConfig Config;
+	Config.bAntiAliasing = Options.bAntiAliasing;
+	SceneCapture->SetCaptureConfig(ERenderCaptureType::BaseColor,   Config);
+	SceneCapture->SetCaptureConfig(ERenderCaptureType::WorldNormal, Config);
+	SceneCapture->SetCaptureConfig(ERenderCaptureType::CombinedMRS, Config);
+	SceneCapture->SetCaptureConfig(ERenderCaptureType::Metallic,    Config);
+	SceneCapture->SetCaptureConfig(ERenderCaptureType::Roughness,   Config);
+	SceneCapture->SetCaptureConfig(ERenderCaptureType::Specular,    Config);
+	SceneCapture->SetCaptureConfig(ERenderCaptureType::Emissive,    Config);
 
 	SceneCapture->SetCaptureSceneActors(Actors[0]->GetWorld(), Actors);
 
@@ -226,7 +252,10 @@ static TUniquePtr<FSceneCapturePhotoSet> CapturePhotoSet(
 
 template <ERenderCaptureType CaptureType>
 TSharedPtr<FRenderCaptureMapEvaluator<FVector4f>>
-MakeColorEvaluator(const FSceneCapturePhotoSet::FSceneSample& DefaultSample, const FSceneCapturePhotoSet* SceneCapture)
+MakeColorEvaluator(
+	const FSceneCapturePhotoSet::FSceneSample& DefaultSample,
+	const FSceneCapturePhotoSet* SceneCapture,
+	const FImageDimensions& CaptureDimensions)
 {
 	TSharedPtr<FRenderCaptureMapEvaluator<FVector4f>> Evaluator = MakeShared<FRenderCaptureMapEvaluator<FVector4f>>();
 
@@ -252,15 +281,18 @@ MakeColorEvaluator(const FSceneCapturePhotoSet::FSceneSample& DefaultSample, con
 	case ERenderCaptureType::CombinedMRS:
 		Evaluator->Channel = ERenderCaptureChannel::CombinedMRS;
 		break;
+	case ERenderCaptureType::DeviceDepth:
+		Evaluator->Channel = ERenderCaptureChannel::DeviceDepth;
+		break;
 	}
 
 	Evaluator->DefaultResult = DefaultSample.GetValue4f(CaptureType);
 
-	Evaluator->EvaluateSampleCallback = [&DefaultSample, SceneCapture](const FMeshMapEvaluator::FCorrespondenceSample& Sample)
+	Evaluator->EvaluateSampleCallback = [&DefaultSample, SceneCapture, &CaptureDimensions](const FMeshMapEvaluator::FCorrespondenceSample& Sample)
 	{
 		const int PhotoIndex = Sample.DetailTriID;
 		const FVector2d PhotoCoords(Sample.DetailBaryCoords.X, Sample.DetailBaryCoords.Y);
-		const FVector4f SampleColor = SceneCapture->ComputeSample<CaptureType>(PhotoIndex, PhotoCoords, DefaultSample);
+		const FVector4f SampleColor = SceneCapture->ComputeSampleNearest<CaptureType>(PhotoIndex, PhotoCoords, DefaultSample);
 		return SampleColor;
 	};
 
@@ -304,10 +336,13 @@ void FRenderCaptureMapBakerOp::CalculateResult(FProgressCancel*)
 
 	FDynamicMeshAABBTree3 BaseMeshSpatial(BaseMesh, true);
 
-	double RayOffsetHackDist = (double)(100.0 * FMathf::ZeroTolerance * BaseMesh->GetBounds().MinDim() );
-	auto VisibilityFunction = [&BaseMeshSpatial, RayOffsetHackDist](const FVector3d& SurfPos, const FVector3d& ImagePosWorld)
+	// Compute an offset shift surface positions toward the render camera position so we don't immediately self intersect
+	// The Max expression ensures the code works when the base mesh is 2D along one coordinate direction
+	const double RayOffsetHackDist = (double)(100.0 * FMathf::ZeroTolerance * FMath::Max(BaseMesh->GetBounds().MinDim(), 0.01));
+
+	auto VisibilityFunction = [&BaseMeshSpatial, RayOffsetHackDist](const FVector3d& SurfPos, const FVector3d& ViewPos)
 	{
-		FVector3d RayDir = ImagePosWorld - SurfPos;
+		FVector3d RayDir = ViewPos - SurfPos;
 		double Dist = Normalize(RayDir);
 		FVector3d RayOrigin = SurfPos + RayOffsetHackDist * RayDir;
 		int32 HitTID = BaseMeshSpatial.FindNearestHitTriangle(FRay3d(RayOrigin, RayDir), IMeshSpatial::FQueryOptions(Dist));
@@ -429,7 +464,7 @@ void FRenderCaptureMapBakerOp::CalculateResult(FProgressCancel*)
 	Result->InteriorSampleCallback = RegisterSampleStats;
 	Result->PostWriteToImageCallback = ComputeAndApplyInfill;
 
-	FSceneCapturePhotoSetSampler Sampler(SceneCapture, VisibilityFunction, BaseMesh, &BaseMeshSpatial, BaseMeshTangents.Get());
+	FSceneCapturePhotoSetSampler Sampler(SceneCapture, Options.ValidSampleDepthThreshold, VisibilityFunction, BaseMesh, &BaseMeshSpatial, BaseMeshTangents.Get());
 	Result->SetDetailSampler(&Sampler);
 	Result->SetCorrespondenceStrategy(FMeshMapBaker::ECorrespondenceStrategy::Custom);
 
@@ -445,58 +480,62 @@ void FRenderCaptureMapBakerOp::CalculateResult(FProgressCancel*)
 	DefaultColorSample.Metallic = InvalidColor.X;
 	DefaultColorSample.Emissive = FVector3f(InvalidColor.X, InvalidColor.Y, InvalidColor.Z);
 	DefaultColorSample.WorldNormal = FVector4f((DefaultNormal + FVector3f::One()) * .5f, InvalidColor.W);
+	
+	// We use 0 here since this corresponds to the infinite far plane value. Also, we don't preview the depth texture
+	DefaultColorSample.DeviceDepth = 0;
 
 	auto AddColorEvaluator = [this, &InfillData] (const TSharedPtr<FRenderCaptureMapEvaluator<FVector4f>>& Evaluator)
 	{
 		Result->AddEvaluator(Evaluator);
 		InfillData.EvaluatorNeedsInfill.Add(true);
 	};
+	
+	const FImageDimensions CaptureDimensions(Options.RenderCaptureImageSize, Options.RenderCaptureImageSize);
 
+	if (Options.bBakeDeviceDepth)
+	{
+		AddColorEvaluator(MakeColorEvaluator<ERenderCaptureType::DeviceDepth>(DefaultColorSample, SceneCapture, CaptureDimensions));
+	}
 	if (Options.bBakeBaseColor)
 	{
-		AddColorEvaluator(MakeColorEvaluator<ERenderCaptureType::BaseColor>(DefaultColorSample, SceneCapture));
+		AddColorEvaluator(MakeColorEvaluator<ERenderCaptureType::BaseColor>(DefaultColorSample, SceneCapture, CaptureDimensions));
 	}
-
 	if (Options.bUsePackedMRS)
 	{
-		AddColorEvaluator(MakeColorEvaluator<ERenderCaptureType::CombinedMRS>(DefaultColorSample, SceneCapture));
+		AddColorEvaluator(MakeColorEvaluator<ERenderCaptureType::CombinedMRS>(DefaultColorSample, SceneCapture, CaptureDimensions));
 	}
-	else
+	if (Options.bBakeRoughness)
 	{
-		if (Options.bBakeRoughness)
-		{
-			AddColorEvaluator(MakeColorEvaluator<ERenderCaptureType::Roughness>(DefaultColorSample, SceneCapture));
-		}
-		if (Options.bBakeMetallic)
-		{
-			AddColorEvaluator(MakeColorEvaluator<ERenderCaptureType::Metallic>(DefaultColorSample, SceneCapture));
-		}
-		if (Options.bBakeSpecular)
-		{
-			AddColorEvaluator(MakeColorEvaluator<ERenderCaptureType::Specular>(DefaultColorSample, SceneCapture));
-		}
+		AddColorEvaluator(MakeColorEvaluator<ERenderCaptureType::Roughness>(DefaultColorSample, SceneCapture, CaptureDimensions));
+	}
+	if (Options.bBakeMetallic)
+	{
+		AddColorEvaluator(MakeColorEvaluator<ERenderCaptureType::Metallic>(DefaultColorSample, SceneCapture, CaptureDimensions));
+	}
+	if (Options.bBakeSpecular)
+	{
+		AddColorEvaluator(MakeColorEvaluator<ERenderCaptureType::Specular>(DefaultColorSample, SceneCapture, CaptureDimensions));
 	}
-
 	if (Options.bBakeEmissive)
 	{
-		AddColorEvaluator(MakeColorEvaluator<ERenderCaptureType::Emissive>(DefaultColorSample, SceneCapture));
+		AddColorEvaluator(MakeColorEvaluator<ERenderCaptureType::Emissive>(DefaultColorSample, SceneCapture, CaptureDimensions));
 	}
-
-	if (Options.bBakeNormalMap) {
+	if (Options.bBakeNormalMap)
+	{
 		TSharedPtr<FRenderCaptureMapEvaluator<FVector3f>> Evaluator = MakeShared<FRenderCaptureMapEvaluator<FVector3f>>();
 
 		Evaluator->Channel = ERenderCaptureChannel::WorldNormal;
 
 		Evaluator->DefaultResult = DefaultNormal;
 
-		Evaluator->EvaluateSampleCallback = [this, &DefaultColorSample](const FMeshMapEvaluator::FCorrespondenceSample& Sample)
+		Evaluator->EvaluateSampleCallback = [this, &DefaultColorSample, &CaptureDimensions](const FMeshMapEvaluator::FCorrespondenceSample& Sample)
 		{
 			const int32 TriangleID = Sample.BaseSample.TriangleIndex;
 			const FVector3d BaryCoords = Sample.BaseSample.BaryCoords;
 			const int PhotoIndex = Sample.DetailTriID;
 			const FVector2d PhotoCoords(Sample.DetailBaryCoords.X, Sample.DetailBaryCoords.Y);
 
-			const FVector4f NormalColor = SceneCapture->ComputeSample<ERenderCaptureType::WorldNormal>(PhotoIndex, PhotoCoords, DefaultColorSample);
+			const FVector4f NormalColor = SceneCapture->ComputeSampleNearest<ERenderCaptureType::WorldNormal>(PhotoIndex, PhotoCoords, DefaultColorSample);
 
 			// Map from color components [0,1] to normal components [-1,1]
 			const FVector3f WorldSpaceNormal(
@@ -617,10 +656,27 @@ void UBakeRenderCaptureTool::Setup()
 	Settings->RestoreProperties(this);
 	AddToolPropertySource(Settings);
 
-	Settings->MapPreview = BaseColorTexParamName; // We always bake the base color
+	Settings->MapPreview = BaseColorTexParamName;
 	Settings->WatchProperty(Settings->MapPreview, [this](FString) { UpdateVisualization(); GetToolManager()->PostInvalidation(); });
 	Settings->WatchProperty(Settings->SamplesPerPixel, [this](EBakeTextureSamplesPerPixel) { OpState |= EBakeOpState::Evaluate; });
 	Settings->WatchProperty(Settings->TextureSize, [this](EBakeTextureResolution) { OpState |= EBakeOpState::Evaluate; });
+	Settings->WatchProperty(Settings->ValidSampleDepthThreshold, [this](float ValidSampleDepthThreshold)
+	{
+		// Only compute the device depth if we compute at least one other channel, the DeviceDepth is used to eliminate
+		// occlusion artefacts from the other channels
+		RenderCaptureProperties->bDeviceDepthMap = (ValidSampleDepthThreshold > 0) &&
+			(
+			RenderCaptureProperties->bBaseColorMap ||
+			RenderCaptureProperties->bNormalMap    ||
+			RenderCaptureProperties->bEmissiveMap  ||
+			RenderCaptureProperties->bPackedMRSMap ||
+			RenderCaptureProperties->bMetallicMap  ||
+			RenderCaptureProperties->bRoughnessMap ||
+			RenderCaptureProperties->bSpecularMap
+			);
+
+		OpState |= EBakeOpState::Evaluate;
+	});
 
 	RenderCaptureProperties = NewObject<URenderCaptureProperties>(this);
 	RenderCaptureProperties->RestoreProperties(this);
@@ -634,6 +690,8 @@ void UBakeRenderCaptureTool::Setup()
 	RenderCaptureProperties->WatchProperty(RenderCaptureProperties->bSpecularMap, [this](bool) { OpState |= EBakeOpState::Evaluate; });
 	RenderCaptureProperties->WatchProperty(RenderCaptureProperties->bPackedMRSMap, [this](bool) { OpState |= EBakeOpState::Evaluate; });
 	RenderCaptureProperties->WatchProperty(RenderCaptureProperties->bEmissiveMap, [this](bool) { OpState |= EBakeOpState::Evaluate; });
+	RenderCaptureProperties->WatchProperty(RenderCaptureProperties->bAntiAliasing, [this](bool) { OpState |= EBakeOpState::Evaluate; });
+	// These are not exposed to the UI, but we watch them anyway because we might change that later
 	RenderCaptureProperties->WatchProperty(RenderCaptureProperties->CaptureFieldOfView, [this](float) { OpState |= EBakeOpState::Evaluate; });
 	RenderCaptureProperties->WatchProperty(RenderCaptureProperties->NearPlaneDist, [this](float) { OpState |= EBakeOpState::Evaluate; });
 	
@@ -645,13 +703,13 @@ void UBakeRenderCaptureTool::Setup()
 	InputMeshSettings->WatchProperty(InputMeshSettings->TargetUVLayer, [this](FString) { OpState |= EBakeOpState::Evaluate; });
 	
 	{
-		Settings->MapPreviewNamesList.Add(NormalTexParamName);
 		Settings->MapPreviewNamesList.Add(BaseColorTexParamName);
-		Settings->MapPreviewNamesList.Add(RoughnessTexParamName);
+		Settings->MapPreviewNamesList.Add(NormalTexParamName);
+		Settings->MapPreviewNamesList.Add(PackedMRSTexParamName);
 		Settings->MapPreviewNamesList.Add(MetallicTexParamName);
+		Settings->MapPreviewNamesList.Add(RoughnessTexParamName);
 		Settings->MapPreviewNamesList.Add(SpecularTexParamName);
 		Settings->MapPreviewNamesList.Add(EmissiveTexParamName);
-		Settings->MapPreviewNamesList.Add(PackedMRSTexParamName);
 	}
 
 	ResultSettings = NewObject<UBakeRenderCaptureResults>(this);
@@ -776,6 +834,10 @@ void UBakeRenderCaptureTool::CreateTextureAssetsRC(UWorld* SourceWorld, UObject*
 		CreateTextureAsset(TexName, FTexture2DBuilder::ETextureType::EmissiveHDR, ResultSettings->EmissiveMap);
 	}
 
+	// We need different code paths based on PackedMRS here because we don't want to uncheck the separate channels
+	// when PackedMRS is enabled to give the user a better UX (they don't have to re-check them after disabling
+	// PackedMRS). In other place we can test the PackedMRS and separate channel booleans in series and avoid the
+	// complexity of nested if statements.
 	if (RenderCaptureProperties->bPackedMRSMap && ResultSettings->PackedMRSMap != nullptr)
 	{
 		const FString TexName = FString::Printf(TEXT("%s_%s"), *BaseName, *PackedMRSTexParamName);
@@ -828,6 +890,11 @@ bool UBakeRenderCaptureTool::CanAccept() const
 	{
 		return false;
 	}
+	
+	// We need different code paths based on PackedMRS here because we don't want to uncheck the separate channels
+	// when PackedMRS is enabled to give the user a better UX (they don't have to re-check them after disabling
+	// PackedMRS). In other place we can test the PackedMRS and separate channel booleans in series and avoid the
+	// complexity of nested if statements.
 	if (RenderCaptureProperties->bPackedMRSMap)
 	{
 		if (ResultSettings->PackedMRSMap == nullptr)
@@ -861,7 +928,7 @@ TUniquePtr<TGenericDataOperator<FMeshMapBaker>> UBakeRenderCaptureTool::MakeNewO
 	TUniquePtr<FRenderCaptureMapBakerOp> Op = MakeUnique<FRenderCaptureMapBakerOp>();
 	Op->BaseMesh = &TargetMesh;
 	Op->BaseMeshTangents = TargetMeshTangents;
-	Op->Options = FRenderCaptureSettings::ConstructOptions(*RenderCaptureProperties, *InputMeshSettings);
+	Op->Options = FRenderCaptureSettings::ConstructOptions(*RenderCaptureProperties, *Settings, *InputMeshSettings);
 	Op->TextureImageSize = static_cast<int32>(Settings->TextureSize);
 	Op->SamplesPerPixel = Settings->SamplesPerPixel;
 	Op->SceneCapture = SceneCapture.Get();
@@ -951,6 +1018,10 @@ void UBakeRenderCaptureTool::OnMapsUpdatedRC(const TUniquePtr<FMeshMapBaker>& Ne
 						bPopulateSourceData);
 					break;
 
+				case ERenderCaptureChannel::DeviceDepth:
+					// Do nothing, the depth capture is used internally and not presented to the user
+					break;
+
 				default:
 					ensure(false);
 					return;
@@ -1082,7 +1153,12 @@ void UBakeRenderCaptureTool::UpdateResult()
 	}
 
 	//
-	// create a set of spatially located render captures of the scene ("photo set").
+	// create a set of spatially located render captures of the scene ("photo set"). We need to recompute this if the
+	// render capture properties changed. Note we only compare the URenderCaptureProperties only, and not the
+	// ValidSampleDepthThreshold, this is intentional so that we only trigger a scene capture recompute when we go from
+	// a zero to a positive threshold (we need to compute the depth capture), or a positive to a zero threshold (we can
+	// save memory and not compute the depth capture), we don't need to recompute the scene capture when the user is
+	// changing between positive threshold values.
 	//
 	if (*RenderCaptureProperties != *ComputedRenderCaptureProperties)
 	{
@@ -1095,7 +1171,7 @@ void UBakeRenderCaptureTool::UpdateResult()
 		const bool bAllowCancel = (bFirstEverSceneCapture == false);
 
 		SceneCapture.Reset();
-		FRenderCaptureSettings::FOptions Options = FRenderCaptureSettings::ConstructOptions(*RenderCaptureProperties, *InputMeshSettings);
+		FRenderCaptureSettings::FOptions Options = FRenderCaptureSettings::ConstructOptions(*RenderCaptureProperties, *Settings, *InputMeshSettings);
 		SceneCapture = CapturePhotoSet(Actors, Options, bAllowCancel);
 
 		for (int Idx = 1; Idx < Targets.Num(); ++Idx)
@@ -1114,11 +1190,15 @@ void UBakeRenderCaptureTool::UpdateResult()
 			RenderCaptureProperties->bSpecularMap       = ComputedRenderCaptureProperties->bSpecularMap;
 			RenderCaptureProperties->bPackedMRSMap      = ComputedRenderCaptureProperties->bPackedMRSMap;
 			RenderCaptureProperties->bEmissiveMap       = ComputedRenderCaptureProperties->bEmissiveMap;
+			RenderCaptureProperties->bAntiAliasing      = ComputedRenderCaptureProperties->bAntiAliasing;
+			RenderCaptureProperties->bDeviceDepthMap    = ComputedRenderCaptureProperties->bDeviceDepthMap;
 			RenderCaptureProperties->CaptureFieldOfView = ComputedRenderCaptureProperties->CaptureFieldOfView;
 			RenderCaptureProperties->NearPlaneDist      = ComputedRenderCaptureProperties->NearPlaneDist;
+			Settings->ValidSampleDepthThreshold = ComputedValidDepthThreshold;
 
 			// Silently make the above updates so we don't overwrite the change to OpState below and call this function again
 			RenderCaptureProperties->SilentUpdateWatched();
+			Settings->SilentUpdateWatched();
 
 			OpState = EBakeOpState::Clean;
 
@@ -1134,8 +1214,11 @@ void UBakeRenderCaptureTool::UpdateResult()
 		ComputedRenderCaptureProperties->bSpecularMap       = RenderCaptureProperties->bSpecularMap;
 		ComputedRenderCaptureProperties->bPackedMRSMap      = RenderCaptureProperties->bPackedMRSMap;
 		ComputedRenderCaptureProperties->bEmissiveMap       = RenderCaptureProperties->bEmissiveMap;
+		ComputedRenderCaptureProperties->bAntiAliasing      = RenderCaptureProperties->bAntiAliasing;
+		ComputedRenderCaptureProperties->bDeviceDepthMap    = RenderCaptureProperties->bDeviceDepthMap;
 		ComputedRenderCaptureProperties->CaptureFieldOfView = RenderCaptureProperties->CaptureFieldOfView;
 		ComputedRenderCaptureProperties->NearPlaneDist      = RenderCaptureProperties->NearPlaneDist;
+		ComputedValidDepthThreshold = Settings->ValidSampleDepthThreshold;
 
 		bFirstEverSceneCapture = false;
 	}
diff --git a/Engine/Plugins/Experimental/MeshModelingToolsetExp/Source/MeshModelingToolsExp/Public/BakeRenderCaptureTool.h b/Engine/Plugins/Experimental/MeshModelingToolsetExp/Source/MeshModelingToolsExp/Public/BakeRenderCaptureTool.h
--- a/Engine/Plugins/Experimental/MeshModelingToolsetExp/Source/MeshModelingToolsExp/Public/BakeRenderCaptureTool.h
+++ b/Engine/Plugins/Experimental/MeshModelingToolsetExp/Source/MeshModelingToolsExp/Public/BakeRenderCaptureTool.h
@@ -91,7 +91,7 @@ public:
 	bool bBaseColorMap = true;
 
 	// Whether to generate a texture for the World Normal property
-	UPROPERTY(Category = RenderCaptureOptions, EditAnywhere, meta = (NoResetToDefault, EditCondition="bEnableNormalMap == true"))
+	UPROPERTY(Category = RenderCaptureOptions, EditAnywhere, meta = (NoResetToDefault))
 	bool bNormalMap = true;
 
 	// Whether to generate a packed texture with Metallic, Roughness and Specular properties
@@ -114,11 +114,21 @@ public:
 	UPROPERTY(Category = RenderCaptureOptions, EditAnywhere, meta = (NoResetToDefault))
 	bool bEmissiveMap = true;
 
-	UPROPERTY(Category = RenderCaptureOptions, EditAnywhere, meta = (ClampMin = "5.0", ClampMax = "160.0"), AdvancedDisplay, DisplayName="Field of View")
+	// Whether to use anti-aliasing in the render captures, this may introduce artefacts if pixels at different scene depths get blended
+	UPROPERTY(Category = RenderCaptureOptions, EditAnywhere, AdvancedDisplay)
+	bool bAntiAliasing = false;
+
+	// These are hidden in the UI right now, we might want to expose them if they turn out to be useful for very large
+	// or very small objects (not tested yet) TODO Figure out if we want to expose these options
+	
+	UPROPERTY(meta = (ClampMin = "5.0", ClampMax = "160.0"))
 	float CaptureFieldOfView = 30.0f;
 
-	UPROPERTY(Category = RenderCaptureOptions, EditAnywhere, meta = (ClampMin = "0.001", ClampMax = "1000.0"), AdvancedDisplay)
+	UPROPERTY(meta = (ClampMin = "0.001", ClampMax = "1000.0"))
 	float NearPlaneDist = 1.0f;
+	
+	UPROPERTY()
+	bool bDeviceDepthMap = false;
 
 	bool operator==(const URenderCaptureProperties& Other) const
 	{
@@ -131,8 +141,10 @@ public:
 			&& bSpecularMap == Other.bSpecularMap
 			&& bPackedMRSMap == Other.bPackedMRSMap
 			&& bEmissiveMap == Other.bEmissiveMap
+			&& bAntiAliasing == Other.bAntiAliasing
 			&& CaptureFieldOfView == Other.CaptureFieldOfView
-			&& NearPlaneDist == Other.NearPlaneDist;
+			&& NearPlaneDist == Other.NearPlaneDist
+			&& bDeviceDepthMap == Other.bDeviceDepthMap;
 	}
 
 	bool operator != (const URenderCaptureProperties& Other) const
@@ -165,8 +177,23 @@ public:
 	EBakeTextureSamplesPerPixel SamplesPerPixel = EBakeTextureSamplesPerPixel::Sample1;
 
 	/* Size of generated textures */
-	UPROPERTY(EditAnywhere, Category = BakeOutput, meta = (ClampMin="1", UIMin="1"), DisplayName="Results Resolution")
+	UPROPERTY(EditAnywhere, Category = BakeOutput, meta = (ClampMin="1", UIMin="1"), DisplayName="Output Resolution")
 	EBakeTextureResolution TextureSize = EBakeTextureResolution::Resolution512;
+
+	/**
+	 * This threshold is used to detect occlusion artefacts (e.g., silhouettes/blotches in the base color) so that texels
+	 * exhibiting them can be cleaned up i.e., rejected and filled in using the values of nearby valid texels instead.
+	 * 
+	 * If the threshold is zero, the cleanup step is skipped.
+	 * If the threshold is too large, texels with artefacts won't be detected and the cleanup step is effectively skipped.
+	 * If the threshold is too small, texels without artefacts get detected and results will be bad because there won't
+	 * be enough nearby valid texels from which to infill values.
+	 * 
+	 * A good starting point is to choose a threshold around the size of the distance (in world space) between the
+	 * target and source meshes; if the target was generated by VoxWrap then the voxel size estimates this distance.
+	 */
+	UPROPERTY(EditAnywhere, Category = BakeOutput, meta = (ClampMin="0", UIMin="0"), DisplayName="Cleanup Threshold")
+	float ValidSampleDepthThreshold = 0.f;
 };
 
 
@@ -177,7 +204,7 @@ class MESHMODELINGTOOLSEXP_API UBakeRenderCaptureInputToolProperties : public UI
 
 public:
 	/** Target mesh to sample to */
-	UPROPERTY(VisibleAnywhere, Category = BakeInput, DisplayName = "Target Mesh", meta = (TransientToolProperty))
+	UPROPERTY(VisibleAnywhere, Category = BakeInput, DisplayName = "Target Mesh", meta = (TransientToolProperty, NoResetToDefault))
 	TObjectPtr<UStaticMesh> TargetStaticMesh = nullptr;
 
 	/** UV channel to use for the target mesh */
@@ -296,6 +323,8 @@ protected:
 	// the capture are reverted to these values
 	UPROPERTY()
 	TObjectPtr<URenderCaptureProperties> ComputedRenderCaptureProperties;
+	float ComputedValidDepthThreshold;
+
 	TMap<int, FText> TargetUVLayerToError;
 
 	//
diff --git a/Engine/Plugins/Runtime/GeometryProcessing/Source/DynamicMesh/Public/Sampling/RenderCaptureMapEvaluator.h b/Engine/Plugins/Runtime/GeometryProcessing/Source/DynamicMesh/Public/Sampling/RenderCaptureMapEvaluator.h
--- a/Engine/Plugins/Runtime/GeometryProcessing/Source/DynamicMesh/Public/Sampling/RenderCaptureMapEvaluator.h
+++ b/Engine/Plugins/Runtime/GeometryProcessing/Source/DynamicMesh/Public/Sampling/RenderCaptureMapEvaluator.h
@@ -18,6 +18,7 @@ enum class ERenderCaptureChannel
 	Specular = 8,
 	Emissive = 16,
 	WorldNormal = 32,
+	DeviceDepth = 64,
 	CombinedMRS = 128
 };
 
diff --git a/Engine/Plugins/Runtime/MeshModelingToolset/Source/ModelingComponents/Private/Scene/SceneCapturePhotoSet.cpp b/Engine/Plugins/Runtime/MeshModelingToolset/Source/ModelingComponents/Private/Scene/SceneCapturePhotoSet.cpp
--- a/Engine/Plugins/Runtime/MeshModelingToolset/Source/ModelingComponents/Private/Scene/SceneCapturePhotoSet.cpp
+++ b/Engine/Plugins/Runtime/MeshModelingToolset/Source/ModelingComponents/Private/Scene/SceneCapturePhotoSet.cpp
@@ -8,6 +8,18 @@ using namespace UE::Geometry;
 
 #define LOCTEXT_NAMESPACE "SceneCapture"
 
+FSceneCapturePhotoSet::FSceneCapturePhotoSet()
+{
+	BaseColorConfig = GetDefaultRenderCaptureConfig(ERenderCaptureType::BaseColor);
+	WorldNormalConfig = GetDefaultRenderCaptureConfig(ERenderCaptureType::WorldNormal);
+	RoughnessConfig = GetDefaultRenderCaptureConfig(ERenderCaptureType::Roughness);
+	MetallicConfig = GetDefaultRenderCaptureConfig(ERenderCaptureType::Metallic);
+	SpecularConfig = GetDefaultRenderCaptureConfig(ERenderCaptureType::Specular);
+	EmissiveConfig = GetDefaultRenderCaptureConfig(ERenderCaptureType::Emissive);
+	PackedMRSConfig = GetDefaultRenderCaptureConfig(ERenderCaptureType::CombinedMRS);
+	DeviceDepthConfig = GetDefaultRenderCaptureConfig(ERenderCaptureType::DeviceDepth);
+}
+
 void FSceneCapturePhotoSet::SetCaptureSceneActors(UWorld* World, const TArray<AActor*>& Actors)
 {
 	this->TargetWorld = World;
@@ -15,6 +27,64 @@ void FSceneCapturePhotoSet::SetCaptureSceneActors(UWorld* World, const TArray<AA
 }
 
 
+void FSceneCapturePhotoSet::SetCaptureConfig(ERenderCaptureType CaptureType, const FRenderCaptureConfig& Config)
+{
+	switch (CaptureType)
+	{
+	case ERenderCaptureType::BaseColor:
+			BaseColorConfig = Config;
+			break;
+		case ERenderCaptureType::WorldNormal:
+			WorldNormalConfig = Config;
+			break;
+		case ERenderCaptureType::Roughness:
+			RoughnessConfig = Config;
+			break;
+		case ERenderCaptureType::Metallic:
+			MetallicConfig = Config;
+			break;
+		case ERenderCaptureType::Specular:
+			SpecularConfig = Config;
+			break;
+		case ERenderCaptureType::Emissive:
+			EmissiveConfig = Config;
+			break;
+		case ERenderCaptureType::CombinedMRS:
+			PackedMRSConfig = Config;
+			break;
+		case ERenderCaptureType::DeviceDepth:
+			DeviceDepthConfig = Config;
+			break;
+		default:
+			check(false);
+	}
+}
+
+FRenderCaptureConfig FSceneCapturePhotoSet::GetCaptureConfig(ERenderCaptureType CaptureType) const
+{
+	switch (CaptureType)
+	{
+	case ERenderCaptureType::BaseColor:
+			return BaseColorConfig;
+		case ERenderCaptureType::WorldNormal:
+			return WorldNormalConfig;
+		case ERenderCaptureType::Roughness:
+			return RoughnessConfig;
+		case ERenderCaptureType::Metallic:
+			return MetallicConfig;
+		case ERenderCaptureType::Specular:
+			return SpecularConfig;
+		case ERenderCaptureType::Emissive:
+			return EmissiveConfig;
+		case ERenderCaptureType::CombinedMRS:
+			return PackedMRSConfig;
+		case ERenderCaptureType::DeviceDepth:
+			return DeviceDepthConfig;
+		default:
+			check(false);
+	}
+	return BaseColorConfig;
+}
 
 void FSceneCapturePhotoSet::SetCaptureTypeEnabled(ERenderCaptureType CaptureType, bool bEnabled)
 {
@@ -41,11 +111,40 @@ void FSceneCapturePhotoSet::SetCaptureTypeEnabled(ERenderCaptureType CaptureType
 		case ERenderCaptureType::CombinedMRS:
 			bEnablePackedMRS = bEnabled;
 			break;
+		case ERenderCaptureType::DeviceDepth:
+			bEnableDeviceDepth = bEnabled;
+			break;
 		default:
 			check(false);
 	}
 }
 
+bool FSceneCapturePhotoSet::GetCaptureTypeEnabled(ERenderCaptureType CaptureType) const
+{
+	switch (CaptureType)
+	{
+	case ERenderCaptureType::BaseColor:
+			return bEnableBaseColor;
+		case ERenderCaptureType::WorldNormal:
+			return bEnableWorldNormal;
+		case ERenderCaptureType::Roughness:
+			return bEnableRoughness;
+		case ERenderCaptureType::Metallic:
+			return bEnableMetallic;
+		case ERenderCaptureType::Specular:
+			return bEnableSpecular;
+		case ERenderCaptureType::Emissive:
+			return bEnableEmissive;
+		case ERenderCaptureType::CombinedMRS:
+			return bEnablePackedMRS;
+		case ERenderCaptureType::DeviceDepth:
+			return bEnableDeviceDepth;
+		default:
+			ensure(false);
+	}
+	return false;
+}
+
 
 
 void FSceneCapturePhotoSet::AddStandardExteriorCapturesFromBoundingBox(
@@ -163,30 +262,25 @@ void FSceneCapturePhotoSet::AddExteriorCaptures(
 		FVector3d ViewDirection = Directions[di];
 		ViewDirection.Normalize();
 
+		// TODO Align the frame with the renderer coordinate system then remove the axis swapping in WorldRenderCapture.cpp
 		FFrame3d ViewFrame;
 		ViewFrame.AlignAxis(0, ViewDirection);
 		ViewFrame.ConstrainedAlignAxis(2, FVector3d::UnitZ(), ViewFrame.X());
 		ViewFrame.Origin = (FVector3d)RenderSphere.Center;
 		ViewFrame.Origin -= (double)RenderSphere.W * ViewFrame.X();
 
-		FSpatialPhotoParams Params;
-		Params.Frame = ViewFrame;
-		Params.NearPlaneDist = NearPlaneDist;
-		Params.HorzFOVDegrees = HorizontalFOVDegrees;
-		Params.Dimensions = PhotoDimensions;
-		PhotoSetParams.Add(Params);
-
 		FSpatialPhoto3f BasePhoto3f;
 		BasePhoto3f.Frame = ViewFrame;
 		BasePhoto3f.NearPlaneDist = NearPlaneDist;
 		BasePhoto3f.HorzFOVDegrees = HorizontalFOVDegrees;
 		BasePhoto3f.Dimensions = PhotoDimensions;
 
-		auto CaptureImageTypeFunc_3f = [&Progress, &BasePhoto3f, &RenderCapture](ERenderCaptureType CaptureType, FSpatialPhotoSet3f& PhotoSet)
+		auto CaptureImageTypeFunc_3f = [this, &Progress, &BasePhoto3f, &RenderCapture](ERenderCaptureType CaptureType, FSpatialPhotoSet3f& PhotoSet)
 		{
 			FSpatialPhoto3f NewPhoto = BasePhoto3f;
 			FImageAdapter Image(&NewPhoto.Image);
-			RenderCapture.CaptureFromPosition(CaptureType, NewPhoto.Frame, NewPhoto.HorzFOVDegrees, NewPhoto.NearPlaneDist, Image);
+			FRenderCaptureConfig Config = GetCaptureConfig(CaptureType);
+			RenderCapture.CaptureFromPosition(CaptureType, NewPhoto.Frame, NewPhoto.HorzFOVDegrees, NewPhoto.NearPlaneDist, Image, Config);
 			PhotoSet.Add(MoveTemp(NewPhoto));
 			Progress.TickProgress();
 		};
@@ -197,15 +291,20 @@ void FSceneCapturePhotoSet::AddExteriorCaptures(
 		BasePhoto1f.HorzFOVDegrees = HorizontalFOVDegrees;
 		BasePhoto1f.Dimensions = PhotoDimensions;
 
-		auto CaptureImageTypeFunc_1f = [&Progress, &BasePhoto1f, &RenderCapture](ERenderCaptureType CaptureType, FSpatialPhotoSet1f& PhotoSet)
+		auto CaptureImageTypeFunc_1f = [this, &Progress, &BasePhoto1f, &RenderCapture](ERenderCaptureType CaptureType, FSpatialPhotoSet1f& PhotoSet)
 		{
 			FSpatialPhoto1f NewPhoto = BasePhoto1f;
 			FImageAdapter Image(&NewPhoto.Image);
-			RenderCapture.CaptureFromPosition(CaptureType, NewPhoto.Frame, NewPhoto.HorzFOVDegrees, NewPhoto.NearPlaneDist, Image);
+			FRenderCaptureConfig Config = GetCaptureConfig(CaptureType);
+			RenderCapture.CaptureFromPosition(CaptureType, NewPhoto.Frame, NewPhoto.HorzFOVDegrees, NewPhoto.NearPlaneDist, Image, Config);
 			PhotoSet.Add(MoveTemp(NewPhoto));
 			Progress.TickProgress();
 		};
 
+		if (bEnableDeviceDepth)
+		{
+			CaptureImageTypeFunc_1f(ERenderCaptureType::DeviceDepth, DeviceDepthPhotoSet);
+		}
 		if (bEnableBaseColor)
 		{
 			CaptureImageTypeFunc_3f(ERenderCaptureType::BaseColor, BaseColorPhotoSet);
@@ -234,6 +333,14 @@ void FSceneCapturePhotoSet::AddExteriorCaptures(
 		{
 			CaptureImageTypeFunc_3f(ERenderCaptureType::Emissive, EmissivePhotoSet);
 		}
+
+		FSpatialPhotoParams Params;
+		Params.Frame = ViewFrame;
+		Params.NearPlaneDist = NearPlaneDist;
+		Params.HorzFOVDegrees = HorizontalFOVDegrees;
+		Params.Dimensions = PhotoDimensions;
+		Params.ViewMatrices = RenderCapture.GetLastCaptureViewMatrices();
+		PhotoSetParams.Add(Params);
 	}
 }
 
@@ -258,6 +365,7 @@ FSceneCapturePhotoSet::FSceneSample::FSceneSample()
 	Metallic = 0.0f;
 	Emissive = FVector3f(0, 0, 0);
 	WorldNormal = FVector3f(0, 0, 1);
+	DeviceDepth = 0.0f;
 }
 
 FVector3f FSceneCapturePhotoSet::FSceneSample::GetValue3f(ERenderCaptureType CaptureType) const
@@ -276,6 +384,8 @@ FVector3f FSceneCapturePhotoSet::FSceneSample::GetValue3f(ERenderCaptureType Cap
 		return Specular * FVector3f::One();
 	case ERenderCaptureType::Emissive:
 		return Emissive;
+	case ERenderCaptureType::DeviceDepth:
+		return DeviceDepth * FVector3f::One();
 	default:
 		check(false);
 	}
@@ -300,6 +410,8 @@ FVector4f FSceneCapturePhotoSet::FSceneSample::GetValue4f(ERenderCaptureType Cap
 		return FVector4f(Metallic, Roughness, Specular, 1.0f);
 	case ERenderCaptureType::Emissive:
 		return FVector4f(Emissive.X, Emissive.Y, Emissive.Z, 1.0f);
+	case ERenderCaptureType::DeviceDepth:
+		return FVector4f(DeviceDepth, DeviceDepth, DeviceDepth, 1.0f);
 	default:
 		check(false);
 	}
@@ -309,6 +421,7 @@ FVector4f FSceneCapturePhotoSet::FSceneSample::GetValue4f(ERenderCaptureType Cap
 bool FSceneCapturePhotoSet::ComputeSampleLocation(
 	const FVector3d& Position,
 	const FVector3d& Normal,
+	const float ValidSampleDepthThreshold,
 	TFunctionRef<bool(const FVector3d&, const FVector3d&)> VisibilityFunction,
 	int& PhotoIndex,
 	FVector2d& PhotoCoords) const
@@ -326,44 +439,82 @@ bool FSceneCapturePhotoSet::ComputeSampleLocation(
 		const FSpatialPhotoParams& Params = PhotoSetParams[Index];
 		check(Params.Dimensions.IsSquare());
 
-		FVector3d ViewDirection = Params.Frame.X();
+		FFrame3d RenderFrame(Params.Frame.Origin, Params.Frame.Y(), Params.Frame.Z(), Params.Frame.X());
+
+		FVector3d ViewDirection = RenderFrame.Z();
 		double ViewDot = ViewDirection.Dot(Normal);
 		if (ViewDot > DotTolerance || ViewDot > MinDot)
 		{
+			// The sample is facing away from the photo, or we found a photo more aligned with this sample
 			continue;
 		}
 
-		FFrame3d ViewPlane = Params.Frame;
+		FFrame3d ViewPlane = RenderFrame;
 		ViewPlane.Origin += Params.NearPlaneDist * ViewDirection;
 
 		double ViewPlaneWidthWorld = Params.NearPlaneDist * FMathd::Tan(Params.HorzFOVDegrees * 0.5 * FMathd::DegToRad);
 		double ViewPlaneHeightWorld = ViewPlaneWidthWorld;
 
-		FVector3d RayOrigin = Params.Frame.Origin;
+		// Shoot a ray from the camera position toward the sample position and find the hit point on photo plane
+		constexpr int NormalAxisIndex = 2;
+		FVector3d RayOrigin = RenderFrame.Origin;
 		FVector3d RayDir = Normalized(Position - RayOrigin);
 		FVector3d HitPoint;
-		bool bHit = ViewPlane.RayPlaneIntersection(RayOrigin, RayDir, 0, HitPoint);
+		bool bHit = ViewPlane.RayPlaneIntersection(RayOrigin, RayDir, NormalAxisIndex, HitPoint);
 		if (bHit)
 		{
-			bool bVisible = VisibilityFunction(Position, HitPoint);
-			if ( bVisible )
+			FVector2d DeviceXY;
+			DeviceXY.X = (HitPoint - ViewPlane.Origin).Dot(ViewPlane.X()) / ViewPlaneWidthWorld;
+			DeviceXY.Y = (HitPoint - ViewPlane.Origin).Dot(ViewPlane.Y()) / ViewPlaneHeightWorld;
+			if (FMathd::Abs(DeviceXY.X) < 1 && FMathd::Abs(DeviceXY.Y) < 1)
 			{
-				double PlaneX = (HitPoint - ViewPlane.Origin).Dot(ViewPlane.Y());
-				double PlaneY = (HitPoint - ViewPlane.Origin).Dot(ViewPlane.Z());
-
-				//FVector2d PlanePos = ViewPlane.ToPlaneUV(HitPoint, 0);
-				double u = PlaneX / ViewPlaneWidthWorld;
-				double v = -(PlaneY / ViewPlaneHeightWorld);
-				if (FMathd::Abs(u) < 1 && FMathd::Abs(v) < 1)
+				// Shoot a ray from the sample position toward the camera position checking occlusion
+				bool bVisible = VisibilityFunction(Position, RayOrigin);
+				if (bVisible)
 				{
-					PhotoCoords.X = (u/2.0 + 0.5) * (double)Params.Dimensions.GetWidth();
-					PhotoCoords.Y = (v/2.0 + 0.5) * (double)Params.Dimensions.GetHeight();
-					PhotoIndex = Index;
-					MinDot = ViewDot;
-				}
-			}
-		}
-	}
+					FVector2d UVCoords = FRenderCaptureCoordinateConverter2D::DeviceToUV(DeviceXY);
+					if (ValidSampleDepthThreshold > 0)
+					{
+						// Look up the device depth from the depth photo set, these values are from 0 (far plane) to
+						// 1 (near plane). We skip points which would unproject to the far plane, which is positioned at
+						// infinity, we also do not interpolate the depth values, doing so is not a good approximation
+						// of the underlying scene
+						float DeviceZ = DeviceDepthPhotoSet.Get(Index).Image.NearestSampleUV(UVCoords);
+						if (DeviceZ > 0)
+						{
+							// Compute the pixel position in world space to use it to compute a depth according to the render
+							FVector3d PixelPositionDevice{DeviceXY, DeviceZ};
+							FVector4d PixelPositionWorld = Params.ViewMatrices.GetInvViewProjectionMatrix().TransformPosition(PixelPositionDevice);
+							PixelPositionWorld /= PixelPositionWorld.W;
+
+							// Compare the depth of the sample with the depth of the pixel and consider the sample invalid
+							// if these do not match closely enough. This fixes artefacts which occur when sample ray just
+							// misses an obstruction and hits a pixel where the color was set by a slightly different ray,
+							// through the pixel center, which does hit the obstruction. This problem occurs because the
+							// depth capture was obtained by renderering the the source meshes but the visiblity function
+							// works on the target mesh
+							float PixelDepth = (RayOrigin - FVector3d(PixelPositionWorld)).Length();
+							float SampleDepth = (RayOrigin - Position).Length();
+							if (FMath::Abs(PixelDepth - SampleDepth) < ValidSampleDepthThreshold)
+							{
+								PhotoCoords.X = UVCoords.X * (double)Params.Dimensions.GetWidth();
+								PhotoCoords.Y = UVCoords.Y * (double)Params.Dimensions.GetHeight();
+								PhotoIndex = Index;
+								MinDot = ViewDot;
+							}
+						} // Test DeviceZ > 0
+					} 
+					else
+					{
+						PhotoCoords.X = UVCoords.X * (double)Params.Dimensions.GetWidth();
+						PhotoCoords.Y = UVCoords.Y * (double)Params.Dimensions.GetHeight();
+						PhotoIndex = Index;
+						MinDot = ViewDot;
+					}
+				} // Test bVisible
+			} // Test UVCoords in (-1,1)x(-1,1)
+		} // Hit photo plane
+	} // Photo loop
 
 	return PhotoIndex != IndexConstants::InvalidID;
 }
diff --git a/Engine/Plugins/Runtime/MeshModelingToolset/Source/ModelingComponents/Private/Scene/WorldRenderCapture.cpp b/Engine/Plugins/Runtime/MeshModelingToolset/Source/ModelingComponents/Private/Scene/WorldRenderCapture.cpp
--- a/Engine/Plugins/Runtime/MeshModelingToolset/Source/ModelingComponents/Private/Scene/WorldRenderCapture.cpp
+++ b/Engine/Plugins/Runtime/MeshModelingToolset/Source/ModelingComponents/Private/Scene/WorldRenderCapture.cpp
@@ -5,6 +5,7 @@
 
 #include "Engine/Canvas.h"
 #include "Engine/TextureRenderTarget2D.h"
+#include "GameFramework/WorldSettings.h"
 #include "LegacyScreenPercentageDriver.h"
 #include "EngineModule.h"
 #include "Rendering/Texture2DResource.h"
@@ -15,6 +16,9 @@
 #include "Components/DirectionalLightComponent.h"
 #include "Components/ChildActorComponent.h"
 
+#include "SceneViewExtension.h"
+#include "RenderCaptureInterface.h" // For debugging with RenderCaptureInterface::FScopedCapture
+
 using namespace UE::Geometry;
 
 static TAutoConsoleVariable<int32> CVarModelingWorldRenderCaptureVTWarmupFrames(
@@ -24,22 +28,25 @@ static TAutoConsoleVariable<int32> CVarModelingWorldRenderCaptureVTWarmupFrames(
 
 FRenderCaptureTypeFlags FRenderCaptureTypeFlags::All()
 {
-	return FRenderCaptureTypeFlags{ true, true, true, true, true, true };
+	// TODO Cleanup this function, its confusing due to CombinedMRS overlapping with the individual channels
+	FRenderCaptureTypeFlags Result{ true, true, true, true, true, true };
+	Result.bDeviceDepth = true;
+	return Result;
 }
 
 FRenderCaptureTypeFlags FRenderCaptureTypeFlags::None()
 {
-	return FRenderCaptureTypeFlags{ false, false, false, false, false, false };
+	return FRenderCaptureTypeFlags{};
 }
 
 FRenderCaptureTypeFlags FRenderCaptureTypeFlags::BaseColor()
 {
-	return FRenderCaptureTypeFlags{ true, false, false, false, false, false };
+	return Single(ERenderCaptureType::BaseColor);
 }
 
 FRenderCaptureTypeFlags FRenderCaptureTypeFlags::WorldNormal()
 {
-	return FRenderCaptureTypeFlags{ false, false, false, false, false, true };
+	return Single(ERenderCaptureType::WorldNormal);
 }
 
 FRenderCaptureTypeFlags FRenderCaptureTypeFlags::Single(ERenderCaptureType CaptureType)
@@ -74,6 +81,9 @@ void FRenderCaptureTypeFlags::SetEnabled(ERenderCaptureType CaptureType, bool bE
 	case ERenderCaptureType::CombinedMRS:
 		bCombinedMRS = bEnabled;
 		break;
+	case ERenderCaptureType::DeviceDepth:
+		bDeviceDepth = bEnabled;
+		break;
 	default:
 		check(false);
 	}
@@ -81,6 +91,27 @@ void FRenderCaptureTypeFlags::SetEnabled(ERenderCaptureType CaptureType, bool bE
 
 
 
+namespace UE
+{
+namespace Geometry
+{
+
+FRenderCaptureConfig GetDefaultRenderCaptureConfig(ERenderCaptureType CaptureType)
+{
+	FRenderCaptureConfig Config;
+
+	// The DeviceDepth channel ignores this option, but since its hard-coded to false we set it to false
+	Config.bAntiAliasing = (CaptureType != ERenderCaptureType::DeviceDepth);
+
+	return Config;
+}
+
+}
+}
+
+
+
+
 
 FWorldRenderCapture::FWorldRenderCapture()
 {
@@ -104,6 +135,11 @@ void FWorldRenderCapture::Shutdown()
 		GammaRenderTexture->RemoveFromRoot();
 		GammaRenderTexture = nullptr;
 	}
+	if (DepthRenderTexture != nullptr)
+	{
+		DepthRenderTexture->RemoveFromRoot();
+		DepthRenderTexture = nullptr;
+	}
 }
 
 
@@ -230,6 +266,50 @@ UTextureRenderTarget2D* FWorldRenderCapture::GetRenderTexture(bool bLinear)
 }
 
 
+UTextureRenderTarget2D* FWorldRenderCapture::GetDepthRenderTexture()
+{
+	if (RenderTextureDimensions != this->Dimensions)
+	{
+		if (DepthRenderTexture != nullptr)
+		{
+			DepthRenderTexture->RemoveFromRoot();
+			DepthRenderTexture = nullptr;
+		}
+	}
+
+	if (DepthRenderTexture != nullptr)
+	{
+		return DepthRenderTexture;
+	}
+
+	RenderTextureDimensions = Dimensions;
+	int32 Width = Dimensions.GetWidth();
+	int32 Height = Dimensions.GetHeight();
+
+	DepthRenderTexture = NewObject<UTextureRenderTarget2D>();
+	if (ensure(DepthRenderTexture))
+	{
+		DepthRenderTexture->AddToRoot();		// keep alive for GC
+		DepthRenderTexture->RenderTargetFormat = ETextureRenderTargetFormat::RTF_RGBA32f;
+		DepthRenderTexture->ClearColor = FLinearColor::Black;
+		DepthRenderTexture->InitAutoFormat(Width, Height);
+		DepthRenderTexture->UpdateResourceImmediate(true);
+
+		// Flush RHI thread after creating texture render target to make sure that RHIUpdateTextureReference is executed
+		// before doing any rendering with it This makes sure that
+		//    Value->TextureReference.TextureReferenceRHI->GetReferencedTexture()
+		// is valid so that
+		//    UniformExpressionSet::FillUniformBuffer
+		// properly uses the texture for rendering, instead of using a fallback texture
+		ENQUEUE_RENDER_COMMAND(FlushRHIThreadToUpdateTextureRenderTargetReference)(
+		[](FRHICommandListImmediate& RHICmdList)
+		{
+			RHICmdList.ImmediateFlush(EImmediateFlushType::FlushRHIThread);
+		});
+	}
+
+	return DepthRenderTexture;
+}
 
 
 namespace UE 
@@ -331,7 +411,8 @@ bool FWorldRenderCapture::CaptureMRSFromPosition(
 	const FFrame3d& ViewFrame,
 	double HorzFOVDegrees,
 	double NearPlaneDist,
-	FImageAdapter& ResultImageOut)
+	FImageAdapter& ResultImageOut,
+	const FRenderCaptureConfig& Config)
 {
 	// this post-process material renders an image with R=Metallic, G=Roughness, B=Specular, A=AmbientOcclusion
 	FString MRSPostProcessMaterialAssetPath = TEXT("/MeshModelingToolsetExp/Materials/PostProcess_PackedMRSA.PostProcess_PackedMRSA");
@@ -376,7 +457,7 @@ bool FWorldRenderCapture::CaptureMRSFromPosition(
 	ApplyViewMode(ViewModeIndex, true, ShowFlags);
 
 	// unclear if these flags need to be set before creating ViewFamily
-	ShowFlags.SetAntiAliasing(false);
+	ShowFlags.SetAntiAliasing(Config.bAntiAliasing);
 	ShowFlags.SetDepthOfField(false);
 	ShowFlags.SetMotionBlur(false);
 	ShowFlags.SetBloom(false);
@@ -468,6 +549,9 @@ bool FWorldRenderCapture::CaptureMRSFromPosition(
 
 	UE::Internal::PerformSceneRender(&Canvas, &ViewFamily);
 
+	// Cache the view/projection matricies we used to render the scene
+	LastCaptureViewMatrices = NewView->ViewMatrices;
+
 	// wait for render
 	FlushRenderingCommands();
 
@@ -489,7 +573,8 @@ bool FWorldRenderCapture::CaptureEmissiveFromPosition(
 	const FFrame3d& Frame,
 	double HorzFOVDegrees,
 	double NearPlaneDist,
-	FImageAdapter& ResultImageOut)
+	FImageAdapter& ResultImageOut,
+	const FRenderCaptureConfig& Config)
 {
 	UTextureRenderTarget2D* RenderTargetTexture = GetRenderTexture(true);
 	if (ensure(RenderTargetTexture) == false)
@@ -518,7 +603,7 @@ bool FWorldRenderCapture::CaptureEmissiveFromPosition(
 
 	// unclear if these flags need to be set before creating ViewFamily
 	FEngineShowFlags ShowFlags(ESFIM_Game);
-	ShowFlags.SetAntiAliasing(false);
+	ShowFlags.SetAntiAliasing(Config.bAntiAliasing);
 	ShowFlags.SetDepthOfField(false);
 	ShowFlags.SetMotionBlur(false);
 	ShowFlags.SetBloom(false);
@@ -579,6 +664,9 @@ bool FWorldRenderCapture::CaptureEmissiveFromPosition(
 
 	UE::Internal::PerformSceneRender(&Canvas, &ViewFamily);
 
+	// Cache the view/projection matricies we used to render the scene
+	LastCaptureViewMatrices = NewView->ViewMatrices;
+
 	// Copy the contents of the remote texture to system memory
 	ReadImageBuffer.Reset();
 	ReadImageBuffer.SetNumUninitialized(Width * Height);
@@ -600,6 +688,166 @@ bool FWorldRenderCapture::CaptureEmissiveFromPosition(
 }
 
 
+bool FWorldRenderCapture::CaptureDeviceDepthFromPosition(
+	const FFrame3d& Frame,
+	double HorzFOVDegrees,
+	double NearPlaneDist,
+	FImageAdapter& ResultImageOut,
+	const FRenderCaptureConfig& Config)
+{
+	DepthRenderTexture = GetDepthRenderTexture();
+	if (ensure(DepthRenderTexture) == false)
+	{
+		return false;
+	}
+	FTextureRenderTargetResource* RenderTargetResource = DepthRenderTexture->GameThread_GetRenderTargetResource();
+
+	int32 Width = Dimensions.GetWidth();
+	int32 Height = Dimensions.GetHeight();
+
+	// TODO The input frame should have been the render frame the whole time
+	FFrame3d RenderFrame(Frame.Origin, Frame.Y(), Frame.Z(), Frame.X());
+	FVector ViewOrigin = (FVector)RenderFrame.Origin;
+	FMatrix ViewRotationMatrix = FRotationMatrix::Make((FQuat)RenderFrame.Rotation.Inverse());
+
+	const float HalfFOVRadians = FMath::DegreesToRadians<float>(HorzFOVDegrees) * 0.5f;
+	static_assert((int32)ERHIZBuffer::IsInverted != 0, "Check NearPlane and Projection Matrix");
+	FMatrix ProjectionMatrix = FReversedZPerspectiveMatrix(HalfFOVRadians, 1.0f, 1.0f, (float)NearPlaneDist);
+
+	// unclear if these flags need to be set before creating ViewFamily
+	FEngineShowFlags ShowFlags(ESFIM_Game);
+	ShowFlags.SetAntiAliasing(false); // Intentionally does NOT use Config.bAntiAliasing
+	ShowFlags.SetDepthOfField(false);
+	ShowFlags.SetMotionBlur(false);
+	ShowFlags.SetBloom(false);
+	ShowFlags.SetSceneColorFringe(false);
+	ShowFlags.SetNaniteMeshes(false);
+	ShowFlags.SetAtmosphere(false);
+	ShowFlags.SetLighting(false);
+	ShowFlags.SetScreenPercentage(false);
+	ShowFlags.SetTranslucency(false);
+	ShowFlags.SetSeparateTranslucency(false);
+	ShowFlags.SetFog(false);
+	ShowFlags.SetVolumetricFog(false);
+	ShowFlags.SetDynamicShadows(false);
+
+	FSceneViewFamilyContext ViewFamily(
+		FSceneViewFamily::ConstructionValues(
+			RenderTargetResource,
+			World->Scene,
+			ShowFlags)
+		.SetRealtimeUpdate(false)
+		.SetResolveScene(false)
+	);
+
+	// Request a scene depth render
+	ViewFamily.SceneCaptureSource = ESceneCaptureSource::SCS_DeviceDepth;
+
+	FSceneViewInitOptions ViewInitOptions;
+	
+	ViewInitOptions.SetViewRectangle(FIntRect(0, 0, Width, Height));
+	ViewInitOptions.ViewFamily = &ViewFamily;
+	if (VisiblePrimitives.Num() > 0)
+	{
+		ViewInitOptions.ShowOnlyPrimitives = VisiblePrimitives;
+	}
+	ViewInitOptions.ViewOrigin = ViewOrigin;
+	ViewInitOptions.ViewRotationMatrix = ViewRotationMatrix;
+	ViewInitOptions.ProjectionMatrix = ProjectionMatrix;
+
+	ViewInitOptions.OverrideFarClippingPlaneDistance = -1.f;
+
+	if (ViewFamily.Scene->GetWorld() != nullptr && ViewFamily.Scene->GetWorld()->GetWorldSettings() != nullptr)
+	{
+		ViewInitOptions.WorldToMetersScale = ViewFamily.Scene->GetWorld()->GetWorldSettings()->WorldToMeters;
+	}
+
+	FSceneView* NewView = new FSceneView(ViewInitOptions);
+	NewView->AntiAliasingMethod = EAntiAliasingMethod::AAM_None;
+	NewView->SetupAntiAliasingMethod();
+	ViewFamily.Views.Add(NewView);
+
+	NewView->StartFinalPostprocessSettings(ViewOrigin);
+	NewView->EndFinalPostprocessSettings(ViewInitOptions);
+	ViewFamily.SetScreenPercentageInterface(new FLegacyScreenPercentageDriver(ViewFamily, 1.f));
+	
+	ViewFamily.ViewExtensions = GEngine->ViewExtensions->GatherActiveExtensions(FSceneViewExtensionContext(World->Scene));
+	for (const FSceneViewExtensionRef& Extension : ViewFamily.ViewExtensions)
+	{
+		Extension->SetupViewFamily(ViewFamily);
+		Extension->SetupView(ViewFamily, *NewView);
+	}
+
+	FCanvas Canvas(RenderTargetResource, NULL, FGameTime::GetTimeSinceAppStart(), World->Scene->GetFeatureLevel());
+	Canvas.Clear(FLinearColor::Transparent);
+
+	UE::Internal::PerformSceneRender(&Canvas, &ViewFamily);
+
+	// Cache the view/projection matricies we used to render the scene
+	LastCaptureViewMatrices = NewView->ViewMatrices;
+
+	// Copy the contents of the remote texture to system memory
+	TArray<FLinearColor> ReadImageColorBuffer;
+	ReadImageColorBuffer.SetNumUninitialized(Width * Height);
+	FReadSurfaceDataFlags ReadSurfaceDataFlags(RCM_MinMax);
+	ReadSurfaceDataFlags.SetLinearToGamma(false);
+
+	RenderTargetResource->ReadLinearColorPixels(ReadImageColorBuffer, ReadSurfaceDataFlags, FIntRect(0, 0, Width, Height));
+	ResultImageOut.SetDimensions(Dimensions);
+	for (int32 yi = 0; yi < Height; ++yi)
+	{
+		for (int32 xi = 0; xi < Width; ++xi)
+		{
+			FLinearColor Color = ReadImageColorBuffer[yi * Width + xi];
+
+			// Reverse the float encoding used for the Depth value in SceneCapturePixelShader.usf (one minus the device
+			// z is encoded in the RGB components). See DecodeFloatRGB in WaterInfoMerge.usf and, for a reference on the
+			// method, see https://aras-p.info/blog/2009/07/30/encoding-floats-to-rgba-the-final/
+			FVector3f EncodedDepth(Color.R, Color.G, Color.B);
+			float Depth = FVector3f::DotProduct(EncodedDepth, FVector3f(1.0, 1/255.0, 1/65025.0));
+
+			// Reverse the expression used to compute the Depth to recover the DeviceZ aka normalized device coordinate z
+			float DeviceZ = -(Depth - 1);
+			ensure(DeviceZ >= 0. && DeviceZ <= 1.); // Points on the near plane have Z=1, points on the far plane have Z=0
+
+			ResultImageOut.SetPixel(FVector2i(xi, yi), FLinearColor(DeviceZ, 0., 0., 0.));
+		}
+	}
+
+	// Set this to true to compute a world point cloud for debugging.
+	// You probably want to change the logging so it writes an .obj file
+	constexpr bool bDebugDepthCapture = false;
+
+	if constexpr (bDebugDepthCapture)
+	{
+		int PointIndex = 1;
+		for (int32 yi = 0; yi < Height; ++yi)
+		{
+			for (int32 xi = 0; xi < Width; ++xi)
+			{
+				float DeviceZ = ResultImageOut.GetPixel(FVector2i(xi, yi)).X;
+				if (DeviceZ > 0.) // Skip points on the far plane since these unproject to infinity
+				{
+					// Map from pixel space to NDC space
+					FVector2d DeviceXY = FRenderCaptureCoordinateConverter2D::PixelToDevice(FVector2i(xi, yi), Width, Height);
+
+					// Compute world coordinates from normalized device coordinates
+					FVector4d Point = GetLastCaptureViewMatrices().GetInvViewProjectionMatrix().TransformPosition(FVector3d(DeviceXY, DeviceZ));
+					Point /= Point.W;
+
+					// Log the point
+					UE_LOG(LogGeometry, Log, TEXT("DebugDepthCapture: [%d] %s"), PointIndex, *Point.ToString());
+					PointIndex += 1;
+				}
+			} // xi
+		} // yi
+	}
+
+	return true;
+}
+
+
+
 /**
  * internal Utility function to render the given Scene to a render target and capture
  * one of the render buffers, defined by VisualizationMode. Not clear where
@@ -617,9 +865,11 @@ namespace Internal
 		const FVector& ViewOrigin,
 		const FMatrix& ViewRotationMatrix,
 		const FMatrix& ProjectionMatrix,
+		const FRenderCaptureConfig& Config,
 		const TSet<FPrimitiveComponentId>& HiddenPrimitives,		// these primitives will be hidden
 		const TSet<FPrimitiveComponentId>& VisiblePrimitives,		// if non-empty, only these primitives are shown
-		TArray<FLinearColor>& OutSamples
+		TArray<FLinearColor>& OutSamples,
+		FViewMatrices& LastCaptureViewMatrices
 	)
 	{
 		int32 Width = Dimensions.GetWidth();
@@ -636,6 +886,7 @@ namespace Internal
 		ViewFamily.EngineShowFlags.SetVisualizeBuffer(true);
 		ViewFamily.EngineShowFlags.SetTonemapper(false);
 		ViewFamily.EngineShowFlags.SetScreenPercentage(false);
+		ViewFamily.EngineShowFlags.SetAntiAliasing(Config.bAntiAliasing);
 
 		FSceneViewInitOptions ViewInitOptions;
 		ViewInitOptions.SetViewRectangle(FIntRect(0, 0, Width, Height));
@@ -661,6 +912,9 @@ namespace Internal
 
 		UE::Internal::PerformSceneRender(&Canvas, &ViewFamily);
 
+		// Cache the view/projection matricies we used to render the scene
+		LastCaptureViewMatrices = NewView->ViewMatrices;
+
 		// Copy the contents of the remote texture to system memory
 		OutSamples.SetNumUninitialized(Width * Height);
 		//FReadSurfaceDataFlags ReadSurfaceDataFlags(RCM_MinMax);		// should we use MinMax to avoid normalization?
@@ -672,17 +926,28 @@ namespace Internal
 }
 
 
+// This is useful to debug the GPU state of the render captures in RenderDoc. Usage: set `r.RenderCaptureDraw 1` in the
+// and open the capture in RenderDoc (which may have opened automatically)
+static int32 RenderCaptureDraws = 0;
+static FAutoConsoleVariableRef CVarRenderCaptureDraws(
+	TEXT("r.RenderCaptureDraws"),
+	RenderCaptureDraws,
+	TEXT("Enable capturing of render capture texture for the next N draws"));
 
 bool FWorldRenderCapture::CaptureFromPosition(
 	ERenderCaptureType CaptureType,
 	const FFrame3d& ViewFrame,
 	double HorzFOVDegrees,
 	double NearPlaneDist,
-	FImageAdapter& ResultImageOut)
+	FImageAdapter& ResultImageOut,
+	const FRenderCaptureConfig& Config)
 {
+	RenderCaptureInterface::FScopedCapture RenderCapture(RenderCaptureDraws > 0, TEXT("RenderCaptureFromPosition"));
+	RenderCaptureDraws--;
+
 	if (CaptureType == ERenderCaptureType::Emissive)
 	{
-		bool bOK = CaptureEmissiveFromPosition(ViewFrame, HorzFOVDegrees, NearPlaneDist, ResultImageOut);
+		bool bOK = CaptureEmissiveFromPosition(ViewFrame, HorzFOVDegrees, NearPlaneDist, ResultImageOut, Config);
 		if (bWriteDebugImage)
 		{
 			WriteDebugImage(ResultImageOut, TEXT("Emissive"));
@@ -691,13 +956,23 @@ bool FWorldRenderCapture::CaptureFromPosition(
 	}
 	else if (CaptureType == ERenderCaptureType::CombinedMRS)
 	{
-		bool bOK = CaptureMRSFromPosition(ViewFrame, HorzFOVDegrees, NearPlaneDist, ResultImageOut);
+		bool bOK = CaptureMRSFromPosition(ViewFrame, HorzFOVDegrees, NearPlaneDist, ResultImageOut, Config);
 		if (bWriteDebugImage)
 		{
 			WriteDebugImage(ResultImageOut, TEXT("CombinedMRS"));
 		}
 		return bOK;
 	}
+	else if (CaptureType == ERenderCaptureType::DeviceDepth)
+	{
+		
+		bool bOK = CaptureDeviceDepthFromPosition(ViewFrame, HorzFOVDegrees, NearPlaneDist, ResultImageOut, Config);
+		if (bWriteDebugImage)
+		{
+			WriteDebugImage(ResultImageOut, TEXT("DeviceDepth"));
+		}
+		return bOK;
+	}
 
 	// Roughness visualization is rendered with gamma correction (unclear why)
 	bool bLinear = (CaptureType != ERenderCaptureType::Roughness);
@@ -732,7 +1007,6 @@ bool FWorldRenderCapture::CaptureFromPosition(
 	case ERenderCaptureType::Roughness:		CaptureTypeName = FName("Roughness");	break;
 	case ERenderCaptureType::Metallic:		CaptureTypeName = FName("Metallic");	break;
 	case ERenderCaptureType::Specular:		CaptureTypeName = FName("Specular");	break;
-	case ERenderCaptureType::Emissive:		CaptureTypeName = FName("PreTonemapHDRColor");	break;
 	case ERenderCaptureType::BaseColor:
 	default:
 		CaptureTypeName = FName("BaseColor");
@@ -743,10 +1017,11 @@ bool FWorldRenderCapture::CaptureFromPosition(
 	UE::Internal::RenderSceneVisualizationToTexture( 
 		RenderTargetTexture, this->Dimensions,
 		World->Scene, CaptureTypeName,
-		ViewOrigin, ViewRotationMatrix, ProjectionMatrix,
+		ViewOrigin, ViewRotationMatrix, ProjectionMatrix, Config,
 		HiddenPrimitives,
 		VisiblePrimitives,
-		ReadImageBuffer);
+		ReadImageBuffer,
+		LastCaptureViewMatrices);
 
 	ResultImageOut.SetDimensions(Dimensions);
 	for (int32 yi = 0; yi < Height; ++yi)
diff --git a/Engine/Plugins/Runtime/MeshModelingToolset/Source/ModelingComponents/Public/Scene/SceneCapturePhotoSet.h b/Engine/Plugins/Runtime/MeshModelingToolset/Source/ModelingComponents/Public/Scene/SceneCapturePhotoSet.h
--- a/Engine/Plugins/Runtime/MeshModelingToolset/Source/ModelingComponents/Public/Scene/SceneCapturePhotoSet.h
+++ b/Engine/Plugins/Runtime/MeshModelingToolset/Source/ModelingComponents/Public/Scene/SceneCapturePhotoSet.h
@@ -19,6 +19,7 @@ namespace Geometry
  * FSceneCapturePhotoSet creates a set of render captures for a given World and set of Actors,
  * stored as a SpatialPhotoSet for each desired render buffer type. Currently the set of buffers are
  * defined by ERenderCaptureType:
+ * 
  *		BaseColor
  *		Roughness
  *		Metallic
@@ -26,6 +27,8 @@ namespace Geometry
  *		PackedMRS   (Metallic / Roughness / Specular)
  *		Emissive
  *		WorldNormal
+ *		DeviceDepth
+ *
  * There are various efficiences possible by doing these captures as a group, rather
  * than doing each one individually.
  * 
@@ -39,15 +42,25 @@ namespace Geometry
 class MODELINGCOMPONENTS_API FSceneCapturePhotoSet
 {
 public:
+
+	FSceneCapturePhotoSet();
+
 	/**
 	 * Set the target World and set of Actors
 	 */
 	void SetCaptureSceneActors(UWorld* World, const TArray<AActor*>& Actors);
 
 	/**
-	 * Enable/Disable a particular capture type. Currently all are enabled by default.
+	 * Enable/Disable a particular capture type
 	 */
 	void SetCaptureTypeEnabled(ERenderCaptureType CaptureType, bool bEnabled);
+	bool GetCaptureTypeEnabled(ERenderCaptureType CaptureType) const;
+
+	/**
+	 * Configure the given capture type
+	 */
+	void SetCaptureConfig(ERenderCaptureType CaptureType, const FRenderCaptureConfig& Config);
+	FRenderCaptureConfig GetCaptureConfig(ERenderCaptureType CaptureType) const;
 
 	/**
 	 * Add captures at the corners and face centers of the "view box",
@@ -93,6 +106,7 @@ public:
 		float Metallic;
 		FVector3f Emissive;
 		FVector3f WorldNormal;
+		float DeviceDepth;
 
 		FSceneSample();
 
@@ -115,15 +129,28 @@ public:
 		TFunctionRef<bool(const FVector3d&, const FVector3d&)> VisibilityFunction,
 		FSceneSample& DefaultsInResultsOut) const;
 
+	/**
+	 * If ValidSampleDepthThreshold >  0 then the VisibilityFunction and the DeviceDepthPhotoSet will be used to determine
+	 * sample validity. Using DeviceDepthPhotoSet will mitigate artefacts caused by samples which cannot be determined
+	 * invalid by querying only the VisibilityFunction.
+	 * If ValidSampleDepthThreshold <= 0 only VisibilityFunction will be used to determine sample validity
+	 */
 	bool ComputeSampleLocation(
 		const FVector3d& Position,
 		const FVector3d& Normal,
+		float ValidSampleDepthThreshold,
 		TFunctionRef<bool(const FVector3d&, const FVector3d&)> VisibilityFunction,
 		int& PhotoIndex,
 		FVector2d& PhotoCoords) const;
 
+	/**
+	 * @returns the value of the nearest pixel in the given photo for the given CaptureType
+	 * Note: This function does not interpolate because this causes artefacts; two adjacent pixels on a photo can be
+	 * from parts of the scene which are very far from each other so interpolating the values when makes no sense.
+	 * Texture filtering is accomplished by the baking framework where it can use correctly localized information.
+	 */
 	template <ERenderCaptureType CaptureType>
-	FVector4f ComputeSample(
+	FVector4f ComputeSampleNearest(
 		const int& PhotoIndex,
 		const FVector2d& PhotoCoords,
 		const FSceneSample& DefaultSample) const;
@@ -135,6 +162,7 @@ public:
 	const FSpatialPhotoSet3f& GetPackedMRSPhotoSet() { return PackedMRSPhotoSet; }
 	const FSpatialPhotoSet3f& GetWorldNormalPhotoSet() { return WorldNormalPhotoSet; }
 	const FSpatialPhotoSet3f& GetEmissivePhotoSet() { return EmissivePhotoSet; }
+	const FSpatialPhotoSet1f& GetDeviceDepthPhotoSet() { return DeviceDepthPhotoSet; }
 
 	/**
 	 * Enable debug image writing. All captured images will be written to <Project>/Intermediate/<FolderName>.
@@ -174,6 +202,7 @@ protected:
 	bool bEnablePackedMRS = true;
 	bool bEnableWorldNormal = true;
 	bool bEnableEmissive = true;
+	bool bEnableDeviceDepth = true;
 
 	FSpatialPhotoSet3f BaseColorPhotoSet;
 	FSpatialPhotoSet1f RoughnessPhotoSet;
@@ -182,6 +211,16 @@ protected:
 	FSpatialPhotoSet3f PackedMRSPhotoSet;
 	FSpatialPhotoSet3f WorldNormalPhotoSet;
 	FSpatialPhotoSet3f EmissivePhotoSet;
+	FSpatialPhotoSet1f DeviceDepthPhotoSet;
+
+	FRenderCaptureConfig BaseColorConfig;
+	FRenderCaptureConfig RoughnessConfig;
+	FRenderCaptureConfig SpecularConfig;
+	FRenderCaptureConfig MetallicConfig;
+	FRenderCaptureConfig PackedMRSConfig;
+	FRenderCaptureConfig WorldNormalConfig;
+	FRenderCaptureConfig EmissiveConfig;
+	FRenderCaptureConfig DeviceDepthConfig;
 
 	TArray<FSpatialPhotoParams> PhotoSetParams;
 
@@ -194,54 +233,60 @@ protected:
 
 
 template <ERenderCaptureType CaptureType>
-FVector4f FSceneCapturePhotoSet::ComputeSample(
+FVector4f FSceneCapturePhotoSet::ComputeSampleNearest(
 	const int& PhotoIndex,
 	const FVector2d& PhotoCoords,
 	const FSceneSample& DefaultSample) const
 {
 	if constexpr (CaptureType == ERenderCaptureType::BaseColor)
 	{
-		FVector3f BaseColor = BaseColorPhotoSet.ComputeSample(PhotoIndex, PhotoCoords, DefaultSample.BaseColor);
+		FVector3f BaseColor = BaseColorPhotoSet.ComputeSampleNearest(PhotoIndex, PhotoCoords, DefaultSample.BaseColor);
 		return FVector4f(BaseColor, 1.f);
 	}
 
 	if constexpr (CaptureType == ERenderCaptureType::Roughness)
 	{
-		float Roughness = RoughnessPhotoSet.ComputeSample(PhotoIndex, PhotoCoords, DefaultSample.Roughness);
+		float Roughness = RoughnessPhotoSet.ComputeSampleNearest(PhotoIndex, PhotoCoords, DefaultSample.Roughness);
 		return FVector4f(Roughness, Roughness, Roughness, 1.f);
 	}
 
 	if constexpr (CaptureType == ERenderCaptureType::Specular)
 	{
-		float Specular = SpecularPhotoSet.ComputeSample(PhotoIndex, PhotoCoords, DefaultSample.Specular);
+		float Specular = SpecularPhotoSet.ComputeSampleNearest(PhotoIndex, PhotoCoords, DefaultSample.Specular);
 		return FVector4f(Specular, Specular, Specular, 1.f);
 	}
 
 	if constexpr (CaptureType == ERenderCaptureType::Metallic)
 	{
-		float Metallic = MetallicPhotoSet.ComputeSample(PhotoIndex, PhotoCoords, DefaultSample.Metallic);
+		float Metallic = MetallicPhotoSet.ComputeSampleNearest(PhotoIndex, PhotoCoords, DefaultSample.Metallic);
 		return FVector4f(Metallic, Metallic, Metallic, 1.f);
 	}
 
 	if constexpr (CaptureType == ERenderCaptureType::CombinedMRS)
 	{
 		FVector3f MRSValue(DefaultSample.Metallic, DefaultSample.Roughness, DefaultSample.Specular);
-		MRSValue = PackedMRSPhotoSet.ComputeSample(PhotoIndex, PhotoCoords, MRSValue);
+		MRSValue = PackedMRSPhotoSet.ComputeSampleNearest(PhotoIndex, PhotoCoords, MRSValue);
 		return FVector4f(MRSValue, 1.f);
 	}
 
 	if constexpr (CaptureType == ERenderCaptureType::Emissive)
 	{
-		FVector3f Emissive = EmissivePhotoSet.ComputeSample(PhotoIndex, PhotoCoords, DefaultSample.Emissive);
+		FVector3f Emissive = EmissivePhotoSet.ComputeSampleNearest(PhotoIndex, PhotoCoords, DefaultSample.Emissive);
 		return FVector4f(Emissive, 1.f);
 	}
 
 	if constexpr (CaptureType == ERenderCaptureType::WorldNormal)
 	{
-		FVector3f WorldNormal = WorldNormalPhotoSet.ComputeSample(PhotoIndex, PhotoCoords, DefaultSample.WorldNormal);
+		FVector3f WorldNormal = WorldNormalPhotoSet.ComputeSampleNearest(PhotoIndex, PhotoCoords, DefaultSample.WorldNormal);
 		return FVector4f(WorldNormal, 1.f);
 	}
 
+	if constexpr (CaptureType == ERenderCaptureType::DeviceDepth)
+	{
+		float Depth = DeviceDepthPhotoSet.ComputeSampleNearest(PhotoIndex, PhotoCoords, DefaultSample.DeviceDepth);
+		return FVector4f(Depth, Depth, Depth, 1.f);
+	}
+
 	ensure(false);
 	return FVector4f::Zero();
 }
diff --git a/Engine/Plugins/Runtime/MeshModelingToolset/Source/ModelingComponents/Public/Scene/WorldRenderCapture.h b/Engine/Plugins/Runtime/MeshModelingToolset/Source/ModelingComponents/Public/Scene/WorldRenderCapture.h
--- a/Engine/Plugins/Runtime/MeshModelingToolset/Source/ModelingComponents/Public/Scene/WorldRenderCapture.h
+++ b/Engine/Plugins/Runtime/MeshModelingToolset/Source/ModelingComponents/Public/Scene/WorldRenderCapture.h
@@ -24,9 +24,19 @@ enum class ERenderCaptureType
 	Specular = 8,
 	Emissive = 16,
 	WorldNormal = 32,
+	DeviceDepth = 64,
 	CombinedMRS = 128
 };
 
+struct MODELINGCOMPONENTS_API FRenderCaptureConfig
+{
+	// You might want to disable this if you're using FMeshMapBaker because it supports multi-sampling in a way which
+	// will avoid blending pixels at different depths. This option is ignored for ERenderCaptureType::DeviceDepth
+	bool bAntiAliasing = true;
+};
+
+FRenderCaptureConfig MODELINGCOMPONENTS_API GetDefaultRenderCaptureConfig(ERenderCaptureType CaptureType);
+
 /**
  * FRenderCaptureTypeFlags is a set of per-capture-type booleans
  */
@@ -39,6 +49,7 @@ struct MODELINGCOMPONENTS_API FRenderCaptureTypeFlags
 	bool bEmissive = false;
 	bool bWorldNormal = false;
 	bool bCombinedMRS = false;
+	bool bDeviceDepth = false;
 
 	/** @return FRenderCaptureTypeFlags with all types enabled/true */
 	static FRenderCaptureTypeFlags All();
@@ -60,6 +71,36 @@ struct MODELINGCOMPONENTS_API FRenderCaptureTypeFlags
 };
 
 
+/** 
+ * Render capture images use the render target coordinate system, which is defined such that:
+ * - The upper-left corner is (0,0) in pixel space, (0,0) in UV space and (-1,1) in NDC space
+ * - The bottom-right corner is (Width,Height) in pixel space, (1,1) in UV Space and (1,-1) in NDC space
+ * Pixel centers are offset by (0.5,0.5) from integer locations
+ */
+struct FRenderCaptureCoordinateConverter2D
+{
+	// Convert normalized device coordinates to render capture image UV coordinates
+	static FVector2d DeviceToUV(FVector2d NDC)
+	{
+		FVector2d UV = NDC / 2. + FVector2d(.5, .5);
+		UV.Y = 1. - UV.Y;
+		return UV;
+	}
+
+	// Convert render capture image pixel coordinates to normalized device coordinates
+	static FVector2d PixelToDevice(FVector2i Pixel, int32 Width, int32 Height)
+	{
+		FVector2d NDC;
+		NDC.X = ((double)Pixel.X + .5) / Width;
+		NDC.Y = ((double)Pixel.Y + .5) / Height;
+		NDC.X = 2. * NDC.X - 1.;
+		NDC.Y = 2. * NDC.Y - 1.;
+		NDC.Y *= -1;
+		return NDC;
+	}
+};
+
+
 /**
  * FWorldRenderCapture captures a rendering of a set of Actors in a World from a
  * specific viewpoint. Various types of rendering are supported, as defined by ERenderCaptureType.
@@ -97,10 +138,17 @@ public:
 
 	/** @return pixel dimensions that target image will be rendered at */
 	const FImageDimensions& GetDimensions() const { return Dimensions; }
+	
+	/**
+	 * @return view matrices used in the last call to CaptureFromPosition
+	 * Useful to get the needed information to unproject the DeviceDepth render capture
+	 */
+	const FViewMatrices& GetLastCaptureViewMatrices() const { return LastCaptureViewMatrices; }
 
 	/**
-	 * Capture the desired buffer type CaptureType with the given view/camera parameters.
-	 * @param ResultImageOut output iamge of size GetDimensions() is stored here.
+	 * Capture the desired buffer type CaptureType with the given view/camera parameters. The returned image
+	 * data is interpreted according to the comment in FRenderCaptureCoordinateConverter2D
+	 * @param ResultImageOut output image of size GetDimensions() is stored here.
 	 * @return true if capture could be rendered successfully
 	 */
 	bool CaptureFromPosition(
@@ -108,7 +156,8 @@ public:
 		const FFrame3d& ViewFrame,
 		double HorzFOVDegrees,
 		double NearPlaneDist,
-		FImageAdapter& ResultImageOut);
+		FImageAdapter& ResultImageOut,
+		const FRenderCaptureConfig& Config = {});
 
 	/**
 	 * Enable debug image write. The captured image will be written to <Project>/Intermediate/<FolderName>/<CaptureType>_<ImageCounter>.bmp
@@ -129,8 +178,13 @@ protected:
 	// Temporary textures used as render targets. We explicitly prevent this from being GC'd internally
 	UTextureRenderTarget2D* LinearRenderTexture = nullptr;
 	UTextureRenderTarget2D* GammaRenderTexture = nullptr;
+	UTextureRenderTarget2D* DepthRenderTexture = nullptr;
+
 	FImageDimensions RenderTextureDimensions;
 	UTextureRenderTarget2D* GetRenderTexture(bool bLinear);
+	UTextureRenderTarget2D* GetDepthRenderTexture();
+
+	FViewMatrices LastCaptureViewMatrices;
 
 	// temporary buffer used to read from texture
 	TArray<FLinearColor> ReadImageBuffer;
@@ -140,14 +194,24 @@ protected:
 		const FFrame3d& Frame,
 		double HorzFOVDegrees,
 		double NearPlaneDist,
-		FImageAdapter& ResultImageOut);
+		FImageAdapter& ResultImageOut,
+		const FRenderCaptureConfig& Config = {});
 
 	/** Combined Metallic/Roughness/Specular uses a custom postprocess material */
 	bool CaptureMRSFromPosition(
 		const FFrame3d& Frame,
 		double HorzFOVDegrees,
 		double NearPlaneDist,
-		FImageAdapter& ResultImageOut);
+		FImageAdapter& ResultImageOut,
+		const FRenderCaptureConfig& Config = {});
+
+	/** Depth is a special case and uses different code than capture of color/property channels */
+	bool CaptureDeviceDepthFromPosition(
+		const FFrame3d& Frame,
+		double HorzFOVDegrees,
+		double NearPlaneDist,
+		FImageAdapter& ResultImageOut,
+		const FRenderCaptureConfig& Config = {});
 
 	bool bWriteDebugImage = false;
 	int32 DebugImageCounter = -1;
diff --git a/Engine/Source/Runtime/GeometryCore/Public/Image/SpatialPhotoSet.h b/Engine/Source/Runtime/GeometryCore/Public/Image/SpatialPhotoSet.h
--- a/Engine/Source/Runtime/GeometryCore/Public/Image/SpatialPhotoSet.h
+++ b/Engine/Source/Runtime/GeometryCore/Public/Image/SpatialPhotoSet.h
@@ -4,6 +4,7 @@
 
 #include "VectorTypes.h"
 #include "FrameTypes.h"
+#include "SceneView.h"
 #include "Image/ImageBuilder.h"
 #include "Image/ImageDimensions.h"
 
@@ -26,6 +27,9 @@ struct FSpatialPhotoParams
 
 	/** Pixel dimensions of the photo image */
 	FImageDimensions Dimensions;
+
+	/** Useful to to unproject the DeviceDepth render capture */
+	FViewMatrices ViewMatrices;
 };
 
 /**
@@ -101,6 +105,12 @@ public:
 		const FVector2d& PhotoCoords,
 		const PixelType& DefaultValue
 	) const;
+	
+	PixelType ComputeSampleNearest(
+		const int& PhotoIndex,
+		const FVector2d& PhotoCoords,
+		const PixelType& DefaultValue
+	) const;
 
 	bool ComputeSampleLocation(
 		const FVector3d& Position, 
@@ -193,6 +203,8 @@ PixelType TSpatialPhotoSet<PixelType, RealType>::ComputeSample(
 	FVector2d PhotoCoords;
 	if (ComputeSampleLocation(Position, Normal, VisibilityFunction, PhotoIndex, PhotoCoords))
 	{
+		// TODO This bilinear sampling causes artefacts when it blends pixels from different depths (hence unrelated
+		// values), we could fix this with some kind of rejection strategy. Search :BilinearSamplingDepthArtefacts 
 		const TSpatialPhoto<PixelType>& Photo = *Photos[PhotoIndex];
 		Result = Photo.Image.template BilinearSample<RealType>(PhotoCoords, DefaultValue);
 	}
@@ -206,11 +218,32 @@ PixelType TSpatialPhotoSet<PixelType, RealType>::ComputeSample(
 	const FVector2d& PhotoCoords, 
 	const PixelType& DefaultValue) const
 {
+	// TODO See the task tagged :BilinearSamplingDepthArtefacts
 	const TSpatialPhoto<PixelType>& Photo = *Photos[PhotoIndex];
 	return Photo.Image.template BilinearSample<RealType>(PhotoCoords, DefaultValue);
 }
 
 
+template<typename PixelType, typename RealType>
+PixelType TSpatialPhotoSet<PixelType, RealType>::ComputeSampleNearest(
+	const int& PhotoIndex, 
+	const FVector2d& PhotoCoords, 
+	const PixelType& DefaultValue) const
+{
+	const TSpatialPhoto<PixelType>& Photo = *Photos[PhotoIndex];
+
+	FVector2d UVCoords;
+	UVCoords.X = PhotoCoords.X / Photo.Image.GetDimensions().GetWidth();
+	UVCoords.Y = PhotoCoords.Y / Photo.Image.GetDimensions().GetHeight();
+	if (UVCoords.X < 0. || UVCoords.X > 1. || UVCoords.Y < 0. || UVCoords.Y > 1.)
+	{
+		return DefaultValue;
+	}
+
+	return Photo.Image.NearestSampleUV(UVCoords);
+}
+
+
 
 } // end namespace UE::Geometry
 } // end namespace UE
\ No newline at end of file
