MeshModelingTools: Fix crash in CubeGrid if gizmo is dragged without ever having selected anything.

diff --git a/Engine/Plugins/Experimental/MeshModelingToolsetExp/Source/MeshModelingToolsExp/Private/CubeGridTool.cpp b/Engine/Plugins/Experimental/MeshModelingToolsetExp/Source/MeshModelingToolsExp/Private/CubeGridTool.cpp
--- a/Engine/Plugins/Experimental/MeshModelingToolsetExp/Source/MeshModelingToolsExp/Private/CubeGridTool.cpp
+++ b/Engine/Plugins/Experimental/MeshModelingToolsetExp/Source/MeshModelingToolsExp/Private/CubeGridTool.cpp
@@ -2205,11 +2205,14 @@ void UCubeGridTool::UpdateCornerGeometrySet()
 {
 	using namespace CubeGridToolLocals;
 
-	FOrientedBox3d FrameSpaceBox = ConvertToOrientedBox(Selection.Box, Selection.Direction);
 	CornersGeometrySet.Reset();
-	for (int i = 0; i < 4; ++i)
+	if (bHaveSelection)
 	{
-		CornersGeometrySet.AddPoint(i, CubeGrid->GetFrame().FromFramePoint(FrameSpaceBox.GetCorner(i)));
+		FOrientedBox3d FrameSpaceBox = ConvertToOrientedBox(Selection.Box, Selection.Direction);
+		for (int i = 0; i < 4; ++i)
+		{
+			CornersGeometrySet.AddPoint(i, CubeGrid->GetFrame().FromFramePoint(FrameSpaceBox.GetCorner(i)));
+		}
 	}
 }
 
diff --git a/Engine/Plugins/Experimental/MeshModelingToolsetExp/Source/MeshModelingToolsExp/Public/CubeGridTool.h b/Engine/Plugins/Experimental/MeshModelingToolsetExp/Source/MeshModelingToolsExp/Public/CubeGridTool.h
--- a/Engine/Plugins/Experimental/MeshModelingToolsetExp/Source/MeshModelingToolsExp/Public/CubeGridTool.h
+++ b/Engine/Plugins/Experimental/MeshModelingToolsetExp/Source/MeshModelingToolsExp/Public/CubeGridTool.h
@@ -297,7 +297,9 @@ public:
 		// Both of these boxes are in the coordinate space of the (unscaled) grid frame.
 		UE::Geometry::FAxisAlignedBox3d Box;
 		UE::Geometry::FAxisAlignedBox3d StartBox; // Box delineating original selected face
-		FCubeGrid::EFaceDirection Direction;
+		
+		// Direction must be initialized to a valid enum value (0 is not one).
+		FCubeGrid::EFaceDirection Direction = FCubeGrid::EFaceDirection::PositiveX;
 
 		bool operator==(const FSelection& Other)
 		{
