ModelingTools: use post-baked-transform X and Y axes to detect symmetry in VSculpt Tool

diff --git a/Engine/Plugins/Experimental/MeshModelingToolsetExp/Source/MeshModelingToolsExp/Private/MeshVertexSculptTool.cpp b/Engine/Plugins/Experimental/MeshModelingToolsetExp/Source/MeshModelingToolsExp/Private/MeshVertexSculptTool.cpp
--- a/Engine/Plugins/Experimental/MeshModelingToolsetExp/Source/MeshModelingToolsExp/Private/MeshVertexSculptTool.cpp
+++ b/Engine/Plugins/Experimental/MeshModelingToolsetExp/Source/MeshModelingToolsExp/Private/MeshVertexSculptTool.cpp
@@ -1321,11 +1321,17 @@ void UMeshVertexSculptTool::TryToInitializeSymmetry()
 
 	TFuture<void> ComputeSymmetryX = Async(VertexSculptToolAsyncExecTarget, [this, &SculptMeshAABBTree, &SymmetryX, &bSymmetryXValid]()
 	{
-		bSymmetryXValid = SymmetryX->Initialize( GetSculptMesh(), &SculptMeshAABBTree, FFrame3d(FVector3d::Zero(), FVector3d::UnitX()) );
+		// UMeshSculptToolBase::InitializeSculptMeshComponent() bakes the local-to-world rotation and scaling into the sculpt mesh so
+		// that we don't have to try to apply sculpting "through" these transforms (v hard). But when detecting symmetry we want to
+		// do it relative to local coordinates. InitialTargetTransform is the transform w/o translation, so use it's X and Y directions.
+		// (it may be better to do symmetry detection on the pre-baked-transform mesh, but then we need a copy...)
+		FVector3d BakedX = this->InitialTargetTransform.GetRotation().AxisX();
+		bSymmetryXValid = SymmetryX->Initialize( GetSculptMesh(), &SculptMeshAABBTree, FFrame3d(FVector3d::Zero(), BakedX) );
 	});
 	TFuture<void> ComputeSymmetryY = Async(VertexSculptToolAsyncExecTarget, [this, &SculptMeshAABBTree, &SymmetryY, &bSymmetryYValid]()
 	{
-		bSymmetryYValid = SymmetryY->Initialize( GetSculptMesh(), &SculptMeshAABBTree, FFrame3d(FVector3d::Zero(), FVector3d::UnitY()) );
+		FVector3d BakedY = this->InitialTargetTransform.GetRotation().AxisY();
+		bSymmetryYValid = SymmetryY->Initialize( GetSculptMesh(), &SculptMeshAABBTree, FFrame3d(FVector3d::Zero(), BakedY) );
 	});
 	ComputeSymmetryX.Wait();
 	ComputeSymmetryY.Wait();
