BakeVertex: Fixed identity bakes w/ VertexColor output type returning black.

diff --git a/Engine/Plugins/Experimental/GeometryScripting/Source/GeometryScriptingCore/Private/MeshBakeFunctions.cpp b/Engine/Plugins/Experimental/GeometryScripting/Source/GeometryScriptingCore/Private/MeshBakeFunctions.cpp
--- a/Engine/Plugins/Experimental/GeometryScripting/Source/GeometryScriptingCore/Private/MeshBakeFunctions.cpp
+++ b/Engine/Plugins/Experimental/GeometryScripting/Source/GeometryScriptingCore/Private/MeshBakeFunctions.cpp
@@ -621,6 +621,24 @@ namespace GeometryScriptBakeLocals
 
 		const bool bIsBakeToSelf = (TargetMesh == SourceMesh);
 
+		// Initialize the source mesh
+		// This must precede target mesh generation in case SourceMesh & TargetMesh are the same mesh
+		// since we mutate the target mesh color topology.
+		FDynamicMesh3 SourceMeshCopy = *SourceMesh->GetMeshPtr();
+		const FDynamicMesh3* SourceMeshToUse = &SourceMeshCopy;
+		if (BakeOptions.bProjectionInWorldSpace && !bIsBakeToSelf)
+		{
+			// Transform the SourceMesh into TargetMesh local space using a copy (oof)
+			// TODO: Remove this once we have support for transforming rays in the core bake loop
+			const FTransformSRT3d SourceToWorld = SourceTransform;
+			MeshTransforms::ApplyTransform(SourceMeshCopy, SourceToWorld, true);
+			const FTransformSRT3d TargetToWorld = TargetTransform;
+			MeshTransforms::ApplyTransformInverse(SourceMeshCopy, TargetToWorld, true);
+		}
+
+		const FDynamicMeshAABBTree3 DetailSpatial(SourceMeshToUse);
+		FMeshBakerDynamicMeshSampler DetailSampler(SourceMeshToUse, &DetailSpatial);
+
 		// Initialize the color overlay on the TargetMesh
 		FDynamicMesh3& TargetMeshRef = TargetMesh->GetMeshRef();
 		TargetMeshRef.EnableAttributes();
@@ -648,27 +666,27 @@ namespace GeometryScriptBakeLocals
 				}
 				return bCanShare;
 			}, 0.0f);
-		
 
-		// Initialize the source mesh
-		FDynamicMesh3 SourceMeshCopy;
-		const FDynamicMesh3* SourceMeshOriginal = SourceMesh->GetMeshPtr();
-		const FDynamicMesh3* SourceMeshToUse = SourceMeshOriginal;
-		if (BakeOptions.bProjectionInWorldSpace && !bIsBakeToSelf)
+		if (bIsBakeToSelf)
 		{
-			// Transform the SourceMesh into TargetMesh local space using a copy (oof)
-			// TODO: Remove this once we have support for transforming rays in the core bake loop
-			SourceMeshCopy = *SourceMeshOriginal;
-			const FTransformSRT3d SourceToWorld = SourceTransform;
-			MeshTransforms::ApplyTransform(SourceMeshCopy, SourceToWorld, true);
-			const FTransformSRT3d TargetToWorld = TargetTransform;
-			MeshTransforms::ApplyTransformInverse(SourceMeshCopy, TargetToWorld, true);
-			SourceMeshToUse = &SourceMeshCopy;
+			// Copy source vertex colors onto new color overlay topology for identity bakes.
+			// This is necessary when sampling vertex color data.
+			const FDynamicMeshColorOverlay* SourceColorOverlay = SourceMeshToUse->HasAttributes() ? SourceMeshToUse->Attributes()->PrimaryColors() : nullptr;
+			FDynamicMeshColorOverlay* TargetColorOverlay = TargetMeshRef.Attributes()->PrimaryColors(); 
+			if (SourceColorOverlay)
+			{
+				for (int VId : TargetMeshRef.VertexIndicesItr())
+				{
+					TargetMeshRef.EnumerateVertexTriangles(VId, [VId, SourceColorOverlay, TargetColorOverlay](int32 TriID)
+					{
+						const FVector4f TargetColor = SourceColorOverlay->GetElementAtVertex(TriID, VId);
+						const int ElemId = TargetColorOverlay->GetElementIDAtVertex(TriID, VId);
+						TargetColorOverlay->SetElement(ElemId, TargetColor);
+					});
+				}
+			}
 		}
 
-		const FDynamicMeshAABBTree3 DetailSpatial(SourceMeshToUse);
-		FMeshBakerDynamicMeshSampler DetailSampler(SourceMeshToUse, &DetailSpatial);
-
 		TUniquePtr<FMeshVertexBaker> Result = MakeUnique<FMeshVertexBaker>();
 		FMeshVertexBaker& Baker = *Result;
 		Baker.BakeMode = BakeTypes.OutputMode == EGeometryScriptBakeOutputMode::RGBA ? FMeshVertexBaker::EBakeMode::RGBA : FMeshVertexBaker::EBakeMode::PerChannel;
diff --git a/Engine/Plugins/Experimental/MeshModelingToolsetExp/Source/MeshModelingToolsExp/Private/BakeMeshAttributeVertexTool.cpp b/Engine/Plugins/Experimental/MeshModelingToolsetExp/Source/MeshModelingToolsExp/Private/BakeMeshAttributeVertexTool.cpp
--- a/Engine/Plugins/Experimental/MeshModelingToolsetExp/Source/MeshModelingToolsExp/Private/BakeMeshAttributeVertexTool.cpp
+++ b/Engine/Plugins/Experimental/MeshModelingToolsetExp/Source/MeshModelingToolsExp/Private/BakeMeshAttributeVertexTool.cpp
@@ -56,7 +56,7 @@ public:
 	// General bake settings
 	TSharedPtr<FDynamicMesh3, ESPMode::ThreadSafe> DetailMesh;
 	TSharedPtr<FDynamicMeshAABBTree3, ESPMode::ThreadSafe> DetailSpatial;
-	FDynamicMesh3* BaseMesh;
+	const FDynamicMesh3* BaseMesh;
 	TSharedPtr<TMeshTangents<double>, ESPMode::ThreadSafe> BaseMeshTangents;
 	TUniquePtr<FMeshVertexBaker> Baker;
 	bool bIsBakeToSelf = false;
@@ -260,25 +260,26 @@ void UBakeMeshAttributeVertexTool::Setup()
 
 	UE::ToolTarget::HideSourceObject(Targets[0]);
 
-	const FDynamicMesh3 InputMeshWithTangents = UE::ToolTarget::GetDynamicMeshCopy(Targets[0], true);
+	// TargetMesh stores the original target mesh. It is intended to remain
+	// const throughout this tool and is used to refresh the PreviewMesh back
+	// to its original state.
+	TargetMesh = UE::ToolTarget::GetDynamicMeshCopy(Targets[0], true);
+	TargetSpatial.SetMesh(&TargetMesh, true);
+	TargetMeshTangents = MakeShared<FMeshTangentsd, ESPMode::ThreadSafe>(&TargetMesh);
+	TargetMeshTangents->CopyTriVertexTangents(TargetMesh);
+
+	// PreviewMesh stores computed result mesh. On shutdown, PreviewMesh will be
+	// used to commit the dynamic mesh to the target tool target.
 	PreviewMesh = NewObject<UPreviewMesh>(this);
 	PreviewMesh->CreateInWorld(GetTargetWorld(), FTransform::Identity);
 	ToolSetupUtil::ApplyRenderingConfigurationToPreview(PreviewMesh, nullptr);
 	PreviewMesh->SetTransform(static_cast<FTransform>(UE::ToolTarget::GetLocalToWorldTransform(Targets[0])));
 	PreviewMesh->SetTangentsMode(EDynamicMeshComponentTangentsMode::ExternallyProvided);
-	PreviewMesh->ReplaceMesh(InputMeshWithTangents);
+	PreviewMesh->ReplaceMesh(TargetMesh);
 	PreviewMesh->SetMaterials(UE::ToolTarget::GetMaterialSet(Targets[0]).Materials);
 	PreviewMesh->SetOverrideRenderMaterial(PreviewMaterial);
 	PreviewMesh->SetVisible(true);
 
-	PreviewMesh->ProcessMesh([this](const FDynamicMesh3& Mesh)
-	{
-		TargetMesh.Copy(Mesh);
-		TargetSpatial.SetMesh(&TargetMesh, true);
-		TargetMeshTangents = MakeShared<FMeshTangentsd, ESPMode::ThreadSafe>(&TargetMesh);
-		TargetMeshTangents->CopyTriVertexTangents(Mesh);
-	});
-
 	UToolTarget* Target = Targets[0];
 	UToolTarget* DetailTarget = Targets[bIsBakeToSelf ? 0 : 1];
 
@@ -454,7 +455,11 @@ TUniquePtr<UE::Geometry::TGenericDataOperator<FMeshVertexBaker>> UBakeMeshAttrib
 	TUniquePtr<FMeshVertexBakerOp> Op = MakeUnique<FMeshVertexBakerOp>();
 	Op->DetailMesh = DetailMesh;
 	Op->DetailSpatial = DetailSpatial;
-	Op->BaseMesh = &TargetMesh;
+
+	// Pass the PreviewMesh here instead of the TargetMesh. The PreviewMesh
+	// contains the updated color topology. TargetMesh holds onto the original
+	// color topology and values.
+	Op->BaseMesh = PreviewMesh->GetMesh();
 	Op->BaseMeshTangents = TargetMeshTangents;
 	Op->BakeSettings = CachedBakeSettings;
 	Op->OcclusionSettings = CachedOcclusionMapSettings;
@@ -475,9 +480,9 @@ void UBakeMeshAttributeVertexTool::UpdateDetailMesh()
 	IPrimitiveComponentBackedTarget* DetailComponent = Cast<IPrimitiveComponentBackedTarget>(Targets[bIsBakeToSelf ? 0 : 1]);
 	UToolTarget* DetailTargetMesh = Targets[bIsBakeToSelf ? 0 : 1];
 
+	const FDynamicMesh3 DetailMeshCopy = UE::ToolTarget::GetDynamicMeshCopy(DetailTargetMesh, true);
 	DetailMesh = MakeShared<FDynamicMesh3, ESPMode::ThreadSafe>();
-	FMeshDescriptionToDynamicMesh Converter;
-	Converter.Convert( UE::ToolTarget::GetMeshDescription(DetailTargetMesh), *DetailMesh);
+	DetailMesh->Copy(DetailMeshCopy);
 	if (InputMeshSettings->bProjectionInWorldSpace && bIsBakeToSelf == false)
 	{
 		const FTransformSRT3d DetailToWorld(DetailComponent->GetWorldTransform());
@@ -588,15 +593,15 @@ void UBakeMeshAttributeVertexTool::UpdateColorTopology()
 	PreviewMesh->EditMesh([this](FDynamicMesh3& Mesh)
 	{
 		Mesh.EnableAttributes();
-		Mesh.Attributes()->DisablePrimaryColors();
 		Mesh.Attributes()->EnablePrimaryColors();
+		Mesh.Attributes()->PrimaryColors()->ClearElements();
 
 		FDynamicMeshNormalOverlay* NormalOverlay = Mesh.Attributes()->PrimaryNormals();
 		FDynamicMeshUVOverlay* UVOverlay = Mesh.Attributes()->PrimaryUV();
 		Mesh.Attributes()->PrimaryColors()->CreateFromPredicate(
-			[&](int ParentVID, int TriIDA, int TriIDB) -> bool
+			[&](int /*ParentVID*/, int TriIDA, int TriIDB) -> bool
 			{
-				auto OverlayCanShare = [&] (auto Overlay) -> bool
+				auto OverlayCanShare = [TriIDA, TriIDB] (auto Overlay) -> bool
 				{
 					return Overlay ? Overlay->AreTrianglesConnected(TriIDA, TriIDB) : true;
 				};
@@ -612,17 +617,24 @@ void UBakeMeshAttributeVertexTool::UpdateColorTopology()
 				}
 				return bCanShare;
 			}, 0.0f);
-	});
 
-	// Update BaseMesh color topology.
-	TargetMesh.EnableAttributes();
-	TargetMesh.Attributes()->DisablePrimaryColors();
-	TargetMesh.Attributes()->EnablePrimaryColors();
-	PreviewMesh->ProcessMesh([this](const FDynamicMesh3& Mesh)
-	{
-		TargetMesh.Attributes()->PrimaryColors()->Copy(*Mesh.Attributes()->PrimaryColors());
+		// Copy source vertex colors onto new color overlay topology.
+		const FDynamicMeshColorOverlay* TargetColorOverlay = TargetMesh.HasAttributes() ? TargetMesh.Attributes()->PrimaryColors() : nullptr;
+		FDynamicMeshColorOverlay* PreviewColorOverlay = Mesh.Attributes()->PrimaryColors(); 
+		if (TargetColorOverlay)
+		{
+			for (int VId : Mesh.VertexIndicesItr())
+			{
+				Mesh.EnumerateVertexTriangles(VId, [VId, TargetColorOverlay, PreviewColorOverlay](int32 TriID)
+				{
+					const FVector4f TargetColor = TargetColorOverlay->GetElementAtVertex(TriID, VId);
+					const int ElemId = PreviewColorOverlay->GetElementIDAtVertex(TriID, VId);
+					PreviewColorOverlay->SetElement(ElemId, TargetColor);
+				});
+			}
+		}
 	});
-	NumColorElements = TargetMesh.Attributes()->PrimaryColors()->ElementCount();
+	NumColorElements = PreviewMesh->GetMesh()->Attributes()->PrimaryColors()->ElementCount();
 
 	bColorTopologyValid = true;
 }
@@ -688,8 +700,13 @@ void UBakeMeshAttributeVertexTool::UpdateResult()
 		case EBakeMapType::FaceNormal:
 		case EBakeMapType::Position:
 		case EBakeMapType::MaterialID:
+			OpState |= UpdateResult_MeshProperty(Dimensions);
+			break;
 		case EBakeMapType::VertexColor:
 			OpState |= UpdateResult_MeshProperty(Dimensions);
+			// Force copy the original vertex colors to our PreviewMesh so that
+			// the baker samples the source vertex colors for identity bakes.
+			UpdateColorTopology();
 			break;
 		case EBakeMapType::Texture:
 			OpState |= UpdateResult_Texture2DImage(Dimensions, DetailMesh.Get());
@@ -734,7 +751,6 @@ void UBakeMeshAttributeVertexTool::OnResultUpdated(const TUniquePtr<FMeshVertexB
 		return;
 	}
 
-	// TODO: Review how to handle the implicit sRGB conversion in the StaticMesh build.
 	PreviewMesh->DeferredEditMesh([this, &ImageResult](FDynamicMesh3& Mesh)
 	{
 		const int NumColors = Mesh.Attributes()->PrimaryColors()->ElementCount();
diff --git a/Engine/Plugins/Runtime/GeometryProcessing/Source/DynamicMesh/Private/Sampling/MeshVertexBaker.cpp b/Engine/Plugins/Runtime/GeometryProcessing/Source/DynamicMesh/Private/Sampling/MeshVertexBaker.cpp
--- a/Engine/Plugins/Runtime/GeometryProcessing/Source/DynamicMesh/Private/Sampling/MeshVertexBaker.cpp
+++ b/Engine/Plugins/Runtime/GeometryProcessing/Source/DynamicMesh/Private/Sampling/MeshVertexBaker.cpp
@@ -97,21 +97,24 @@ void FMeshVertexBaker::BakeImpl(void* Data)
 	FMeshVertexBaker* Baker = static_cast<FMeshVertexBaker*>(Data);
 	
 	ECorrespondenceStrategy UseStrategy = Baker->CorrespondenceStrategy;
-	bool bIsIdentity = true;
-	int NumDetailMeshes = 0;
-	auto CheckIdentity = [Baker, &bIsIdentity, &NumDetailMeshes](const void* DetailMesh)
+	if (UseStrategy == ECorrespondenceStrategy::Identity)
 	{
-		// When the mesh pointers differ, loosely compare the meshes as a sanity check.
-		// TODO: Expose additional comparison metrics on the detail sampler when the mesh pointers differ.
-		bIsIdentity = bIsIdentity && (Baker->TargetMesh == DetailMesh || Baker->TargetMesh->TriangleCount() == Baker->DetailSampler->GetTriangleCount(DetailMesh));
-		++NumDetailMeshes;
-	};
-	Baker->DetailSampler->ProcessMeshes(CheckIdentity);
-	if (UseStrategy == ECorrespondenceStrategy::Identity && !ensure(bIsIdentity && NumDetailMeshes == 1))
-	{
-		// Identity strategy requires there to be only one mesh that is the same
-		// as the target mesh.
-		UseStrategy = ECorrespondenceStrategy::NearestPoint;
+		bool bIsIdentity = true;
+		int NumDetailMeshes = 0;
+		auto CheckIdentity = [Baker, &bIsIdentity, &NumDetailMeshes](const void* DetailMesh)
+		{
+			// When the mesh pointers differ, loosely compare the meshes as a sanity check.
+			// TODO: Expose additional comparison metrics on the detail sampler when the mesh pointers differ.
+			bIsIdentity = bIsIdentity && (Baker->TargetMesh == DetailMesh || Baker->TargetMesh->TriangleCount() == Baker->DetailSampler->GetTriangleCount(DetailMesh));
+			++NumDetailMeshes;
+		};
+		Baker->DetailSampler->ProcessMeshes(CheckIdentity);
+		if (!ensure(bIsIdentity && NumDetailMeshes == 1))
+		{
+			// Identity strategy requires there to be only one mesh that is the same
+			// as the target mesh.
+			UseStrategy = ECorrespondenceStrategy::NearestPoint;
+		}
 	}
 	
 	const FDynamicMesh3* Mesh = Baker->TargetMesh;
