MeshModelingTools: Various CubeGrid improvements.

- Made the tool Accept/Cancel.
- Added the ability to set block size at current power.
- Added ability to initialize the grid from a given actor, initialized to current target.
- Added the option for the grid sizes to follow the default editor grid (1, 5, 10, 50, 100, etc).
- Allowed global mode for gizmo manipulation, to allow editor grid snapping to work.
- Removed the non-useful selection plane tolerance setting

diff --git a/Engine/Plugins/Experimental/MeshModelingToolsetExp/Source/MeshModelingToolsExp/Private/CubeGridTool.cpp b/Engine/Plugins/Experimental/MeshModelingToolsetExp/Source/MeshModelingToolsExp/Private/CubeGridTool.cpp
--- a/Engine/Plugins/Experimental/MeshModelingToolsetExp/Source/MeshModelingToolsExp/Private/CubeGridTool.cpp
+++ b/Engine/Plugins/Experimental/MeshModelingToolsetExp/Source/MeshModelingToolsExp/Private/CubeGridTool.cpp
@@ -85,7 +85,6 @@ namespace CubeGridToolLocals
 	const double CornerModeWireframeThickness = SelectionGridLineThickness;
 	const double CornerModeWireframeDepthBias = 0.1;
 
-
 	/**
 	 * Undoes the actual mesh change (i.e. after boolean operations)
 	 */
@@ -748,6 +747,8 @@ void UCubeGridTool::Setup()
 
 		CurrentMeshTransform = UE::ToolTarget::GetLocalToWorldTransform(Target);
 		CurrentMeshMaterials = UE::ToolTarget::GetMaterialSet(Target).Materials;
+
+		ToolActions->GridSourceActor = UE::ToolTarget::GetTargetActor(Target);
 	}
 
 	MeshSpatial = MakeShared<FDynamicMeshAABBTree3, ESPMode::ThreadSafe>();
@@ -782,8 +783,10 @@ void UCubeGridTool::Setup()
 
 	CubeGrid = MakeShared<FCubeGrid>();
 	CubeGrid->SetGridFrame(FFrame3d(Settings->GridFrameOrigin, Settings->GridFrameOrientation.Quaternion()));
-	CubeGrid->SetBaseGridCellSize(Settings->BlockBaseSize);
-	CubeGrid->SetGridPower(Settings->PowerOfTwo);
+	CubeGrid->SetGridPowerMode(Settings->bPowerOfTwoBlockSizes ? FCubeGrid::EPowerMode::PowerOfTwo : FCubeGrid::EPowerMode::FiveAndTen);
+	CubeGrid->SetGridPower(Settings->GridPower);
+	CubeGrid->SetCurrentGridCellSize(Settings->CurrentBlockSize);
+	Settings->BlockBaseSize = CubeGrid->GetBaseGridCellSize();
 
 	FTransform GridTransform = CubeGrid->GetFrame().ToFTransform();
 
@@ -803,7 +806,6 @@ void UCubeGridTool::Setup()
 
 	GridGizmo = UE::TransformGizmoUtil::CreateCustomTransformGizmo(GetToolManager(),
 		ETransformGizmoSubElements::StandardTranslateRotate, this);
-	GridGizmo->bUseContextCoordinateSystem = false;
 	GridGizmo->SetActiveTarget(GridGizmoTransformProxy, GetToolManager());
 	GridGizmoAlignmentMechanic = NewObject<UDragAlignmentMechanic>(this);
 	GridGizmoAlignmentMechanic->Setup(this);
@@ -879,20 +881,62 @@ void UCubeGridTool::Setup()
 	};
 	AddInputBehavior(MiddleClickDragBehavior);
 
-	PowerOfTwoPrevious = Settings->PowerOfTwo;
-	Settings->WatchProperty(Settings->PowerOfTwo,
+	GridPowerWatcherIdx = Settings->WatchProperty(Settings->GridPower,
 		[this](uint8 NewPower) {
-			SetPowerOfTwoClamped(Settings->PowerOfTwo);
+			SetGridPowerClamped(Settings->GridPower);
+		});
+
+	auto ChangeCurrentBlockSize = [this](double BlockSizeIn)
+	{
+		ClearHover();
+
+		CubeGrid->SetCurrentGridCellSize(BlockSizeIn);
+
+		// Update our views of size. The current block size may end up slightly different due to 
+		// rounding error, since it's actually the base block size that is authoratative in cubegrid.
+		Settings->CurrentBlockSize = CubeGrid->GetCurrentGridCellSize();
+		Settings->BlockBaseSize = CubeGrid->GetBaseGridCellSize();
+
+		Settings->SilentUpdateWatcherAtIndex(CurrentBlockSizeWatcherIdx);
+		Settings->SilentUpdateWatcherAtIndex(BlockBaseSizeWatcherIdx);
+		NotifyOfPropertyChangeByTool(Settings);
+
+		UpdateSelectionLineSet();
+		UpdateGridLineSet();
+	};
+
+	Settings->WatchProperty(Settings->bPowerOfTwoBlockSizes,
+		[ChangeCurrentBlockSize, this](bool bOn) {
+			CubeGrid->SetGridPowerMode(Settings->bPowerOfTwoBlockSizes ? FCubeGrid::EPowerMode::PowerOfTwo : FCubeGrid::EPowerMode::FiveAndTen);
+			
+			// We reset grid power and cell size to default. The reason we do this is because a user is likely to
+			// only change this setting if they want the cube sizes to match up to editor grid snapping values. This
+			// requires the grid size to be set a particular way, and while a user could realize this and set it
+			// appropriately, it feels very clunky to do so.
+			SetGridPowerClamped(Settings->DEFAULT_GRID_POWER);
+			ChangeCurrentBlockSize(Settings->DEFAULT_CURRENT_BLOCK_SIZE);
+		});
+
+	CurrentBlockSizeWatcherIdx = Settings->WatchProperty(Settings->CurrentBlockSize,
+		[ChangeCurrentBlockSize](double NewSize) {
+			ChangeCurrentBlockSize(NewSize);
 		});
-	Settings->WatchProperty(Settings->BlockBaseSize,
+
+	BlockBaseSizeWatcherIdx = Settings->WatchProperty(Settings->BlockBaseSize,
 		[this](double NewBaseSize) {
 			ClearHover();
 
 			CubeGrid->SetBaseGridCellSize(NewBaseSize);
 
+			Settings->CurrentBlockSize = CubeGrid->GetCurrentGridCellSize();
+			Settings->SilentUpdateWatcherAtIndex(CurrentBlockSizeWatcherIdx);
+			NotifyOfPropertyChangeByTool(Settings);
+
 			UpdateSelectionLineSet();
 			UpdateGridLineSet();
 		});
+	
+
 	Settings->WatchProperty(Settings->bShowGizmo,
 		[this](bool bOn) {
 			UpdateGizmoVisibility(bOn);
@@ -928,24 +972,17 @@ void UCubeGridTool::Setup()
 			InvalidatePreview(false);
 		});
 
-	auto UpdateFromDetailsPanelTransformChange = [this]()
-	{
-		CubeGrid->SetGridFrame(FFrame3d(Settings->GridFrameOrigin, Settings->GridFrameOrientation.Quaternion()));
-
-		FTransform GridTransform = CubeGrid->GetFrame().ToFTransform();
-		LineSets->SetTransform(GridTransform);
-		GridGizmo->ReinitializeGizmoTransform(GridTransform);
-
-		InvalidatePreview(false);
-	};
-
 	GridFrameOriginWatcherIdx = Settings->WatchProperty(Settings->GridFrameOrigin,
-		[UpdateFromDetailsPanelTransformChange](FVector) {
-			UpdateFromDetailsPanelTransformChange();
+		[this](FVector) {
+			UpdateGridTransform(FTransform(Settings->GridFrameOrientation, Settings->GridFrameOrigin, FVector::One()), 
+				// Update gizmo, don't update detail panel
+				true, false);
 		});
 	GridFrameOrientationWatcherIdx = Settings->WatchProperty(Settings->GridFrameOrientation,
-		[UpdateFromDetailsPanelTransformChange](FRotator) {
-			UpdateFromDetailsPanelTransformChange();
+		[this](FRotator) {
+			UpdateGridTransform(FTransform(Settings->GridFrameOrientation, Settings->GridFrameOrigin, FVector::One()), 
+				// Update gizmo, don't update detail panel
+				true, false);
 		});
 
 	Settings->SilentUpdateWatched();
@@ -994,58 +1031,79 @@ void UCubeGridTool::Shutdown(EToolShutdownType ShutdownType)
 		ApplyCornerMode(true);
 	}
 
-	bool bCreatingNewAsset = !Target && CurrentMesh->TriangleCount() > 0 && ShutdownType != EToolShutdownType::Cancel;
 	if (Target)
 	{
 		Cast<IPrimitiveComponentBackedTarget>(Target)->SetOwnerVisibility(true);
+	}
 
-		// We check the shutdown type because even though we are not an accept/cancel tool, we
-		// get a cancel shutdown via Ctrl+Z. In this case we definitely don't want to update because
-		// an update in an undo transaction results in a crash.
-		if (bChangesMade && ShutdownType != EToolShutdownType::Cancel && Target->IsValid())
+	// CubeGrid might get used for long stretches at a time, and an accidental Esc hit could result in a fair
+	// amount of wasted work. So if the user made any saveable changes, we'll make sure they definitely want to
+	// discard their work. Unfortunately we can't currently give a classic "are you sure you want to quit" sort
+	// of message because our tool can't actually prevent a shutdown. Hence we have this "do you want to save 
+	// the changes instead" message.
+	if (ShutdownType == EToolShutdownType::Cancel && bChangesMade && (Target || CurrentMesh->TriangleCount() > 0))
+	{
+		FText Title = LOCTEXT("AcceptChangesTitle", "Accept changes instead?");
+		EAppReturnType::Type Ret = FMessageDialog::Open(EAppMsgType::YesNo,
+			LOCTEXT("AcceptChangesQuestion", "The tool is being cancelled, which normally discards all changes. "
+				"Would you like to apply the changes instead?\n\n Selecting \"No\" or closing this window will "
+				"discard the tool's current work."), &Title);
+		if (Ret == EAppReturnType::Yes)
 		{
-			GetToolManager()->BeginUndoTransaction(LOCTEXT("CubeGridToolEditTransactionName", "Block Tool Edit"));
-			FComponentMaterialSet OutputMaterialSet;
-			OutputMaterialSet.Materials = CurrentMeshMaterials;
-
-			UE::ToolTarget::CommitDynamicMeshUpdate(Target, *CurrentMesh, true,
-				FConversionToMeshDescriptionOptions(), &OutputMaterialSet);
-			GetToolManager()->EndUndoTransaction();
+			ShutdownType = EToolShutdownType::Accept;
 		}
-		else if (!Target->IsValid() && CurrentMesh->TriangleCount() > 0)
+	}
+
+	if (ShutdownType == EToolShutdownType::Accept && bChangesMade)
+	{
+		bool bCreatingNewAsset = !Target && CurrentMesh->TriangleCount() > 0;
+		if (Target)
 		{
-			FText Title = LOCTEXT("RecreateAssetTitle", "Recreate Mesh Asset?");
-			EAppReturnType::Type Ret = FMessageDialog::Open(EAppMsgType::YesNo,
-				LOCTEXT("RecreateAssetQuestion", "The underlying asset that this tool was "
-					"operating on seems to no longer be valid (it was likely forcibly removed). "
-					"Would you like to recreate a new asset from the tool's current working "
-					"mesh? Selecting \"No\" or closing this window will discard the tool's "
-					"current work."));
-			if (Ret == EAppReturnType::Yes)
+			if (Target->IsValid())
+			{
+				GetToolManager()->BeginUndoTransaction(LOCTEXT("CubeGridToolEditTransactionName", "Cube Grid Edit"));
+				FComponentMaterialSet OutputMaterialSet;
+				OutputMaterialSet.Materials = CurrentMeshMaterials;
+
+				UE::ToolTarget::CommitDynamicMeshUpdate(Target, *CurrentMesh, true,
+					FConversionToMeshDescriptionOptions(), &OutputMaterialSet);
+				GetToolManager()->EndUndoTransaction();
+			}
+			else if (!Target->IsValid() && CurrentMesh->TriangleCount() > 0)
 			{
-				bCreatingNewAsset = true;
+				FText Title = LOCTEXT("RecreateAssetTitle", "Recreate Mesh Asset?");
+				EAppReturnType::Type Ret = FMessageDialog::Open(EAppMsgType::YesNo,
+					LOCTEXT("RecreateAssetQuestion", "The underlying asset that this tool was "
+						"operating on seems to no longer be valid (it was likely forcibly removed). "
+						"Would you like to recreate a new asset from the tool's current working "
+						"mesh? Selecting \"No\" or closing this window will discard the tool's "
+						"current work."), &Title);
+				if (Ret == EAppReturnType::Yes)
+				{
+					bCreatingNewAsset = true;
+				}
 			}
 		}
-	}
-
-	if (bCreatingNewAsset)
-	{
-		GetToolManager()->BeginUndoTransaction(LOCTEXT("CubeGridToolCreateTransactionName", "Block Tool Create New"));
 
-		FCreateMeshObjectParams NewMeshObjectParams;
-		NewMeshObjectParams.TargetWorld = TargetWorld;
-		NewMeshObjectParams.Transform = (FTransform)CurrentMeshTransform;
-		NewMeshObjectParams.BaseName = TEXT("CubeGridToolOutput");
-		NewMeshObjectParams.Materials = CurrentMeshMaterials;
-		NewMeshObjectParams.SetMesh(CurrentMesh.Get());
-		OutputTypeProperties->ConfigureCreateMeshObjectParams(NewMeshObjectParams);
-		FCreateMeshObjectResult Result = UE::Modeling::CreateMeshObject(GetToolManager(), MoveTemp(NewMeshObjectParams));
-		if (Result.IsOK() && Result.NewActor != nullptr)
+		if (bCreatingNewAsset)
 		{
-			ToolSelectionUtil::SetNewActorSelection(GetToolManager(), Result.NewActor);
-		}
+			GetToolManager()->BeginUndoTransaction(LOCTEXT("CubeGridToolCreateTransactionName", "Cube Grid Create New"));
+
+			FCreateMeshObjectParams NewMeshObjectParams;
+			NewMeshObjectParams.TargetWorld = TargetWorld;
+			NewMeshObjectParams.Transform = (FTransform)CurrentMeshTransform;
+			NewMeshObjectParams.BaseName = TEXT("CubeGridToolOutput");
+			NewMeshObjectParams.Materials = CurrentMeshMaterials;
+			NewMeshObjectParams.SetMesh(CurrentMesh.Get());
+			OutputTypeProperties->ConfigureCreateMeshObjectParams(NewMeshObjectParams);
+			FCreateMeshObjectResult Result = UE::Modeling::CreateMeshObject(GetToolManager(), MoveTemp(NewMeshObjectParams));
+			if (Result.IsOK() && Result.NewActor != nullptr)
+			{
+				ToolSelectionUtil::SetNewActorSelection(GetToolManager(), Result.NewActor);
+			}
 
-		GetToolManager()->EndUndoTransaction();
+			GetToolManager()->EndUndoTransaction();
+		}
 	}
 
 	Preview->OnOpCompleted.RemoveAll(this);
@@ -1182,24 +1240,9 @@ void UCubeGridTool::ClearHover()
 
 void  UCubeGridTool::GridGizmoMoved(UTransformProxy* Proxy, FTransform Transform)
 {
-	Settings->GridFrameOrigin = Transform.GetTranslation();
-	Settings->SilentUpdateWatcherAtIndex(GridFrameOriginWatcherIdx);
-
-	Settings->GridFrameOrientation = Transform.GetRotation().Rotator();
-	Settings->SilentUpdateWatcherAtIndex(GridFrameOrientationWatcherIdx);
-
-	CubeGrid->SetGridFrame(FFrame3d(Settings->GridFrameOrigin, Settings->GridFrameOrientation.Quaternion()));
-	LineSets->SetTransform(CubeGrid->GetFrame().ToFTransform());
-
-	if (Mode == EMode::Corner)
-	{
-		InvalidatePreview(false);
-
-		if (!bInGizmoDrag)
-		{
-			UpdateCornerGeometrySet();
-		}
-	}
+	UpdateGridTransform(Transform, 
+		// Don't update gizmo, update detail panel
+		false, true);
 }
 
 void UCubeGridTool::OnCtrlMiddleClick(const FInputDeviceRay& ClickPos)
@@ -1270,13 +1313,9 @@ void UCubeGridTool::OnCtrlMiddleClick(const FInputDeviceRay& ClickPos)
 	}
 
 	// Move the gizmo to that corner.
-	Settings->GridFrameOrigin = CubeGrid->ToWorldPoint(Corners[ClosestCornerIndex]);
-	Settings->SilentUpdateWatcherAtIndex(GridFrameOriginWatcherIdx);
-	CubeGrid->SetGridFrame(FFrame3d(Settings->GridFrameOrigin, Settings->GridFrameOrientation.Quaternion()));
-
-	FTransform GridTransform = CubeGrid->GetFrame().ToFTransform();
-	LineSets->SetTransform(GridTransform);
-	GridGizmo->SetNewGizmoTransform(GridTransform);
+	UpdateGridTransform(FTransform(Settings->GridFrameOrientation, CubeGrid->ToWorldPoint(Corners[ClosestCornerIndex]), FVector::One())
+		// Update both gizmo and detail panel
+		, true, true);
 
 	GetToolManager()->EndUndoTransaction();
 }
@@ -1404,6 +1443,36 @@ void UCubeGridTool::UpdateGizmoVisibility(bool bVisible)
 	Settings->SilentUpdateWatched();
 }
 
+void UCubeGridTool::UpdateGridTransform(FTransform NewTransform, bool bUpdateGizmo, bool bUpdateDetailPanel)
+{
+	if (bUpdateGizmo)
+	{
+		GridGizmo->ReinitializeGizmoTransform(NewTransform);
+	}
+
+	if (bUpdateDetailPanel)
+	{
+		Settings->GridFrameOrigin = NewTransform.GetTranslation();
+		Settings->SilentUpdateWatcherAtIndex(GridFrameOriginWatcherIdx);
+
+		Settings->GridFrameOrientation = NewTransform.GetRotation().Rotator();
+		Settings->SilentUpdateWatcherAtIndex(GridFrameOrientationWatcherIdx);
+	}
+
+	CubeGrid->SetGridFrame(FFrame3d(NewTransform));
+	LineSets->SetTransform(NewTransform);
+
+	if (Mode == EMode::Corner)
+	{
+		InvalidatePreview(false);
+
+		if (!bInGizmoDrag)
+		{
+			UpdateCornerGeometrySet();
+		}
+	}
+}
+
 bool UCubeGridTool::GetHitGridFace(const FRay& WorldRay, FCubeGrid::FCubeFace& FaceOut)
 {
 	using namespace CubeGridToolLocals;
@@ -1936,34 +2005,71 @@ void UCubeGridTool::ApplyPushPull(int32 NumBlocks)
 	}
 }
 
-void UCubeGridTool::SetPowerOfTwoClamped(int32 PowerOfTwo)
+void UCubeGridTool::SetGridPowerClamped(int32 GridPower)
 {
-	Settings->PowerOfTwo = FMath::Clamp<int32>(PowerOfTwo, 0, Settings->MaxPowerOfTwo);
-	CubeGrid->SetGridPower(Settings->PowerOfTwo);
+	Settings->GridPower = FMath::Clamp<int32>(GridPower, 0, Settings->MaxGridPower);
+	CubeGrid->SetGridPower(Settings->GridPower);
+	Settings->SilentUpdateWatcherAtIndex(GridPowerWatcherIdx);
+
+	Settings->CurrentBlockSize = CubeGrid->GetCurrentGridCellSize();
+	Settings->SilentUpdateWatcherAtIndex(CurrentBlockSizeWatcherIdx);
+	
+	NotifyOfPropertyChangeByTool(Settings);
 
-	// Update CurrentExtrudeAmount to reflect the new step size (mainly important
-	// for corner mode).
-	int32 PowerOfTwoDifference = static_cast<int32>(Settings->PowerOfTwo) - PowerOfTwoPrevious;
-	if (CurrentExtrudeAmount != 0 && PowerOfTwoDifference != 0)
+	if (Mode == EMode::Corner)
 	{
-		int32 AbsExtrudeAmount = FMath::Abs(CurrentExtrudeAmount);
+		InvalidatePreview(); // effective extrude distance is now different
+		UpdateCornerModeLineSet();
 
-		int32 NewAbsExtrudeAmount = AbsExtrudeAmount >> PowerOfTwoDifference;
-		if (NewAbsExtrudeAmount << PowerOfTwoDifference != AbsExtrudeAmount)
+		// Note: If we wanted to, we could adjust the extrude amount in corner mode such that the currently extruded
+		// distance stays roughly the same. However, it's not clear whether that is desirable. On the one hand, it might
+		// make it easier to reach certain far-from-zero values if they are not on a big grid boundary. On the other hand,
+		// the behavior makes it harder to see how the ramp will change on the next E/Q press.
+		// For now we decided against it. If we ever want it, we would use this code instead of the above (requires 
+		// keeping track of GridPowerPrevious):
+		if constexpr (false)
 		{
-			++NewAbsExtrudeAmount;
-			InvalidatePreview();
-		}
-
-		CurrentExtrudeAmount = NewAbsExtrudeAmount * FMath::Sign(CurrentExtrudeAmount);
+			// Needs to match behavior in CubeGrid.cpp. We would probably expose it if we used this code block.
+			auto GetMultiplierForGridPower = [](bool bPowerOfTwo, uint8 GridPower)
+			{
+				if (bPowerOfTwo)
+				{
+					// This is split up into two statements to avoid a static analysis warning about
+					// shifting a 32 bit value and casting to a 64 bit value.
+					uint32 ShiftedResult = static_cast<uint32>(1) << GridPower;
+					return static_cast<double>(ShiftedResult);
+				}
+				else
+				{
+					// For FiveAndTen, we multiply  by 2 half the time and by 5 the second half, rounding up for 2's.
+					uint8 FloorHalfGridPower = GridPower / 2;
+					uint32 TwoMultiplier = static_cast<uint32>(1) << (GridPower - FloorHalfGridPower);
+					return TwoMultiplier * FMath::Pow(5.0, static_cast<double>(FloorHalfGridPower));
+				}
+			};
+
+			// This would be made a member in the tool
+			uint8 GridPowerPrevious = 0;
+
+			double PreviousMultiplier = GetMultiplierForGridPower(Settings->bPowerOfTwoBlockSizes, GridPowerPrevious);
+			double NewMultiplier = GetMultiplierForGridPower(Settings->bPowerOfTwoBlockSizes, Settings->GridPower);;
+
+			double NewExtrudeAmount = CurrentExtrudeAmount * PreviousMultiplier / NewMultiplier;
+			// Could consider allowing CurrentExtrudeAmount not to be an int (and moving it to the nearest int
+			// on next manipulation).
+			CurrentExtrudeAmount = FMath::RoundFromZero(NewExtrudeAmount);
+
+			if (CurrentExtrudeAmount != NewExtrudeAmount)
+			{
+				InvalidatePreview();
+				UpdateCornerModeLineSet();
+			}
+		}//end unused adjustment code
 	}
-	UpdateCornerModeLineSet();
 
 	ClearHover();
 	UpdateSelectionLineSet(); // Updates the grid drawn inside
 	UpdateGridLineSet();
-
-	PowerOfTwoPrevious = Settings->PowerOfTwo;
 }
 
 // Action support
@@ -2011,12 +2117,12 @@ void UCubeGridTool::ApplyAction(ECubeGridToolAction ActionType)
 	case ECubeGridToolAction::SlideBack:
 		ApplySlide(Settings->BlocksPerStep);
 		break;
-	case ECubeGridToolAction::DecreasePowerOfTwo:
+	case ECubeGridToolAction::DecreaseGridPower:
 		// cast is just to be explicit
-		SetPowerOfTwoClamped(static_cast<int32>(Settings->PowerOfTwo) - 1); 
+		SetGridPowerClamped(static_cast<int32>(Settings->GridPower) - 1); 
 		break;
-	case ECubeGridToolAction::IncreasePowerOfTwo:
-		SetPowerOfTwoClamped(Settings->PowerOfTwo + 1);
+	case ECubeGridToolAction::IncreaseGridPower:
+		SetGridPowerClamped(Settings->GridPower + 1);
 		break;
 
 	case ECubeGridToolAction::CornerMode:
@@ -2038,6 +2144,14 @@ void UCubeGridTool::ApplyAction(ECubeGridToolAction ActionType)
 	case ECubeGridToolAction::Cancel:
 		RevertToDefaultMode();
 		break;
+	case ECubeGridToolAction::ResetFromActor:
+		if (ToolActions->GridSourceActor)
+		{
+			UpdateGridTransform(ToolActions->GridSourceActor->GetTransform(), 
+				// Update gizmo and detail panel
+				true, true);
+		}
+		break;
 	}
 }
 
@@ -2072,17 +2186,17 @@ void UCubeGridTool::RegisterActions(FInteractiveToolActionSet& ActionSet)
 		[this]() { RequestAction(ECubeGridToolAction::SlideForward); });
 
 	ActionSet.RegisterAction(this, ActionID++,
-		TEXT("DecreasePowerOfTwo"),
-		LOCTEXT("DecreasePowerOfTwoAction", "Decrease Power Of Two"),
-		LOCTEXT("DecreasePowerOfTwoTooltip", ""),
+		TEXT("DecreaseGridPower"),
+		LOCTEXT("DecreaseGridPowerAction", "Decrease Grid Power"),
+		LOCTEXT("DecreaseGridPowerTooltip", ""),
 		EModifierKey::Control, EKeys::Q,
-		[this]() { RequestAction(ECubeGridToolAction::DecreasePowerOfTwo); });
+		[this]() { RequestAction(ECubeGridToolAction::DecreaseGridPower); });
 	ActionSet.RegisterAction(this, ActionID++,
-		TEXT("IncreasePowerOfTwo"),
-		LOCTEXT("IncreasePowerOfTwoAction", "Increase Power Of Two"),
-		LOCTEXT("IncreasePowerOfTwoTooltip", ""),
+		TEXT("IncreaseGridPower"),
+		LOCTEXT("IncreaseGridPowerAction", "Increase Grid Power"),
+		LOCTEXT("IncreaseGridPowerTooltip", ""),
 		EModifierKey::Control, EKeys::E,
-		[this]() { RequestAction(ECubeGridToolAction::IncreasePowerOfTwo); });
+		[this]() { RequestAction(ECubeGridToolAction::IncreaseGridPower); });
 
 	ActionSet.RegisterAction(this, ActionID++,
 		TEXT("ToggleGizmoVisibility"),
diff --git a/Engine/Plugins/Experimental/MeshModelingToolsetExp/Source/MeshModelingToolsExp/Public/CubeGridTool.h b/Engine/Plugins/Experimental/MeshModelingToolsetExp/Source/MeshModelingToolsExp/Public/CubeGridTool.h
--- a/Engine/Plugins/Experimental/MeshModelingToolsetExp/Source/MeshModelingToolsExp/Public/CubeGridTool.h
+++ b/Engine/Plugins/Experimental/MeshModelingToolsetExp/Source/MeshModelingToolsExp/Public/CubeGridTool.h
@@ -91,23 +91,46 @@ public:
 	UPROPERTY(EditAnywhere, Category = Options, meta = (TransientToolProperty))
 	bool bShowGizmo = false;
 
-	/** How many blocks each push/pull invocation will do at a time.*/
-	UPROPERTY(EditAnywhere, Category = Options, meta = (
-		UIMin = "0", ClampMin = "0"))
-	int32 BlocksPerStep = 1;
+	// These are here so that we can reset the appropriate settings to their defaults in code.
+	const uint8 DEFAULT_GRID_POWER = 5;
+	const double DEFAULT_CURRENT_BLOCK_SIZE = 100;
 
 	/** Determines cube grid scale. Can also be adjusted with Ctrl + E/Q. */
 	UPROPERTY(EditAnywhere, Category = Options, meta = (
 		EditCondition = "bAllowedToEditGrid", HideEditConditionToggle,
 		UIMin = "0", UIMax = "10", ClampMin = "0", ClampMax = "31"))
-	uint8 PowerOfTwo = 5;
+	uint8 GridPower = DEFAULT_GRID_POWER;
+
+	/** 
+	 * Sets the size of a block at the current grid power. This is done by changing the 
+	 * base block size (i.e. the size at grid power 0) such that the target size is achieved at 
+	 * the the current value of Grid Power.
+	 */
+	UPROPERTY(EditAnywhere, Category = Options, meta = (
+		EditCondition = "bAllowedToEditGrid", HideEditConditionToggle,
+		UIMin = "1", UIMax = "1000", ClampMin = "0.001"))
+	double CurrentBlockSize = DEFAULT_CURRENT_BLOCK_SIZE;
 
-	// Must match ClampMax in PowerOfTwo, used to make hotkeys not exceed it.
-	const uint8 MaxPowerOfTwo = 31;
+	/** How many blocks each push/pull invocation will do at a time.*/
+	UPROPERTY(EditAnywhere, Category = Options, meta = (
+		UIMin = "0", ClampMin = "0"))
+	int32 BlocksPerStep = 1;
+
+	/** 
+	 * When true, block sizes change by powers of two as grid power is changed. When false, block
+	 * sizes change by twos and fives, much like the default editor grid snapping options (for
+	 * instance, sizes might increase from 10 to 50 to 100 to 500).
+	 * Note that toggling this option will reset BlocksPerStep and CurrentBlockSize to default values.
+	 */
+	UPROPERTY(EditAnywhere, Category = Options, AdvancedDisplay)
+	bool bPowerOfTwoBlockSizes = true;
+
+	// Must match ClampMax in GridPower, used to make hotkeys not exceed it.
+	const uint8 MaxGridPower = 31;
 
 	/** Smallest block size to use in the grid. For instance, 3.125 results in
 	 blocks that are 100 sized at 5 power of two since 3.125 * 2^5 = 100. */
-	UPROPERTY(EditAnywhere, Category = Options, meta = (
+	UPROPERTY(EditAnywhere, Category = Options, AdvancedDisplay, meta = (
 		EditCondition = "bAllowedToEditGrid", HideEditConditionToggle,
 		UIMin = "0.1", UIMax = "10", ClampMin = "0.001", ClampMax = "1000"))
 	double BlockBaseSize = 3.125;
@@ -127,8 +150,10 @@ public:
 
 	/** When performing selection, the tolerance to use when determining
 	 whether things lie in the same plane as a cube face. */
-	UPROPERTY(EditAnywhere, Category = BlockSelection, AdvancedDisplay, meta = (
-		UIMin = "0", UIMax = "0.5", ClampMin = "0", ClampMax = "10"))
+	//~ This turned out to not be a useful setting, so it is no longer EditAnywhere. The only cases where it
+	//~ seems to have a noticeable effect were ones where it was set so high that it broke the selection
+	//~ behavior slightly.
+	UPROPERTY()
 	double PlaneTolerance = 0.01;
 
 	/** When raycasting to find a selected grid face, this determines whether geometry
@@ -189,10 +214,11 @@ enum class ECubeGridToolAction
 	Flip,
 	SlideForward,
 	SlideBack,
-	DecreasePowerOfTwo,
-	IncreasePowerOfTwo,
+	DecreaseGridPower,
+	IncreaseGridPower,
 	CornerMode,
 	// FitGrid,
+	ResetFromActor,
 	Done,
 	Cancel,
 };
@@ -234,6 +260,20 @@ public:
 	/** Can also be invoked with T. */
 	UFUNCTION(CallInEditor, Category = Actions, meta = (DisplayPriority = 6))
 	void Flip() { PostAction(ECubeGridToolAction::Flip); }
+
+	/** Actor whose transform to use when doing Reset Grid From Actor. */
+	//~ For some reason we can't seem to use TWeakObjectPtr here- it becomes unsettable in the tool.
+	UPROPERTY(EditAnywhere, Category = GridReinitialization, meta = (TransientToolProperty))
+	TObjectPtr<AActor> GridSourceActor = nullptr;
+
+	/** 
+	 * Resets the grid position and orientation based on the actor in Grid Source Actor. This allows 
+	 * grid positions/orientations to be saved by pasting them into the transform of some actor that
+	 * is later used, or by relying on the fact that the tool initializes transforms of newly created
+	 * meshes based on the grid used.
+	 */
+	UFUNCTION(CallInEditor, Category = GridReinitialization)
+	void ResetGridFromActor () { PostAction(ECubeGridToolAction::ResetFromActor); }
 };
 
 UCLASS()
@@ -314,6 +354,10 @@ public:
 		}
 	};
 
+	virtual bool HasAccept() const override { return true; };
+	virtual bool CanAccept() const override { return true; };
+	virtual bool HasCancel() const override { return true; };
+
 	virtual void Setup() override;
 	virtual void Shutdown(EToolShutdownType ShutdownType) override;
 
@@ -454,9 +498,8 @@ protected:
 	void ApplySlide(int32 NumBlocks);
 	void ApplyPushPull(int32 NumBlocks);
 
-	// Parameter is signed on purpose so we can give negatives
-	void SetPowerOfTwoClamped(int32 PowerOfTwo);
-	uint8 PowerOfTwoPrevious;
+	// Parameter is signed on purpose so we can give negatives for clamping.
+	void SetGridPowerClamped(int32 GridPower);
 
 	UPROPERTY()
 	TObjectPtr<UMeshOpPreviewWithBackgroundCompute> Preview;
@@ -504,10 +547,15 @@ protected:
 	ECubeGridToolAction PendingAction = ECubeGridToolAction::NoAction;
 	void ApplyAction(ECubeGridToolAction ActionType);
 
+	int32 GridPowerWatcherIdx;
+	int32 BlockBaseSizeWatcherIdx;
+	int32 CurrentBlockSizeWatcherIdx;
+
 	int32 GridFrameOriginWatcherIdx;
 	int32 GridFrameOrientationWatcherIdx;
 	void GridGizmoMoved(UTransformProxy* Proxy, FTransform Transform);
 	void UpdateGizmoVisibility(bool bVisible);
+	void UpdateGridTransform(FTransform NewTransform, bool bUpdateGizmo, bool bUpdateDetailPanel);
 	bool bInGizmoDrag = false;
 
 	FVector3d MiddleClickDragStart;
diff --git a/Engine/Plugins/Runtime/MeshModelingToolset/Source/ModelingComponents/Private/Mechanics/CubeGrid.cpp b/Engine/Plugins/Runtime/MeshModelingToolset/Source/ModelingComponents/Private/Mechanics/CubeGrid.cpp
--- a/Engine/Plugins/Runtime/MeshModelingToolset/Source/ModelingComponents/Private/Mechanics/CubeGrid.cpp
+++ b/Engine/Plugins/Runtime/MeshModelingToolset/Source/ModelingComponents/Private/Mechanics/CubeGrid.cpp
@@ -16,6 +16,24 @@ namespace CubeGridLocals {
 			FMath::Floor(VectorIn.Y),
 			FMath::Floor(VectorIn.Z));
 	}
+
+	double GetMultiplierForGridPower(bool bPowerOfTwo, uint8 GridPower)
+	{
+		if (bPowerOfTwo)
+		{
+			// This is split up into two statements to avoid a static analysis warning about
+			// shifting a 32 bit value and casting to a 64 bit value.
+			uint32 ShiftedResult = static_cast<uint32>(1) << GridPower;
+			return ShiftedResult;
+		}
+		else
+		{
+			// For FiveAndTen, we multiply  by 2 half the time and by 5 the second half, rounding up for 2's.
+			uint8 FloorHalfGridPower = GridPower / 2;
+			uint32 TwoMultiplier = static_cast<uint32>(1) << (GridPower - FloorHalfGridPower);
+			return TwoMultiplier * FMath::Pow(5.0, static_cast<double>(FloorHalfGridPower));
+		}
+	}
 }
 
 FCubeGrid::FCubeFace::FCubeFace(const FVector3d& PointOnFace, EFaceDirection DirectionIn, uint8 SourceCubeGridPower)
@@ -41,6 +59,21 @@ FVector3d FCubeGrid::FCubeFace::GetMaxCorner() const
 		FMath::CeilToDouble(Center.Z));
 }
 
+void FCubeGrid::SetCurrentGridCellSize(double SizeIn)
+{
+	using namespace CubeGridLocals;
+
+	double Multiplier = GetMultiplierForGridPower(GridPowerMode == EPowerMode::PowerOfTwo, CurrentGridPower);
+	double BaseGridSize = SizeIn / Multiplier;
+	SetBaseGridCellSize(BaseGridSize);
+}
+
+double FCubeGrid::GetCellSize(uint8 GridPower) const
+{
+	using namespace CubeGridLocals;
+
+	return BaseGridCellSize * GetMultiplierForGridPower(GridPowerMode == EPowerMode::PowerOfTwo, GridPower);
+}
 
 bool FCubeGrid::GetHitGridFaceBasedOnRay(const FVector3d& WorldHitPoint, const FVector3d& NormalOrTowardCameraRay,
 	FCubeGrid::FCubeFace& FaceOut, bool bPierceToBack, double PlaneTolerance) const
diff --git a/Engine/Plugins/Runtime/MeshModelingToolset/Source/ModelingComponents/Public/Mechanics/CubeGrid.h b/Engine/Plugins/Runtime/MeshModelingToolset/Source/ModelingComponents/Public/Mechanics/CubeGrid.h
--- a/Engine/Plugins/Runtime/MeshModelingToolset/Source/ModelingComponents/Public/Mechanics/CubeGrid.h
+++ b/Engine/Plugins/Runtime/MeshModelingToolset/Source/ModelingComponents/Public/Mechanics/CubeGrid.h
@@ -35,6 +35,12 @@ public:
 		PositiveZ = 3
 	};
 
+	enum class EPowerMode : int8
+	{
+		PowerOfTwo,
+		FiveAndTen
+	};
+
 	/**
 	 * Convert face direction to the index of the nonzero normal dimension, or equivalently,
 	 * the dimension in which that face is flat.
@@ -112,7 +118,7 @@ public:
 		const FVector3d& NormalOrTowardCameraRay,
 		FCubeGrid::FCubeFace& FaceOut, bool bPierceToBack, double PlaneTolerance) const;
 
-	int8 GetGridPower() const { return CurrentGridPower; }
+	uint8 GetGridPower() const { return CurrentGridPower; }
 
 	void SetGridPower(uint8 PowerIn)
 	{
@@ -120,10 +126,16 @@ public:
 		CurrentGridPower = PowerIn;
 		UpdateCellSize();
 	}
-	int8 GetMaxGridPower() const
+	uint8 GetMaxGridPower() const
 	{
 		return sizeof(uint32) * 8 - 1;
 	}
+	void SetGridPowerMode(EPowerMode PowerModeIn)
+	{
+		GridPowerMode = PowerModeIn;
+		UpdateCellSize();
+	}
+	EPowerMode GetGridPowerMode() const { return GridPowerMode; }
 
 	double GetBaseGridCellSize() const { return BaseGridCellSize; }
 	void SetBaseGridCellSize(double SizeIn)
@@ -137,19 +149,16 @@ public:
 	}
 
 	double GetCurrentGridCellSize() const { return CurrentCellSize; }
+	void SetCurrentGridCellSize(double SizeIn);
 
-	double GetCellSize(int8 GridPower) const { 
-		// This is split up into two statements to avoid a static analysis warning about
-		// shifting a 32 bit value and casting to a 64 bit value.
-		uint32 Multiplier = (uint32)1 << GridPower; 
-		return BaseGridCellSize * Multiplier; }
+	double GetCellSize(uint8 GridPower) const;
 
 	FVector3d ToWorldPoint(FVector3d GridPoint) const
 	{
 		return GridFrame.FromFramePoint(GridPoint * CurrentCellSize);
 	}
 
-	FVector3d ToWorldPoint(FVector3d GridPoint, int8 GridPower) const
+	FVector3d ToWorldPoint(FVector3d GridPoint, uint8 GridPower) const
 	{
 		return GridFrame.FromFramePoint(GridPoint * GetCellSize(GridPower));
 	}
@@ -174,7 +183,8 @@ protected:
 
 	FFrame3d GridFrame;
 	double BaseGridCellSize;
-	int8 CurrentGridPower;
+	uint8 CurrentGridPower;
+	EPowerMode GridPowerMode = EPowerMode::PowerOfTwo;
 	// This should always be BaseGridCellSize * 2^CurrentGridPower
 	double CurrentCellSize;
 };
