MeshModelingTools: Fix incorrect comment in CubeGridTool.

diff --git a/Engine/Plugins/Experimental/MeshModelingToolsetExp/Source/MeshModelingToolsExp/Public/CubeGridTool.h b/Engine/Plugins/Experimental/MeshModelingToolsetExp/Source/MeshModelingToolsExp/Public/CubeGridTool.h
--- a/Engine/Plugins/Experimental/MeshModelingToolsetExp/Source/MeshModelingToolsExp/Public/CubeGridTool.h
+++ b/Engine/Plugins/Experimental/MeshModelingToolsetExp/Source/MeshModelingToolsExp/Public/CubeGridTool.h
@@ -120,7 +120,7 @@ public:
 	 * When true, block sizes change by powers of two as grid power is changed. When false, block
 	 * sizes change by twos and fives, much like the default editor grid snapping options (for
 	 * instance, sizes might increase from 10 to 50 to 100 to 500).
-	 * Note that toggling this option will reset BlocksPerStep and CurrentBlockSize to default values.
+	 * Note that toggling this option will reset Grid Power and Current Block Size to default values.
 	 */
 	UPROPERTY(EditAnywhere, Category = Options, AdvancedDisplay)
 	bool bPowerOfTwoBlockSizes = true;
