ModelingTools: add Packed option to Pattern Tool, tiles objects densely based on bounding box extents

diff --git a/Engine/Plugins/Experimental/MeshModelingToolsetExp/Source/MeshModelingToolsExp/Private/PatternTool.cpp b/Engine/Plugins/Experimental/MeshModelingToolsetExp/Source/MeshModelingToolsExp/Private/PatternTool.cpp
--- a/Engine/Plugins/Experimental/MeshModelingToolsetExp/Source/MeshModelingToolsExp/Private/PatternTool.cpp
+++ b/Engine/Plugins/Experimental/MeshModelingToolsetExp/Source/MeshModelingToolsExp/Private/PatternTool.cpp
@@ -24,6 +24,7 @@
 #include "Components/DynamicMeshComponent.h"
 #include "Components/InstancedStaticMeshComponent.h"
 #include "Engine/StaticMeshActor.h"
+#include "Engine/StaticMesh.h"
 
 #include "BaseGizmos/GizmoComponents.h"
 #include "BaseGizmos/TransformGizmoUtil.h"
@@ -67,7 +68,8 @@ public:
 	enum class ESpacingMode
 	{
 		ByCount = 0,
-		ByStepSize = 1
+		ByStepSize = 1,
+		Packed = 2
 	};
 
 public:
@@ -84,6 +86,8 @@ public:
 	FVector3d EndScale = FVector3d::One();
 	bool bInterpolateScale = false;				// if false, only StartScale is used
 
+	FAxisAlignedBox3d Dimensions = FAxisAlignedBox3d(FVector3d::Zero(), 10.0);
+
 public:
 	// current result pattern
 	TArray<FTransformSRT3d> Pattern;
@@ -151,7 +155,7 @@ public:
 		}
 	}
 
-	void ComputeSteps(ESpacingMode UseSpacingMode, int32 CountIn, int32 StepSizeIn, int32 LengthIn, int32& Iterations, double& AlphaStep, bool& bForceLastStepToOne)
+	void ComputeSteps(FPatternGenerator::ESpacingMode UseSpacingMode, int32 CountIn, double StepSizeIn, double LengthIn, int32& Iterations, double& AlphaStep, bool& bForceLastStepToOne)
 	{
 		Iterations = 1;
 		AlphaStep = 0.5;
@@ -166,7 +170,13 @@ public:
 		{
 			bForceLastStepToOne = false;
 			AlphaStep = StepSizeIn / LengthIn;
-			Iterations = FMath::Clamp( (int)(1.0 / AlphaStep) + 1, 1, 1000);
+			Iterations = FMath::Clamp( (int)(1.0 / (AlphaStep + FMathd::ZeroTolerance) ) + 1, 1, 1000);
+		}
+		else if (UseSpacingMode == ESpacingMode::Packed)
+		{
+			bForceLastStepToOne = false;
+			AlphaStep = StepSizeIn / LengthIn;
+			Iterations = FMath::Clamp( (int)(1.0 / (AlphaStep + FMathd::ZeroTolerance) ) + 1, 1, 1000);
 		}
 		else
 		{
@@ -178,7 +188,6 @@ public:
 
 
 
-
 class FLinearPatternGenerator : public FPatternGenerator
 {
 public:
@@ -233,7 +242,6 @@ public:
 };
 
 
-
 void FLinearPatternGenerator::UpdatePattern_LineFill()
 {
 	double LineLength = Distance(StartFrame.Origin, EndFrame.Origin);
@@ -241,7 +249,8 @@ void FLinearPatternGenerator::UpdatePattern_LineFill()
 	int32 Iterations = 1;
 	double AlphaStep = 0.5;
 	bool bForceLastStepToOne = true;
-	ComputeSteps(SpacingMode, Count, StepSize, LineLength, Iterations, AlphaStep, bForceLastStepToOne);
+	double UseStepSize = (SpacingMode == ESpacingMode::Packed) ? Dimensions.Dimension(Axis) : StepSize;
+	ComputeSteps(SpacingMode, Count, UseStepSize, LineLength, Iterations, AlphaStep, bForceLastStepToOne);
 
 	for (int32 k = 0; k < Iterations; ++k)
 	{
@@ -264,8 +273,6 @@ void FLinearPatternGenerator::UpdatePattern_LineFill()
 
 
 
-
-
 void FLinearPatternGenerator::UpdatePattern_RectangleFill()
 {
 	FVector3d LocalPt = StartFrame.ToFramePoint(EndFrame.Origin);
@@ -275,12 +282,14 @@ void FLinearPatternGenerator::UpdatePattern_RectangleFill()
 	int32 IterationsX = 1;
 	double AlphaStepX = 0.5;
 	bool bForceLastStepToOneX = true;
-	ComputeSteps(SpacingMode, Count, StepSize, ExtentX, IterationsX, AlphaStepX, bForceLastStepToOneX);
+	double UseStepSizeX = (SpacingMode == ESpacingMode::Packed) ? Dimensions.Dimension(Axis) : StepSize;
+	ComputeSteps(SpacingMode, Count, UseStepSizeX, ExtentX, IterationsX, AlphaStepX, bForceLastStepToOneX);
 
 	int32 IterationsY = 1;
 	double AlphaStepY = 0.5;
 	bool bForceLastStepToOneY = true;
-	ComputeSteps(SpacingModeY, CountY, StepSizeY, ExtentY, IterationsY, AlphaStepY, bForceLastStepToOneY);
+	double UseStepSizeY = (SpacingModeY == ESpacingMode::Packed) ? Dimensions.Dimension(AxisY): StepSizeY;
+	ComputeSteps(SpacingModeY, CountY, UseStepSizeY, ExtentY, IterationsY, AlphaStepY, bForceLastStepToOneY);
 
 	FFrame3d MidFrameOrientation = Lerp(StartFrame, EndFrame, 0.5);
 	FFrame3d Frame00 = StartFrame;
@@ -311,7 +320,6 @@ void FLinearPatternGenerator::UpdatePattern_RectangleFill()
 
 
 
-
 class FRadialPatternGenerator : public FPatternGenerator
 {
 public:
@@ -352,7 +360,6 @@ public:
 };
 
 
-
 void FRadialPatternGenerator::UpdatePattern_CircleFill()
 {
 	double ArcLengthDeg = FMathd::Abs(EndAngleDeg - StartAngleDeg);
@@ -370,6 +377,13 @@ void FRadialPatternGenerator::UpdatePattern_CircleFill()
 		AlphaStep = StepSizeDeg / ArcLengthDeg;
 		Iterations = FMath::Clamp( (int)(1.0 / AlphaStep) + 1, 1, 1000);
 	}
+	else if (SpacingMode == ESpacingMode::Packed)
+	{
+		double ArcLenStepSize = Dimensions.Dimension(0);		// currently only support X axis in circle pattern?
+		double CalcStepSizeRad = ArcLenStepSize / Radius;
+		AlphaStep = (CalcStepSizeRad * FMathd::RadToDeg) / ArcLengthDeg;
+		Iterations = FMath::Clamp( (int)(1.0 / AlphaStep) + 1, 1, 1000);
+	}
 	else
 	{
 		check(false);
@@ -409,7 +423,6 @@ void FRadialPatternGenerator::UpdatePattern_CircleFill()
 
 
 
-
 /*
  * Tool
  * 
@@ -620,10 +633,16 @@ void UPatternTool::InitializeElements()
 		if (UStaticMeshComponent* StaticMeshComp = Cast<UStaticMeshComponent>(Element.SourceComponent))
 		{
 			Element.SourceStaticMesh = StaticMeshComp->GetStaticMesh();
+			Element.LocalBounds = (FAxisAlignedBox3d)Element.SourceStaticMesh->GetBounds().GetBox();
+			Element.PatternBounds = Element.LocalBounds;
 		}
 		else if (UDynamicMeshComponent* DynamicMeshComp = Cast<UDynamicMeshComponent>(Element.SourceComponent))
 		{
 			Element.SourceDynamicMesh = DynamicMeshComp->GetDynamicMesh();
+			Element.SourceDynamicMesh->ProcessMesh([&](const FDynamicMesh3& Mesh) {
+				Element.LocalBounds = Mesh.GetBounds(true);
+			});
+			Element.PatternBounds = Element.LocalBounds;
 		}
 		else
 		{
@@ -750,6 +769,7 @@ void UPatternTool::GetPatternTransforms_Linear(TArray<UE::Geometry::FTransformSR
 {
 	FLinearPatternGenerator Generator;
 	InitializeGenerator(Generator, this);
+	Generator.Dimensions = this->Elements[0].PatternBounds;
 
 	double ExtentX = LinearSettings->Extent;
 
@@ -781,6 +801,7 @@ void UPatternTool::GetPatternTransforms_Grid(TArray<UE::Geometry::FTransformSRT3
 {
 	FLinearPatternGenerator Generator;
 	InitializeGenerator(Generator, this);
+	Generator.Dimensions = this->Elements[0].PatternBounds;
 
 	double ExtentX = GridSettings->ExtentX;
 	double ExtentY = GridSettings->ExtentY;
@@ -831,6 +852,7 @@ void UPatternTool::GetPatternTransforms_Radial(TArray<UE::Geometry::FTransformSR
 {
 	FRadialPatternGenerator Generator;
 	InitializeGenerator(Generator, this);
+	Generator.Dimensions = this->Elements[0].PatternBounds;
 
 	Generator.CenterFrame = FFrame3d();
 
diff --git a/Engine/Plugins/Experimental/MeshModelingToolsetExp/Source/MeshModelingToolsExp/Public/PatternTool.h b/Engine/Plugins/Experimental/MeshModelingToolsetExp/Source/MeshModelingToolsExp/Public/PatternTool.h
--- a/Engine/Plugins/Experimental/MeshModelingToolsetExp/Source/MeshModelingToolsExp/Public/PatternTool.h
+++ b/Engine/Plugins/Experimental/MeshModelingToolsetExp/Source/MeshModelingToolsExp/Public/PatternTool.h
@@ -83,7 +83,10 @@ enum class EPatternToolAxisSpacingMode : uint8
 	ByCount = 0,
 
 	/** Place Pattern Elements at regular increments along the Pattern Geometry (on-center) */
-	StepSize = 1
+	StepSize = 1,
+
+	/** Pack in as many Pattern Elements as fits in the available space */
+	Packed = 2
 };
 
 
@@ -449,7 +452,8 @@ protected:
 		UDynamicMesh* SourceDynamicMesh = nullptr;
 		UStaticMesh* SourceStaticMesh = nullptr;
 
-		UE::Geometry::FAxisAlignedBox3d Bounds = UE::Geometry::FAxisAlignedBox3d::Empty();
+		UE::Geometry::FAxisAlignedBox3d LocalBounds = UE::Geometry::FAxisAlignedBox3d::Empty();
+		UE::Geometry::FAxisAlignedBox3d PatternBounds = UE::Geometry::FAxisAlignedBox3d::Empty();
 	};
 	TArray<FPatternElement> Elements;
 
