Adding selective tessellation option to the Displace mesh tool to allow for tessellating triangles that belong to a certain material.

#rnx

diff --git a/Engine/Plugins/Experimental/MeshModelingToolsetExp/Source/MeshModelingToolsExp/Private/DisplaceMeshTool.cpp b/Engine/Plugins/Experimental/MeshModelingToolsetExp/Source/MeshModelingToolsExp/Private/DisplaceMeshTool.cpp
--- a/Engine/Plugins/Experimental/MeshModelingToolsetExp/Source/MeshModelingToolsExp/Private/DisplaceMeshTool.cpp
+++ b/Engine/Plugins/Experimental/MeshModelingToolsetExp/Source/MeshModelingToolsExp/Private/DisplaceMeshTool.cpp
@@ -18,7 +18,8 @@
 #include "ModelingToolTargetUtil.h"
 #include "Operations/PNTriangles.h"
 #include "Operations/UniformTessellate.h"
-
+#include "Operations/AdaptiveTessellate.h"
+#include "Materials/MaterialInterface.h"
 
 // needed to disable normals recalculation on the underlying asset
 #include "AssetUtils/MeshDescriptionUtil.h"
@@ -211,18 +212,32 @@ namespace DisplaceMeshToolLocals{
 
 	}
 
+	/** A collection of subdivision parameters. */
+	struct FSubdivideParameters
+	{
+		int SubdivisionsCount = -1;
+		TSharedPtr<FIndexedWeightMap, ESPMode::ThreadSafe> WeightMap = nullptr;
+
+		// Optional selection parameters
+		TOptional<EDisplaceMeshToolTriangleSelectionType> SelectionType;
+		
+		// Material selection variables (SelectionType == EDisplaceMeshToolTriangleSelectionType::Material)
+		TOptional<int> ActiveMaterialID;
+	};
+
 	class FSubdivideMeshOp : public FDynamicMeshOperator
 	{
 	public:
-		FSubdivideMeshOp(const FDynamicMesh3& SourceMesh, EDisplaceMeshToolSubdivisionType SubdivisionTypeIn, int SubdivisionsCountIn, TSharedPtr<FIndexedWeightMap, ESPMode::ThreadSafe> WeightMap);
+		FSubdivideMeshOp(const FDynamicMesh3& SourceMesh, const FSubdivideParameters& InParameters, EDisplaceMeshToolSubdivisionType InSubdivisionType);
 		void CalculateResult(FProgressCancel* Progress) final;
 	private:
+		FSubdivideParameters Parameters;
 		EDisplaceMeshToolSubdivisionType SubdivisionType;
-		int SubdivisionsCount;
 	};
 
-	FSubdivideMeshOp::FSubdivideMeshOp(const FDynamicMesh3& SourceMesh, EDisplaceMeshToolSubdivisionType SubdivisionTypeIn, int SubdivisionsCountIn, TSharedPtr<FIndexedWeightMap, ESPMode::ThreadSafe> WeightMap)
-		: SubdivisionType(SubdivisionTypeIn), SubdivisionsCount(SubdivisionsCountIn)
+	FSubdivideMeshOp::FSubdivideMeshOp(const FDynamicMesh3& SourceMesh, const FSubdivideParameters& InParameters, EDisplaceMeshToolSubdivisionType InSubdivisionType)
+	: 
+	Parameters(InParameters), SubdivisionType(InSubdivisionType)
 	{
 		ResultMesh->Copy(SourceMesh);
 
@@ -230,11 +245,11 @@ namespace DisplaceMeshToolLocals{
 		// we could (for exmaple) speculatively compute another weightmap, or store previous weightmap values there, to support
 		// fast switching between two...
 		ResultMesh->EnableVertexUVs(FVector2f::Zero());
-		if (WeightMap != nullptr)
+		if (Parameters.WeightMap != nullptr)
 		{
 			for (int32 vid : ResultMesh->VertexIndicesItr())
 			{
-				ResultMesh->SetVertexUV(vid, FVector2f(WeightMap->GetValue(vid), 0));
+				ResultMesh->SetVertexUV(vid, FVector2f(Parameters.WeightMap->GetValue(vid), 0));
 			}
 		}
 		else
@@ -248,15 +263,48 @@ namespace DisplaceMeshToolLocals{
 
 	void FSubdivideMeshOp::CalculateResult(FProgressCancel* ProgressCancel)
 	{
+		int SubdivisionsCount = Parameters.SubdivisionsCount;
 		if (SubdivisionType == EDisplaceMeshToolSubdivisionType::Flat) 
 		{
-			FUniformTessellate Tessellator(ResultMesh.Get());
-			Tessellator.Progress = ProgressCancel;
-			Tessellator.TessellationNum = SubdivisionsCount;
-						
-			if (Tessellator.Validate() == EOperationValidationResult::Ok) 
+			if (Parameters.SelectionType.IsSet() == false)
+			{
+				return;
+			}
+
+			if (Parameters.SelectionType.GetValue() == EDisplaceMeshToolTriangleSelectionType::None) 
+			{
+				FUniformTessellate Tessellator(ResultMesh.Get());
+				Tessellator.Progress = ProgressCancel;
+				Tessellator.TessellationNum = SubdivisionsCount;
+							
+				if (Tessellator.Validate() == EOperationValidationResult::Ok) 
+				{
+					Tessellator.Compute();
+				}
+			}
+			else if (Parameters.SelectionType.GetValue() == EDisplaceMeshToolTriangleSelectionType::Material)
+			{
+				if (Parameters.ActiveMaterialID.IsSet())
+				{
+					TUniquePtr<FTessellationPattern> Pattern = FAdaptiveTessellate::CreateConcentricRingsPatternFromMaterial(ResultMesh.Get(), SubdivisionsCount, Parameters.ActiveMaterialID.GetValue());;
+					
+					FDynamicMesh3 OutMesh;
+					FAdaptiveTessellate Tessellator(ResultMesh.Get(), &OutMesh);
+					
+					Tessellator.SetPattern(Pattern.Get());
+					
+					if (ensureMsgf(Tessellator.Validate() == EOperationValidationResult::Ok, TEXT("The tessellator parameters are invalid.")))
+					{
+						if (ensureMsgf(Tessellator.Compute(), TEXT("Failed to tessellate the selected triangles.")))
+						{
+							*ResultMesh = MoveTemp(OutMesh);
+						}
+					}
+				}
+			}
+			else
 			{
-				Tessellator.Compute();
+				checkNoEntry(); // Unsupported selection type
 			}
 		}
 		else if (SubdivisionType == EDisplaceMeshToolSubdivisionType::PNTriangles) 
@@ -281,10 +329,10 @@ namespace DisplaceMeshToolLocals{
 	{
 	public:
 		FSubdivideMeshOpFactory(FDynamicMesh3& SourceMeshIn,
-			EDisplaceMeshToolSubdivisionType SubdivisionTypeIn,
-			int SubdivisionsCountIn,
-			TSharedPtr<FIndexedWeightMap, ESPMode::ThreadSafe> WeightMapIn)
-			: SourceMesh(SourceMeshIn), SubdivisionType(SubdivisionTypeIn), SubdivisionsCount(SubdivisionsCountIn), WeightMap(WeightMapIn)
+								const FSubdivideParameters& InParameters,
+								EDisplaceMeshToolSubdivisionType SubdivisionTypeIn)
+		: 
+		SourceMesh(SourceMeshIn), Parameters(InParameters), SubdivisionType(SubdivisionTypeIn)
 		{
 		}
 
@@ -296,15 +344,25 @@ namespace DisplaceMeshToolLocals{
 
 		void SetWeightMap(TSharedPtr<FIndexedWeightMap, ESPMode::ThreadSafe> WeightMapIn);
 
+		void SetSelectionType(EDisplaceMeshToolTriangleSelectionType SelectionTypeIn) 
+		{
+			Parameters.SelectionType = SelectionTypeIn;
+		}
+
+		// Materials
+		void SetActiveMaterialID(int ActiveMaterialID) 
+		{
+			Parameters.ActiveMaterialID = ActiveMaterialID;
+		}
+		
 		TUniquePtr<FDynamicMeshOperator> MakeNewOperator() final
 		{
-			return MakeUnique<FSubdivideMeshOp>(SourceMesh, SubdivisionType, SubdivisionsCount, WeightMap);
+			return MakeUnique<FSubdivideMeshOp>(SourceMesh, Parameters, SubdivisionType);
 		}
 	private:
 		const FDynamicMesh3& SourceMesh;
+		FSubdivideParameters Parameters;
 		EDisplaceMeshToolSubdivisionType SubdivisionType;
-		int SubdivisionsCount;
-		TSharedPtr<FIndexedWeightMap, ESPMode::ThreadSafe> WeightMap;
 	};
 
 	void FSubdivideMeshOpFactory::SetSubdivisionType(EDisplaceMeshToolSubdivisionType SubdivisionTypeIn) 
@@ -319,17 +377,17 @@ namespace DisplaceMeshToolLocals{
 
 	void FSubdivideMeshOpFactory::SetSubdivisionsCount(int SubdivisionsCountIn)
 	{
-		SubdivisionsCount = SubdivisionsCountIn;
+		Parameters.SubdivisionsCount = SubdivisionsCountIn;
 	}
 
 	int FSubdivideMeshOpFactory::GetSubdivisionsCount() const
 	{
-		return SubdivisionsCount;
+		return Parameters.SubdivisionsCount;
 	}
 
 	void FSubdivideMeshOpFactory::SetWeightMap(TSharedPtr<FIndexedWeightMap, ESPMode::ThreadSafe> WeightMapIn)
 	{
-		WeightMap = WeightMapIn;
+		Parameters.WeightMap = WeightMapIn;
 	}
 
 	// A collection of parameters to avoid having excess function parameters
@@ -381,9 +439,10 @@ namespace DisplaceMeshToolLocals{
 	FDisplaceMeshOp::FDisplaceMeshOp(TSharedPtr<FDynamicMesh3, ESPMode::ThreadSafe> SourceMeshIn,
 									 const DisplaceMeshParameters& DisplaceParametersIn,
 									 EDisplaceMeshToolDisplaceType DisplacementTypeIn)
-		: SourceMesh(MoveTemp(SourceMeshIn)), 
-		  Parameters(DisplaceParametersIn), 
-		  DisplacementType(DisplacementTypeIn)
+	: 
+	SourceMesh(MoveTemp(SourceMeshIn)), 
+	Parameters(DisplaceParametersIn), 
+	DisplacementType(DisplacementTypeIn)
 	{
 	}
 
@@ -524,9 +583,10 @@ namespace DisplaceMeshToolLocals{
 	{
 	public:
 		FDisplaceMeshOpFactory(TSharedPtr<FDynamicMesh3, ESPMode::ThreadSafe>& SourceMeshIn,
-			const DisplaceMeshParameters& DisplaceParametersIn,
-			EDisplaceMeshToolDisplaceType DisplacementTypeIn )
-			: SourceMesh(SourceMeshIn)
+							   const DisplaceMeshParameters& DisplaceParametersIn,
+							   EDisplaceMeshToolDisplaceType DisplacementTypeIn )
+		: 
+		SourceMesh(SourceMeshIn)
 		{
 			SetIntensity(DisplaceParametersIn.DisplaceIntensity);
 			SetRandomSeed(DisplaceParametersIn.RandomSeed);
@@ -549,6 +609,7 @@ namespace DisplaceMeshToolLocals{
 
 			Parameters.AdjustmentCurve = DisplaceParametersIn.AdjustmentCurve;
 		}
+		
 		void SetIntensity(float IntensityIn);
 		void SetRandomSeed(int RandomSeedIn);
 		void SetDisplacementMap(UTexture2D* DisplacementMapIn, int32 ChannelIn);
@@ -744,6 +805,8 @@ void UDisplaceMeshTool::Setup()
 	TextureMapProperties->RestoreProperties(this);
 	SineWaveProperties = NewObject<UDisplaceMeshSineWaveProperties>();
 	SineWaveProperties->RestoreProperties(this);
+	SelectiveTessellationProperties = NewObject<USelectiveTessellationProperties>();
+	SelectiveTessellationProperties->RestoreProperties(this);
 
 	if (TextureMapProperties->DisplacementMap != nullptr && TextureMapProperties->DisplacementMap->IsValidLowLevel() == false)
 	{
@@ -813,6 +876,25 @@ void UDisplaceMeshTool::Setup()
 		
 	// hide input StaticMeshComponent
 	UE::ToolTarget::HideSourceObject(Target);
+	
+	// Fetch the list of the materials if available
+	SelectiveTessellationProperties->MaterialIDList.Reset();
+	for (int32 Idx = 0; Idx < MaterialSet.Materials.Num(); ++Idx)
+	{
+		UMaterialInterface* Mat = MaterialSet.Materials[Idx];
+		FString MatName = (Mat != nullptr) ? Mat->GetName() : "(none)";
+		FString UseName = FString::Printf(TEXT("[%d] %s"), Idx, *MatName);
+		SelectiveTessellationProperties->MaterialIDList.Add(UseName);
+	}
+
+	if (SelectiveTessellationProperties->MaterialIDList.IsEmpty())
+	{
+		SelectiveTessellationProperties->ActiveMaterial = NAME_None;
+	}
+	else if (SelectiveTessellationProperties->MaterialIDList.Contains(SelectiveTessellationProperties->ActiveMaterial.ToString()) == false) 
+	{
+		SelectiveTessellationProperties->ActiveMaterial = NAME_None;
+	}
 
 	// initialize our properties
 	ToolPropertyObjects.Add(this);
@@ -829,6 +911,9 @@ void UDisplaceMeshTool::Setup()
 	AddToolPropertySource(SineWaveProperties);
 	SetToolPropertySourceEnabled(SineWaveProperties, CommonProperties->DisplacementType == EDisplaceMeshToolDisplaceType::SineWave);
 
+	AddToolPropertySource(SelectiveTessellationProperties);
+	SetToolPropertySourceEnabled(SelectiveTessellationProperties, CommonProperties->SubdivisionType == EDisplaceMeshToolSubdivisionType::Flat);
+
 	AddToolPropertySource(NoiseProperties);
 	SetToolPropertySourceEnabled(NoiseProperties, CommonProperties->DisplacementType == EDisplaceMeshToolDisplaceType::PerlinNoise);
 
@@ -841,8 +926,30 @@ void UDisplaceMeshTool::Setup()
 										SetToolPropertySourceEnabled(TextureMapProperties, (NewType == EDisplaceMeshToolDisplaceType::DisplacementMap));
 									} );
 
+	CommonProperties->WatchProperty(CommonProperties->SubdivisionType,
+									[this](EDisplaceMeshToolSubdivisionType NewType)
+									{
+										SetToolPropertySourceEnabled(SelectiveTessellationProperties, (NewType == EDisplaceMeshToolSubdivisionType::Flat));
+									} );
+
+
 	ValidateSubdivisions();
-	Subdivider = MakeUnique<FSubdivideMeshOpFactory>(OriginalMesh, CommonProperties->SubdivisionType, CommonProperties->Subdivisions, ActiveWeightMap);
+
+	FSubdivideParameters SubParameters;
+	SubParameters.SubdivisionsCount = CommonProperties->Subdivisions;
+	SubParameters.WeightMap = ActiveWeightMap;
+	SubParameters.SelectionType = SelectiveTessellationProperties->SelectionType;
+	
+	if (SelectiveTessellationProperties->ActiveMaterial.IsNone() == false) 
+	{
+		int32 Index = SelectiveTessellationProperties->MaterialIDList.Find(SelectiveTessellationProperties->ActiveMaterial.ToString());
+		if (ensure(Index != INDEX_NONE)) 
+		{
+			SubParameters.ActiveMaterialID = Index;
+		}
+	}
+
+	Subdivider = MakeUnique<FSubdivideMeshOpFactory>(OriginalMesh, SubParameters, CommonProperties->SubdivisionType);
 
 	StartComputation();
 
@@ -862,6 +969,8 @@ void UDisplaceMeshTool::OnShutdown(EToolShutdownType ShutdownType)
 	DirectionalFilterProperties->SaveProperties(this);
 	SineWaveProperties->SaveProperties(this);
 	TextureMapProperties->SaveProperties(this);
+	SelectiveTessellationProperties->SaveProperties(this);
+	
 
 	if (DynamicMeshComponent != nullptr)
 	{
@@ -1079,7 +1188,21 @@ void UDisplaceMeshTool::OnPropertyModified(UObject* PropertySet, FProperty* Prop
 			DisplacerDownCast->SetSineWaveDirection(SineWaveProperties->SineWaveDirection);
 			DisplacerDownCast->SetDisplacementMapUVAdjustment(FVector2f(TextureMapProperties->UVScale), FVector2f(TextureMapProperties->UVOffset));	// LWC_TODO: Precision loss
 		}
-
+		else if (PropName == GET_MEMBER_NAME_CHECKED(USelectiveTessellationProperties, SelectionType))
+		{
+			SubdividerDownCast->SetSelectionType(SelectiveTessellationProperties->SelectionType);
+			bNeedsSubdivided = true;
+		}
+		else if (PropName == GET_MEMBER_NAME_CHECKED(USelectiveTessellationProperties, ActiveMaterial)) 
+		{
+			int32 Index = SelectiveTessellationProperties->MaterialIDList.Find(SelectiveTessellationProperties->ActiveMaterial.ToString());
+			if (ensure(Index != INDEX_NONE)) 
+			{
+				SubdividerDownCast->SetActiveMaterialID(Index);
+				bNeedsSubdivided = true;
+			}
+		}
+		
 		StartComputation();
 	}
 }
diff --git a/Engine/Plugins/Experimental/MeshModelingToolsetExp/Source/MeshModelingToolsExp/Public/DisplaceMeshTool.h b/Engine/Plugins/Experimental/MeshModelingToolsetExp/Source/MeshModelingToolsExp/Public/DisplaceMeshTool.h
--- a/Engine/Plugins/Experimental/MeshModelingToolsetExp/Source/MeshModelingToolsExp/Public/DisplaceMeshTool.h
+++ b/Engine/Plugins/Experimental/MeshModelingToolsetExp/Source/MeshModelingToolsExp/Public/DisplaceMeshTool.h
@@ -53,7 +53,17 @@ enum class EDisplaceMeshToolSubdivisionType : uint8
 		retriangulated into a number of small subtriangles. The geometry of a PN triangle is defined as one cubic Bezier 
 		patch using control points. The patch matches the point and normal information at the vertices of the original 
 		flat triangle.*/
-	PNTriangles UMETA(DisplayName = "PN Triangles"),
+	PNTriangles UMETA(DisplayName = "PN Triangles")
+};
+
+UENUM()
+enum class EDisplaceMeshToolTriangleSelectionType : uint8
+{	
+	/** Tessellate the whole mesh */
+	None UMETA(DisplayName = "None"),
+
+	/** Tessellate only triangles assigned to the chosen material */
+	Material UMETA(DisplayName = "Material")
 };
 
 UENUM() 
@@ -115,6 +125,28 @@ public:
 	bool bDisableSizeWarning = false;
 };
 
+/** PropertySet for properties affecting the selective tessellation. */
+UCLASS()
+class MESHMODELINGTOOLSEXP_API USelectiveTessellationProperties : public UInteractiveToolPropertySet
+{
+	GENERATED_BODY()
+			
+public:
+
+	/** Optionally, restrict tessellation to a subset of the triangles. */
+	UPROPERTY(EditAnywhere, Category = SelectiveTessellationOptions)
+	EDisplaceMeshToolTriangleSelectionType SelectionType = EDisplaceMeshToolTriangleSelectionType::None;
+
+	/** Name of the selected material. */
+	UPROPERTY(EditAnywhere, Category = SelectiveTessellationOptions, meta = (EditCondition = "SelectionType == EDisplaceMeshToolTriangleSelectionType::Material", EditConditionHides, GetOptions = GetMaterialIDsFunc))
+	FName ActiveMaterial;
+
+	UFUNCTION()
+	TArray<FString> GetMaterialIDsFunc() { return MaterialIDList; }
+
+	UPROPERTY(meta = (TransientToolProperty))
+	TArray<FString> MaterialIDList;	
+};
 
 
 /** PropertySet for properties affecting the Image Map displacement type. */
@@ -302,6 +334,10 @@ public:
 	UPROPERTY()
 	TObjectPtr<UDisplaceMeshSineWaveProperties> SineWaveProperties;
 
+	/** Selective tessellation properties */
+	UPROPERTY()
+	TObjectPtr<USelectiveTessellationProperties> SelectiveTessellationProperties;
+
 	/** Contrast Curve we are actively listening to */
 	UPROPERTY()
 	TObjectPtr<UCurveFloat> ActiveContrastCurveTarget = nullptr;
diff --git a/Engine/Plugins/Runtime/GeometryProcessing/Source/DynamicMesh/Private/Operations/AdaptiveTessellate.cpp b/Engine/Plugins/Runtime/GeometryProcessing/Source/DynamicMesh/Private/Operations/AdaptiveTessellate.cpp
--- a/Engine/Plugins/Runtime/GeometryProcessing/Source/DynamicMesh/Private/Operations/AdaptiveTessellate.cpp
+++ b/Engine/Plugins/Runtime/GeometryProcessing/Source/DynamicMesh/Private/Operations/AdaptiveTessellate.cpp
@@ -1756,7 +1756,7 @@ TUniquePtr<FTessellationPattern> FAdaptiveTessellate::CreateConcentricRingsTesse
 
 
 TUniquePtr<FTessellationPattern> FAdaptiveTessellate::CreateConcentricRingsTessellationPattern(const TFunctionRef<int(const int EdgeID)> InEdgeFunc, 
-                                                               					    const TFunctionRef<int(const int TriangleID)> InTriFunc) 
+                                                               					    		   const TFunctionRef<int(const int TriangleID)> InTriFunc) 
 {
 
 	// TODO: not implemented yet
@@ -1765,7 +1765,98 @@ TUniquePtr<FTessellationPattern> FAdaptiveTessellate::CreateConcentricRingsTesse
 }
 
 
+TUniquePtr<FTessellationPattern> FAdaptiveTessellate::CreateConcentricRingsPatternFromTriangleGroup(const FDynamicMesh3* InMesh,
+																					 			    const int InTessellationLevel,
+                                                                                     			    const int InPolygroupID) 
+{
+	if (InMesh->HasTriangleGroups() == false)
+	{
+		return nullptr;
+	}
+
+	TArray<int> TriangleList;
+	for (const int TID : InMesh->TriangleIndicesItr()) 
+	{
+		const int PolygroupID = InMesh->GetTriangleGroup(TID);
+		
+		if (PolygroupID == InPolygroupID) 
+		{
+			TriangleList.Add(TID);
+		}
+	}
 
+	return FAdaptiveTessellate::CreateConcentricRingsTessellationPattern(InMesh, InTessellationLevel, TriangleList);
+}
+
+
+TUniquePtr<FTessellationPattern> FAdaptiveTessellate::CreateConcentricRingsPatternFromPolyGroup(const FDynamicMesh3* InMesh,
+                                                                           						const int InTessellationLevel,
+                                                                           						const FString& InLayerName,
+                                                                           						const int InPolygroupID) 
+{
+	if (InMesh->HasAttributes() == false)
+	{
+		return nullptr;
+	}
+
+	const FDynamicMeshPolygroupAttribute* PolygroupAttribute = nullptr;
+
+	for (int Idx = 0; Idx < InMesh->Attributes()->NumPolygroupLayers(); ++Idx) 
+	{
+		if (InMesh->Attributes()->GetPolygroupLayer(Idx)->GetName().ToString() == InLayerName) 
+		{
+			PolygroupAttribute = InMesh->Attributes()->GetPolygroupLayer(Idx);
+			break;
+		}
+	}
+
+	if (PolygroupAttribute == nullptr) 
+	{
+		return nullptr;
+	}
+
+	TArray<int> TriangleList;
+	for (const int TID : InMesh->TriangleIndicesItr()) 
+	{
+		const int TrianglePolygroupID = PolygroupAttribute->GetValue(TID);
+		
+		if (TrianglePolygroupID == InPolygroupID) 
+		{
+			TriangleList.Add(TID);
+		}
+	}
+
+	return FAdaptiveTessellate::CreateConcentricRingsTessellationPattern(InMesh, InTessellationLevel, TriangleList);
+}
+
+
+TUniquePtr<FTessellationPattern> FAdaptiveTessellate::CreateConcentricRingsPatternFromMaterial(const FDynamicMesh3* InMesh,
+																					 		   const int InTessellationLevel,
+                                                                                     		   const int MaterialID) 
+{
+	if (InMesh->HasAttributes() == false)
+	{
+		return nullptr;
+	}
+
+	if (InMesh->Attributes()->HasMaterialID() == false)
+	{
+		return nullptr;
+	}
+
+	TArray<int> TriangleList;
+	for (const int TID : InMesh->TriangleIndicesItr()) 
+	{
+		const int TriangleMaterialID = InMesh->Attributes()->GetMaterialID()->GetValue(TID);
+		
+		if (TriangleMaterialID == MaterialID) 
+		{
+			TriangleList.Add(TID);
+		}
+	}
+
+	return FAdaptiveTessellate::CreateConcentricRingsTessellationPattern(InMesh, InTessellationLevel, TriangleList);
+}
 
 //
 // Inner uniform pattern
diff --git a/Engine/Plugins/Runtime/GeometryProcessing/Source/DynamicMesh/Public/Operations/AdaptiveTessellate.h b/Engine/Plugins/Runtime/GeometryProcessing/Source/DynamicMesh/Public/Operations/AdaptiveTessellate.h
--- a/Engine/Plugins/Runtime/GeometryProcessing/Source/DynamicMesh/Public/Operations/AdaptiveTessellate.h
+++ b/Engine/Plugins/Runtime/GeometryProcessing/Source/DynamicMesh/Public/Operations/AdaptiveTessellate.h
@@ -373,7 +373,35 @@ public:
                                                                                      const TFunctionRef<int(const int TriangleID)> InTriFunc);
 
     
-    //TODO: add methods for automatic selection of triangles based on Material or Polygroup IDs.
+    /** 
+     * Tessellate only triangles that belong to a given triangle group id. Triangle groups are defined by the 
+     * FDynamicMesh3::GetTriangleGroup().
+     */
+    static TUniquePtr<FTessellationPattern> CreateConcentricRingsPatternFromTriangleGroup(const FDynamicMesh3* InMesh,
+                                                                                          const int InTessellationLevel,
+                                                                                          const int InPolygroupID);
+
+
+    /** 
+     * Tessellate only triangles that belong to a given polygroup id within the polygroup layer attribute 
+     * (FDynamicMeshPolygroupAttribute) specified by its name. 
+     * 
+     * @return nullptr if InMesh has no attributes or if no polygroup layer attribute with the given name or id exists.
+     */
+    static TUniquePtr<FTessellationPattern> CreateConcentricRingsPatternFromPolyGroup(const FDynamicMesh3* InMesh,
+                                                                                      const int InTessellationLevel,
+                                                                                      const FString& InLayerName,
+                                                                                      const int InPolygroupID);
+
+
+    /** 
+     * Tessellate only triangles that belong to a given material (FDynamicMeshMaterialAttribute). 
+     * 
+     * @return nullptr if InMesh has no attributes, no material attribute or if no material attribute with the given id exists.
+     */
+	static TUniquePtr<FTessellationPattern> CreateConcentricRingsPatternFromMaterial(const FDynamicMesh3* InMesh,
+																				 	 const int InTessellationLevel,
+																					 const int MaterialID);
 
 
 
