MeshModeling: Pivot tool - fix undo.

Re-ordered the tool target commit mesh and the collision geometry update to avoid a PostEditUndo crash related to physics collision data.

diff --git a/Engine/Plugins/Experimental/MeshModelingToolsetExp/Source/MeshModelingToolsExp/Private/EditPivotTool.cpp b/Engine/Plugins/Experimental/MeshModelingToolsetExp/Source/MeshModelingToolsExp/Private/EditPivotTool.cpp
--- a/Engine/Plugins/Experimental/MeshModelingToolsetExp/Source/MeshModelingToolsExp/Private/EditPivotTool.cpp
+++ b/Engine/Plugins/Experimental/MeshModelingToolsetExp/Source/MeshModelingToolsExp/Private/EditPivotTool.cpp
@@ -409,7 +409,8 @@ void UEditPivotTool::UpdateAssets(const FFrame3d& NewPivotWorldFrame)
 	}
 	for (int32 ComponentIdx = 0; ComponentIdx < Targets.Num(); ComponentIdx++)
 	{
-		UPrimitiveComponent* Component = UE::ToolTarget::GetTargetComponent(Targets[ComponentIdx]);
+		UToolTarget* Target = Targets[ComponentIdx];
+		UPrimitiveComponent* Component = UE::ToolTarget::GetTargetComponent(Target);
 		Component->Modify();
 
 		UInstancedStaticMeshComponent* InstancedComponent = Cast<UInstancedStaticMeshComponent>(Component);
@@ -488,6 +489,21 @@ void UEditPivotTool::UpdateAssets(const FFrame3d& NewPivotWorldFrame)
 				// else do nothing -- scale is not invertible and must be baked
 			}
 
+			
+
+			FMeshDescription SourceMesh(UE::ToolTarget::GetMeshDescriptionCopy(Target));
+			FMeshDescriptionEditableTriangleMeshAdapter EditableMeshDescAdapter(&SourceMesh);
+			
+			MeshAdapterTransforms::ApplyTransform(EditableMeshDescAdapter, ToBake);
+			
+			if (bNeedSeparateScale)
+			{
+				MeshAdapterTransforms::ApplyTransform(EditableMeshDescAdapter, SeparateBakeScale);
+			}
+			
+			// todo: support vertex-only update
+			UE::ToolTarget::CommitMeshDescriptionUpdate(Target, &SourceMesh);
+
 			// transform simple collision geometry
 			if (UE::Geometry::ComponentTypeSupportsCollision(Component))
 			{
@@ -498,16 +514,6 @@ void UEditPivotTool::UpdateAssets(const FFrame3d& NewPivotWorldFrame)
 				}
 			}
 
-			FMeshDescription SourceMesh(UE::ToolTarget::GetMeshDescriptionCopy(Targets[ComponentIdx]));
-			FMeshDescriptionEditableTriangleMeshAdapter EditableMeshDescAdapter(&SourceMesh);
-			MeshAdapterTransforms::ApplyTransform(EditableMeshDescAdapter, ToBake);
-			if (bNeedSeparateScale)
-			{
-				MeshAdapterTransforms::ApplyTransform(EditableMeshDescAdapter, SeparateBakeScale);
-			}
-			// todo: support vertex-only update
-			UE::ToolTarget::CommitMeshDescriptionUpdate(Targets[ComponentIdx], &SourceMesh);
-
 			Component->SetWorldTransform(ScaledNewWorldTransform);
 		}
 		else
@@ -518,11 +524,11 @@ void UEditPivotTool::UpdateAssets(const FFrame3d& NewPivotWorldFrame)
 			Component->SetWorldTransform(Baked.Inverse() * OriginalTransforms[ComponentIdx]);
 		}
 
-		AActor* OwnerActor = UE::ToolTarget::GetTargetActor(Targets[ComponentIdx]);
-		if (OwnerActor)
+		AActor* TargetActor = UE::ToolTarget::GetTargetActor(Target);
+		if (TargetActor)
 		{
-			OwnerActor->MarkComponentsRenderStateDirty();
-			OwnerActor->UpdateComponentTransforms();
+			TargetActor->MarkComponentsRenderStateDirty();
+			TargetActor->UpdateComponentTransforms();
 		}
 	}
 
