Only displace those vertices that belong to triangles affected by the tessellation.

#rnx

diff --git a/Engine/Plugins/Experimental/MeshModelingToolsetExp/Source/MeshModelingToolsExp/Private/DisplaceMeshTool.cpp b/Engine/Plugins/Experimental/MeshModelingToolsetExp/Source/MeshModelingToolsExp/Private/DisplaceMeshTool.cpp
--- a/Engine/Plugins/Experimental/MeshModelingToolsetExp/Source/MeshModelingToolsExp/Private/DisplaceMeshTool.cpp
+++ b/Engine/Plugins/Experimental/MeshModelingToolsetExp/Source/MeshModelingToolsExp/Private/DisplaceMeshTool.cpp
@@ -64,20 +64,42 @@ namespace DisplaceMeshToolLocals{
 			const TArray<FVector3d>& Positions,
 			const FMeshNormals& Normals,
 			TArray<FVector3d>& DisplacedPositions,
-			DisplaceFunc Displace)
+			TArray<int>* VerticesToDisplace,
+			DisplaceFunc Displace,
+			bool bUseParallel = true)
 		{
 			ensure(Positions.Num() == Normals.GetNormals().Num());
 			ensure(Positions.Num() == DisplacedPositions.Num());
 			ensure(Mesh.VertexCount() == Positions.Num());
 
-			int32 NumVertices = Mesh.MaxVertexID();
-			ParallelFor(NumVertices, [&](int32 vid)
-			{
-				if (Mesh.IsVertex(vid))
+			if (VerticesToDisplace) 
+			{	
+				// Copy over the original positions first since most likely some of the vertices won't be displaced
+				for (int VID : Mesh.VertexIndicesItr()) 
 				{
-					DisplacedPositions[vid] = Displace(vid, Positions[vid], Normals[vid]);
+					DisplacedPositions[VID] = Positions[VID];
 				}
-			});
+
+				ParallelFor(VerticesToDisplace->Num(), [&](int32 Idx)
+				{	
+					const int VID = (*VerticesToDisplace)[Idx];
+					if (ensure(Mesh.IsVertex(VID))) // if this test fails then VerticesToDisplace is invalid
+					{
+						DisplacedPositions[VID] = Displace(VID, Positions[VID], Normals[VID]);
+					}
+				 }, bUseParallel ? EParallelForFlags::None : EParallelForFlags::ForceSingleThread);
+			}
+			else 
+			{	
+				// Displace all vertices
+				ParallelFor(Mesh.MaxVertexID(), [&](int32 VID)
+				{
+					if (Mesh.IsVertex(VID))
+					{
+						DisplacedPositions[VID] = Displace(VID, Positions[VID], Normals[VID]);
+					}
+				}, bUseParallel ? EParallelForFlags::None : EParallelForFlags::ForceSingleThread);
+			}
 		}
 
 
@@ -85,9 +107,10 @@ namespace DisplaceMeshToolLocals{
 			const TArray<FVector3d>& Positions, 
 			const FMeshNormals& Normals, 
 			TFunctionRef<float(int32, const FVector3d&, const FVector3d&)> IntensityFunc,
-			TArray<FVector3d>& DisplacedPositions)
+			TArray<FVector3d>& DisplacedPositions,
+			TArray<int>* VerticesToDisplace = nullptr)
 		{
-			ParallelDisplace(Mesh, Positions, Normals, DisplacedPositions,
+			ParallelDisplace(Mesh, Positions, Normals, DisplacedPositions, VerticesToDisplace,
 				[&](int32 vid, const FVector3d& Position, const FVector3d& Normal)
 			{
 				double Intensity = IntensityFunc(vid, Position, Normal);
@@ -101,15 +124,20 @@ namespace DisplaceMeshToolLocals{
 			const FMeshNormals& Normals,
 			TFunctionRef<float(int32, const FVector3d&, const FVector3d&)> IntensityFunc,
 			int RandomSeed, 
-			TArray<FVector3d>& DisplacedPositions)
+			TArray<FVector3d>& DisplacedPositions,
+			TArray<int>* VerticesToDisplace = nullptr)
 		{
 			FMath::SRandInit(RandomSeed);
-			for (int vid : Mesh.VertexIndicesItr())
+			ParallelDisplace(Mesh, Positions, Normals, DisplacedPositions, VerticesToDisplace,
+			[&](int32 vid, const FVector3d& Position, const FVector3d& Normal)
 			{
-				double RandVal = 2.0 * (FMath::SRand() - 0.5);
-				double Intensity = IntensityFunc(vid, Positions[vid], Normals[vid]);
-				DisplacedPositions[vid] = Positions[vid] + (Normals[vid] * RandVal * Intensity);
-			}
+				// FMath::SRand() is not thread safe, hence we pass bUseParallel = false to ParallelDisplace to force 
+				// displacement to be single threaded
+				double RandVal = 2.0 * (FMath::SRand() - 0.5); 
+				double Intensity = IntensityFunc(vid, Position, Normal);
+				return Position + (Normal * RandVal * Intensity);
+			}, false);
+	
 		}
 
 		void PerlinNoise(const FDynamicMesh3& Mesh,
@@ -118,12 +146,13 @@ namespace DisplaceMeshToolLocals{
 			TFunctionRef<float(int32, const FVector3d&, const FVector3d&)> IntensityFunc,
 			const TArray<FPerlinLayerProperties>& PerlinLayerProperties,
 			int RandomSeed,
-			TArray<FVector3d>& DisplacedPositions)
+			TArray<FVector3d>& DisplacedPositions,
+			TArray<int>* VerticesToDisplace = nullptr)
 		{
 			FMath::SRandInit(RandomSeed);
 			const float RandomOffset = 10000.0f * FMath::SRand();
 
-			ParallelDisplace(Mesh, Positions, Normals, DisplacedPositions,
+			ParallelDisplace(Mesh, Positions, Normals, DisplacedPositions, VerticesToDisplace,
 				[&](int32 vid, const FVector3d& Position, const FVector3d& Normal)
 			{
 				// Compute the sum of Perlin noise evaluations for this point
@@ -147,7 +176,9 @@ namespace DisplaceMeshToolLocals{
 			float DisplaceFieldBaseValue = 128.0/255, // value that corresponds to zero displacement
 			FVector2f UVScale = FVector2f(1, 1),
 			FVector2f UVOffset = FVector2f(0,0),
-			FRichCurve* AdjustmentCurve = nullptr)
+			TArray<int>* VerticesToDisplace = nullptr,
+			FRichCurve* AdjustmentCurve = nullptr,
+			bool bUseParallel = true)
 		{
 			const FDynamicMeshUVOverlay* UVOverlay = Mesh.Attributes()->GetUVLayer(0);
 
@@ -155,36 +186,85 @@ namespace DisplaceMeshToolLocals{
 			// but the V direction may be shorter or longer if the texture is not square
 			// (it will be 1/AspectRatio)
 			float VHeight = DisplaceField.Height() * DisplaceField.CellDimensions.Y;
+		
+			// Stores average offset for all elements sharing a vertex
+			TArray<double> OffsetArray;
+			OffsetArray.SetNumZeroed(Mesh.MaxVertexID());
+			
+			auto ComputeOffsetForElement = [&](int32 ElementID) 
+			{	
+				FVector2f UV = UVOverlay->GetElement(ElementID);
+
+				// Adjust UV value and tile it. 
+				// Note that we're effectively stretching the texture to be square before tiling, since this
+				// seems to be what non square textures do by default in UE. If we decide to tile without 
+				// stretching by default someday, we'd do UV - FVector2f(FMath::Floor(UV.X), FMath:Floor(UV.Y/VHeight)*VHeight)
+				// without multiplying by VHeight afterward.
+				UV = UV * UVScale + UVOffset;
+				UV = UV - FVector2f(FMath::Floor(UV.X), FMath::Floor(UV.Y));
+				UV.Y *= VHeight;
+
+				double Offset = DisplaceField.BilinearSampleClamped(UV);
+				if (AdjustmentCurve) 
+				{
+					Offset = AdjustmentCurve->Eval(Offset);
+				}
+
+				Offset -= DisplaceFieldBaseValue;
+
+				return Offset;
+			};
+			
+			// Copy over the original positions if we are only displacing some of them
+			if (VerticesToDisplace) 
+			{
+				for (int VID : Mesh.VertexIndicesItr()) 
+				{
+					DisplacedPositions[VID] = Positions[VID];
+				}
+			}
 
-			for (int tid : Mesh.TriangleIndicesItr())
+			const int NumVertices = VerticesToDisplace ? VerticesToDisplace->Num() : Mesh.MaxVertexID();
+
+			ParallelFor(NumVertices, [&](int32 Idx)
 			{
-				FIndex3i Tri = Mesh.GetTriangle(tid);
-				FIndex3i UVTri = UVOverlay->GetTriangle(tid);
-				for (int j = 0; j < 3; ++j)
+				const int VID = VerticesToDisplace ? (*VerticesToDisplace)[Idx] : Idx;
+
+				if (Mesh.IsVertex(VID) == false) 
 				{
-					int vid = Tri[j];
-					FVector2f UV = UVOverlay->GetElement(UVTri[j]);
-
-					// Adjust UV value and tile it. 
-					// Note that we're effectively stretching the texture to be square before tiling, since this
-					// seems to be what non square textures do by default in UE. If we decide to tile without 
-					// stretching by default someday, we'd do UV - FVector2f(FMath::Floor(UV.X), FMath:Floor(UV.Y/VHeight)*VHeight)
-					// without multiplying by VHeight afterward.
-					UV = UV * UVScale + UVOffset;
-					UV = UV - FVector2f(FMath::Floor(UV.X), FMath::Floor(UV.Y));
-					UV.Y *= VHeight;
-
-					double Offset = DisplaceField.BilinearSampleClamped(UV);
-					if (AdjustmentCurve)
+					return;
+				}	
+
+				
+				if (UVOverlay->IsSeamVertex(VID)) 
+				{
+					// If the vertex is on a UV seam edge then average out offsets for all elements sharing the vertex
+					TArray<int> TrianglesOut;
+					Mesh.GetVtxTriangles(VID, TrianglesOut);
+
+					for (const int TID : TrianglesOut) 
 					{
-						Offset = AdjustmentCurve->Eval(Offset);
+						int ElementID = UVOverlay->GetElementIDAtVertex(TID, VID);
+						double Offset = ComputeOffsetForElement(ElementID);
+						OffsetArray[VID] += Offset;
 					}
-					Offset -= DisplaceFieldBaseValue;
 
-					double Intensity = IntensityFunc(vid, Positions[vid], Normals[vid]);
-					DisplacedPositions[vid] = Positions[vid] + (Offset * Intensity * Normals[vid]);
+					checkSlow(TrianglesOut.Num() > 0);
+					OffsetArray[VID] /= TrianglesOut.Num();
 				}
-			}
+				else 
+				{
+					// Grab any triangle if it's not a uv seam vertex
+					int TID = Mesh.GetVtxSingleTriangle(VID);
+					int ElementID = UVOverlay->GetElementIDAtVertex(TID, VID);
+					double Offset = ComputeOffsetForElement(ElementID);
+					OffsetArray[VID] = Offset;
+				}
+				
+				double Intensity = IntensityFunc(VID, Positions[VID], Normals[VID]);
+				DisplacedPositions[VID] = Positions[VID] + (OffsetArray[VID] * Intensity * Normals[VID]);
+
+			}, bUseParallel ? EParallelForFlags::None : EParallelForFlags::ForceSingleThread);
 		}
 		
 		void Sine(const FDynamicMesh3& Mesh,
@@ -194,11 +274,12 @@ namespace DisplaceMeshToolLocals{
 			double Frequency,
 			double PhaseShift,
 			const FVector3d& Direction,
-			TArray<FVector3d>& DisplacedPositions)
+			TArray<FVector3d>& DisplacedPositions,
+			TArray<int>* VerticesToDisplace = nullptr)
 		{
 			FQuaterniond RotateToDirection(Direction, { 0.0, 0.0, 1.0 });
 
-			ParallelDisplace(Mesh, Positions, Normals, DisplacedPositions,
+			ParallelDisplace(Mesh, Positions, Normals, DisplacedPositions, VerticesToDisplace,
 				[&](int32 vid, const FVector3d& Position, const FVector3d& Normal)
 			{
 				FVector3d RotatedPosition = RotateToDirection * Position;
@@ -230,6 +311,14 @@ namespace DisplaceMeshToolLocals{
 	public:
 		FSubdivideMeshOp(const FDynamicMesh3& SourceMesh, const FSubdivideParameters& InParameters, EDisplaceMeshToolSubdivisionType InSubdivisionType);
 		void CalculateResult(FProgressCancel* Progress) final;
+		
+		TUniquePtr<TArray<int>> ExtractVerticesToDisplace() 
+		{
+			return MoveTemp(VerticesToDisplace);
+		}
+		
+	protected: 
+		TUniquePtr<TArray<int>> VerticesToDisplace = nullptr;
 	private:
 		FSubdivideParameters Parameters;
 		EDisplaceMeshToolSubdivisionType SubdivisionType;
@@ -292,6 +381,9 @@ namespace DisplaceMeshToolLocals{
 					FAdaptiveTessellate Tessellator(ResultMesh.Get(), &OutMesh);
 					Tessellator.Progress = ProgressCancel;
 					Tessellator.SetPattern(Pattern.Get());
+
+					VerticesToDisplace = MakeUnique<TArray<int>>();
+					Tessellator.TessInfo.SelectedVertices = VerticesToDisplace.Get();
 					
 					if (ensureMsgf(Tessellator.Validate() == EOperationValidationResult::Ok, TEXT("The tessellator parameters are invalid.")))
 					{
@@ -417,6 +509,8 @@ namespace DisplaceMeshToolLocals{
 
 		TSharedPtr<FIndexedWeightMap, ESPMode::ThreadSafe> WeightMap;
 		TFunction<float(const FVector3d&, const FIndexedWeightMap)> WeightMapQueryFunc;
+
+		TSharedPtr<TArray<int>, ESPMode::ThreadSafe> VerticesToDisplace; // if set, only displace vertices whose ids are part of the array
 	};
 
 	class FDisplaceMeshOp : public FDynamicMeshOperator
@@ -509,7 +603,8 @@ namespace DisplaceMeshToolLocals{
 				SourcePositions, 
 				SourceNormals,
 				IntensityFunc,
-				DisplacedPositions);
+				DisplacedPositions,
+				Parameters.VerticesToDisplace.Get());
 			break;
 
 		case EDisplaceMeshToolDisplaceType::RandomNoise:
@@ -518,7 +613,8 @@ namespace DisplaceMeshToolLocals{
 				SourceNormals,
 				IntensityFunc,
 				Parameters.RandomSeed,
-				DisplacedPositions);
+				DisplacedPositions,
+				Parameters.VerticesToDisplace.Get());
 			break;
 			
 		case EDisplaceMeshToolDisplaceType::PerlinNoise:
@@ -528,7 +624,8 @@ namespace DisplaceMeshToolLocals{
 				IntensityFunc,
 				Parameters.PerlinLayerProperties,	
 				Parameters.RandomSeed,
-				DisplacedPositions);
+				DisplacedPositions,
+				Parameters.VerticesToDisplace.Get());
 			break;
 
 		case EDisplaceMeshToolDisplaceType::DisplacementMap:
@@ -541,6 +638,7 @@ namespace DisplaceMeshToolLocals{
 				Parameters.DisplacementMapBaseValue,
 				Parameters.UVScale,
 				Parameters.UVOffset,
+				Parameters.VerticesToDisplace.Get(),
 				Parameters.AdjustmentCurve.Get());
 			break;
 
@@ -552,7 +650,8 @@ namespace DisplaceMeshToolLocals{
 				Parameters.SineWaveFrequency,
 				Parameters.SineWavePhaseShift,
 				(FVector3d)Parameters.SineWaveDirection,
-				DisplacedPositions);
+				DisplacedPositions,
+				Parameters.VerticesToDisplace.Get());
 			break;
 		}
 
@@ -599,6 +698,7 @@ namespace DisplaceMeshToolLocals{
 			SetFilterFalloffWidth(DisplaceParametersIn.FilterWidth);
 			SetPerlinNoiseLayerProperties(DisplaceParametersIn.PerlinLayerProperties);
 			SetDisplacementType(DisplacementTypeIn);
+			SetVerticesToDisplace(DisplaceParametersIn.VerticesToDisplace);
 
 			Parameters.WeightMap = DisplaceParametersIn.WeightMap;
 			Parameters.WeightMapQueryFunc = DisplaceParametersIn.WeightMapQueryFunc;
@@ -626,6 +726,7 @@ namespace DisplaceMeshToolLocals{
 		void SetPerlinNoiseLayerProperties(const TArray<FPerlinLayerProperties>& PerlinLayerProperties);
 		void SetWeightMap(TSharedPtr<FIndexedWeightMap, ESPMode::ThreadSafe> WeightMap);
 		void SetRecalculateNormals(bool bRecalculateNormals);
+		void SetVerticesToDisplace(TSharedPtr<TArray<int>, ESPMode::ThreadSafe> VerticesToDisplace);
 
 		TUniquePtr<FDynamicMeshOperator> MakeNewOperator() final
 		{
@@ -770,6 +871,11 @@ namespace DisplaceMeshToolLocals{
 		Parameters.bRecalculateNormals = RecalcNormalsIn;
 	}
 
+	void FDisplaceMeshOpFactory::SetVerticesToDisplace(TSharedPtr<TArray<int>, ESPMode::ThreadSafe> VerticesToDisplace)
+	{
+		Parameters.VerticesToDisplace = VerticesToDisplace;
+	}
+
 } // namespace
 
 /*
@@ -1300,7 +1406,12 @@ void UDisplaceMeshTool::AdvanceComputation()
 
 	if (SubdivideTask && SubdivideTask->IsDone())
 	{
-		SubdividedMesh = TSharedPtr<FDynamicMesh3, ESPMode::ThreadSafe>(SubdivideTask->GetTask().ExtractOperator()->ExtractResult().Release());
+		TUniquePtr<FDynamicMeshOperator> SubOp = SubdivideTask->GetTask().ExtractOperator();
+		FSubdivideMeshOp* SubOpDownCast = static_cast<FSubdivideMeshOp*>(SubOp.Get());
+		
+		SubdividedMesh = TSharedPtr<FDynamicMesh3, ESPMode::ThreadSafe>(SubOp->ExtractResult().Release());
+		VerticesToDisplace = TSharedPtr<TArray<int>, ESPMode::ThreadSafe>(SubOpDownCast->ExtractVerticesToDisplace().Release());
+
 		delete SubdivideTask;
 		SubdivideTask = nullptr;
 	}
@@ -1309,6 +1420,7 @@ void UDisplaceMeshTool::AdvanceComputation()
 		// force update of contrast curve
 		FDisplaceMeshOpFactory* DisplacerDownCast = static_cast<FDisplaceMeshOpFactory*>(Displacer.Get());
 		DisplacerDownCast->SetAdjustmentCurve(TextureMapProperties->bApplyAdjustmentCurve ? TextureMapProperties->AdjustmentCurve : nullptr);
+		DisplacerDownCast->SetVerticesToDisplace(VerticesToDisplace);
 
 		DisplaceTask = new FAsyncTaskExecuterWithAbort<TModelingOpTask<FDynamicMeshOperator>>(Displacer->MakeNewOperator());
 		DisplaceTask->StartBackgroundTask();
diff --git a/Engine/Plugins/Experimental/MeshModelingToolsetExp/Source/MeshModelingToolsExp/Public/DisplaceMeshTool.h b/Engine/Plugins/Experimental/MeshModelingToolsetExp/Source/MeshModelingToolsExp/Public/DisplaceMeshTool.h
--- a/Engine/Plugins/Experimental/MeshModelingToolsetExp/Source/MeshModelingToolsExp/Public/DisplaceMeshTool.h
+++ b/Engine/Plugins/Experimental/MeshModelingToolsetExp/Source/MeshModelingToolsExp/Public/DisplaceMeshTool.h
@@ -359,6 +359,7 @@ private:
 	float WeightMapQuery(const FVector3d& Position, const UE::Geometry::FIndexedWeightMap& WeightMap) const;
 
 	TSharedPtr<UE::Geometry::FDynamicMesh3, ESPMode::ThreadSafe> SubdividedMesh = nullptr;
+	TSharedPtr<TArray<int>, ESPMode::ThreadSafe> VerticesToDisplace = nullptr;
 	
 	UPROPERTY()
 	TObjectPtr<AInternalToolFrameworkActor> PreviewMeshActor = nullptr;
diff --git a/Engine/Plugins/Runtime/GeometryProcessing/Source/DynamicMesh/Private/Operations/AdaptiveTessellate.cpp b/Engine/Plugins/Runtime/GeometryProcessing/Source/DynamicMesh/Private/Operations/AdaptiveTessellate.cpp
--- a/Engine/Plugins/Runtime/GeometryProcessing/Source/DynamicMesh/Private/Operations/AdaptiveTessellate.cpp
+++ b/Engine/Plugins/Runtime/GeometryProcessing/Source/DynamicMesh/Private/Operations/AdaptiveTessellate.cpp
@@ -25,7 +25,8 @@ namespace AdaptiveTessellateLocals
 								  FProgressCancel* Progress,
 								  FTessellationData& TessData,
 								  FCompactMaps& CompactInfo,
-								  FDynamicMesh3* ResultMesh);
+								  FDynamicMesh3* ResultMesh,
+								  FAdaptiveTessellate::FTessellationInformation& TessInfo);
 
 	template<typename RealType, int ElementSize> 
 	bool ConstructTessellatedOverlay(const FDynamicMesh3* Mesh, 
@@ -859,9 +860,10 @@ namespace AdaptiveTessellateLocals
 							FProgressCancel* Progress,
 							FCompactMaps& OutCompactInfo, 
 							FTessellationData& TessData,
-							FDynamicMesh3* OutMesh) 
+							FDynamicMesh3* OutMesh,
+							FAdaptiveTessellate::FTessellationInformation& TessInfo) 
 	{	
-		
+
 		TessData.Init(Pattern);
 
 		// Tessellate edge patches
@@ -951,7 +953,12 @@ namespace AdaptiveTessellateLocals
 			return false;
 		}
 
-		if (AdaptiveTessellateLocals::ConstructTessellatedMesh(Mesh, Progress, TessData, OutCompactInfo, OutMesh) == false) 
+		if (AdaptiveTessellateLocals::ConstructTessellatedMesh(Mesh, 
+														       Progress, 
+															   TessData, 
+															   OutCompactInfo, 
+															   OutMesh, 
+															   TessInfo) == false)
 		{
 			return false;
 		}
@@ -1142,7 +1149,8 @@ namespace AdaptiveTessellateLocals
 								  FProgressCancel* Progress,
 								  FTessellationData& TessData,
 								  FCompactMaps& CompactInfo,
-								  FDynamicMesh3* ResultMesh)
+								  FDynamicMesh3* ResultMesh,
+								  FAdaptiveTessellate::FTessellationInformation& TessInfo)
 	{
 		ResultMesh->Clear();
 
@@ -1187,7 +1195,14 @@ namespace AdaptiveTessellateLocals
 				
 				const FVector3d Vertex = (1.0 - Alpha)*Mesh->GetVertex(EdgeV[0]) + Alpha*Mesh->GetVertex(EdgeV[1]);
 				
-				CompactInfo.SetVertexMapping(VertexID, ResultMesh->AppendVertex(Vertex));
+				const int NewVertexID = ResultMesh->AppendVertex(Vertex);
+
+				CompactInfo.SetVertexMapping(VertexID, NewVertexID);
+
+				if (TessInfo.SelectedVertices) 
+				{
+					TessInfo.SelectedVertices->Add(NewVertexID);
+				}
 			}
 		}
 
@@ -1214,7 +1229,13 @@ namespace AdaptiveTessellateLocals
 										 Bary[1] * Mesh->GetVertex(TriangleV[1]) + 
 										 Bary[2] * Mesh->GetVertex(TriangleV[2]);
 				
-				CompactInfo.SetVertexMapping(VertexID, ResultMesh->AppendVertex(Vertex));
+				const int NewVertexID = ResultMesh->AppendVertex(Vertex);
+				CompactInfo.SetVertexMapping(VertexID, NewVertexID);
+
+				if (TessInfo.SelectedVertices) 
+				{
+					TessInfo.SelectedVertices->Add(NewVertexID);
+				}
 			}
 		}
 
@@ -1254,6 +1275,26 @@ namespace AdaptiveTessellateLocals
 				const FIndex3i Tri = TrianglesBlock[Index];
 				const FIndex3i MappedTri = CompactInfo.GetVertexMapping(Tri);
 				CompactInfo.SetTriangleMapping(NewTriangleID, ResultMesh->AppendTriangle(MappedTri, TriGrp));
+			}	 
+		}
+		
+		if (TessInfo.SelectedVertices) 
+		{
+			// Add all of the original vertices of the triangles we are tessellating
+			TSet<int> TempSetSelectedVertices; // first add them to tset to avoid duplicates
+			for (const int TriangleID : TessData.TrianglesToTessellate) 
+			{
+				const FIndex3i Tri = Mesh->GetTriangle(TriangleID);
+				const FIndex3i MappedTri = CompactInfo.GetVertexMapping(Tri);
+				
+				TempSetSelectedVertices.Add(MappedTri[0]);
+				TempSetSelectedVertices.Add(MappedTri[1]);
+				TempSetSelectedVertices.Add(MappedTri[2]);
+			}
+
+			for (int VID : TempSetSelectedVertices) 
+			{
+				TessInfo.SelectedVertices->Add(VID);
 			}
 		}
 
@@ -1709,14 +1750,19 @@ namespace AdaptiveTessellateLocals
 	 * @param OutMesh The resulting tessellated mesh.
 	 * @return false if the tessellation failed or the user cancelled it
 	 */
-	bool Tessellate(const FDynamicMesh3* InMesh, const FTessellationPattern* Pattern, const bool bUseParallel, FProgressCancel* Progress, FDynamicMesh3* OutMesh)
+	bool Tessellate(const FDynamicMesh3* InMesh, 
+					const FTessellationPattern* Pattern, 
+					const bool bUseParallel, 
+					FProgressCancel* Progress, 
+					FDynamicMesh3* OutMesh,
+					FAdaptiveTessellate::FTessellationInformation& TessInfo)
 	{
 		OutMesh->Clear();
 
 		FCompactMaps CompactInfo;
 
 		AdaptiveTessellateLocals::FTessellationData TessData(InMesh);
-		if (TessellateGeometry(InMesh, Pattern, bUseParallel, Progress, CompactInfo, TessData, OutMesh) == false) 
+		if (TessellateGeometry(InMesh, Pattern, bUseParallel, Progress, CompactInfo, TessData, OutMesh, TessInfo) == false) 
 		{
 			return false;
 		}
@@ -2157,7 +2203,7 @@ bool FAdaptiveTessellate::Compute()
 		Mesh = ResultMeshCopy.Get();
 	}
 
-	bool bTessResult = AdaptiveTessellateLocals::Tessellate(Mesh, Pattern, bUseParallel, Progress, ResultMesh);
+	bool bTessResult = AdaptiveTessellateLocals::Tessellate(Mesh, Pattern, bUseParallel, Progress, ResultMesh, TessInfo);
 
 	if (Cancelled() || bTessResult == false)
 	{
diff --git a/Engine/Plugins/Runtime/GeometryProcessing/Source/DynamicMesh/Public/Operations/AdaptiveTessellate.h b/Engine/Plugins/Runtime/GeometryProcessing/Source/DynamicMesh/Public/Operations/AdaptiveTessellate.h
--- a/Engine/Plugins/Runtime/GeometryProcessing/Source/DynamicMesh/Public/Operations/AdaptiveTessellate.h
+++ b/Engine/Plugins/Runtime/GeometryProcessing/Source/DynamicMesh/Public/Operations/AdaptiveTessellate.h
@@ -228,7 +228,29 @@ public:
 
 	/** The tessellated mesh. */
 	FDynamicMesh3* ResultMesh = nullptr;
-	
+
+    //
+    // Output
+    //
+    
+    /** 
+     * Stores any additional tessellation information that was requested by the caller.
+     * The caller should set the pointers to request the information to be computed. 
+     * 
+     * @note The caller is responsible for managing the memmory the pointers are pointing to.
+     *       The operator simply populates the data structures if they are not null.
+     */
+    struct FTessellationInformation
+    {
+        /** 
+         * If not null, construction of tessellated mesh will append a list of all the vertices that belong to 
+         * triangles affected by the tessellation.
+         */
+        TArray<int32>* SelectedVertices = nullptr;
+    };
+
+    FTessellationInformation TessInfo;  
+    	
 protected:
 	
 	/** 
