GeometryFramework: harden handling of an empty material set in MeshRenderDecomposition
ModelingTools: prevent user from accepting EditMeshMaterialsTool with an empty material set

diff --git a/Engine/Plugins/Experimental/MeshModelingToolsetExp/Source/MeshModelingToolsEditorOnlyExp/Private/EditMeshMaterialsTool.cpp b/Engine/Plugins/Experimental/MeshModelingToolsetExp/Source/MeshModelingToolsEditorOnlyExp/Private/EditMeshMaterialsTool.cpp
--- a/Engine/Plugins/Experimental/MeshModelingToolsetExp/Source/MeshModelingToolsEditorOnlyExp/Private/EditMeshMaterialsTool.cpp
+++ b/Engine/Plugins/Experimental/MeshModelingToolsetExp/Source/MeshModelingToolsEditorOnlyExp/Private/EditMeshMaterialsTool.cpp
@@ -308,6 +308,18 @@ void UEditMeshMaterialsTool::OnMaterialSetChanged()
 	MaterialProps->UpdateFromMaterialsList();
 
 	bHaveModifiedMaterials = true;
+
+	if (MaterialProps->Materials.Num() == 0)
+	{
+		GetToolManager()->DisplayMessage(LOCTEXT("NoMaterialsMessage", "Material Set must contain at least one Material"), EToolMessageLevel::UserWarning);
+		bShowingMaterialSetError = true;
+	}
+	else if (bShowingMaterialSetError)
+	{
+		GetToolManager()->DisplayMessage({}, EToolMessageLevel::UserWarning);
+		bShowingMaterialSetError = false;
+	}
+
 }
 
 
@@ -324,6 +336,13 @@ void UEditMeshMaterialsTool::ExternalUpdateMaterialSet(const TArray<UMaterialInt
 
 
 
+bool UEditMeshMaterialsTool::CanAccept() const
+{ 
+	return  (CurrentMaterials.Num() > 0) &&
+		    (UMeshSelectionTool::CanAccept() || bHaveModifiedMaterials); 
+}
+
+
 void UEditMeshMaterialsTool::ApplyShutdownAction(EToolShutdownType ShutdownType)
 {
 	if (ShutdownType == EToolShutdownType::Accept)
diff --git a/Engine/Plugins/Experimental/MeshModelingToolsetExp/Source/MeshModelingToolsEditorOnlyExp/Public/EditMeshMaterialsTool.h b/Engine/Plugins/Experimental/MeshModelingToolsetExp/Source/MeshModelingToolsEditorOnlyExp/Public/EditMeshMaterialsTool.h
--- a/Engine/Plugins/Experimental/MeshModelingToolsetExp/Source/MeshModelingToolsEditorOnlyExp/Public/EditMeshMaterialsTool.h
+++ b/Engine/Plugins/Experimental/MeshModelingToolsetExp/Source/MeshModelingToolsEditorOnlyExp/Public/EditMeshMaterialsTool.h
@@ -91,7 +91,7 @@ public:
 	virtual void RegisterActions(FInteractiveToolActionSet& ActionSet) override;
 	virtual void OnTick(float DeltaTime) override;
 
-	virtual bool CanAccept() const override { return UMeshSelectionTool::CanAccept() || bHaveModifiedMaterials; }
+	virtual bool CanAccept() const override;
 
 	void RequestMaterialAction(EEditMeshMaterialsToolActions ActionType);
 
@@ -120,6 +120,7 @@ protected:
 
 	FMaterialSetKey InitialMaterialKey;
 	bool bHaveModifiedMaterials = false;
+	bool bShowingMaterialSetError = false;
 
 	virtual void ApplyShutdownAction(EToolShutdownType ShutdownType) override;
 
diff --git a/Engine/Source/Runtime/GeometryFramework/Private/Components/DynamicMeshComponent.cpp b/Engine/Source/Runtime/GeometryFramework/Private/Components/DynamicMeshComponent.cpp
--- a/Engine/Source/Runtime/GeometryFramework/Private/Components/DynamicMeshComponent.cpp
+++ b/Engine/Source/Runtime/GeometryFramework/Private/Components/DynamicMeshComponent.cpp
@@ -947,7 +947,7 @@ void UDynamicMeshComponent::DisableSecondaryTriangleBuffers()
 
 void UDynamicMeshComponent::SetExternalDecomposition(TUniquePtr<FMeshRenderDecomposition> DecompositionIn)
 {
-	check(DecompositionIn->Num() > 0);
+	ensure(DecompositionIn->Num() > 0);
 	Decomposition = MoveTemp(DecompositionIn);
 	NotifyMeshUpdated();
 }
diff --git a/Engine/Source/Runtime/GeometryFramework/Private/Components/MeshRenderDecomposition.cpp b/Engine/Source/Runtime/GeometryFramework/Private/Components/MeshRenderDecomposition.cpp
--- a/Engine/Source/Runtime/GeometryFramework/Private/Components/MeshRenderDecomposition.cpp
+++ b/Engine/Source/Runtime/GeometryFramework/Private/Components/MeshRenderDecomposition.cpp
@@ -48,6 +48,8 @@ void FMeshRenderDecomposition::BuildMaterialDecomposition(const FDynamicMesh3* M
 		Group.Material = UseMaterials->Materials[k];
 	}
 
+	TArray<int32> InvalidTris;
+
 	for (int32 tid : Mesh->TriangleIndicesItr())
 	{
 		int32 MatIdx = 0;
@@ -60,6 +62,21 @@ void FMeshRenderDecomposition::BuildMaterialDecomposition(const FDynamicMesh3* M
 			FMeshRenderDecomposition::FGroup& Group = Decomp.GetGroup(MatIdx);
 			Group.Triangles.Add(tid);
 		}
+		else
+		{
+			InvalidTris.Add(tid);		// MaterialID refers to missing material
+		}
+	}
+
+	// If any triangles were not allocated to a specific material group, put them into a new group.
+	// Possibly these will not be rendered correctly by higher-level code, but at least we can
+	// guarantee that they are not lost here
+	if (InvalidTris.Num() > 0)
+	{
+		int32 NewGroup = Decomp.AppendGroup();
+		FMeshRenderDecomposition::FGroup& Group = Decomp.GetGroup(NewGroup);
+		Group.Triangles = MoveTemp(InvalidTris);
+		Group.Material = nullptr;
 	}
 }
 
