MeshModeling: PatternTool , add clamp to input.

#rnx

diff --git a/Engine/Plugins/Experimental/MeshModelingToolsetExp/Source/MeshModelingToolsExp/Public/PatternTool.h b/Engine/Plugins/Experimental/MeshModelingToolsetExp/Source/MeshModelingToolsExp/Public/PatternTool.h
--- a/Engine/Plugins/Experimental/MeshModelingToolsetExp/Source/MeshModelingToolsExp/Public/PatternTool.h
+++ b/Engine/Plugins/Experimental/MeshModelingToolsetExp/Source/MeshModelingToolsExp/Public/PatternTool.h
@@ -130,7 +130,7 @@ public:
 	EPatternToolAxisSpacingMode SpacingMode = EPatternToolAxisSpacingMode::ByCount;
 
 	/** Number of Pattern Elements to place */
-	UPROPERTY(EditAnywhere, Category = LinearPattern, meta = (UIMin = 1, UIMax = 25, EditCondition = "SpacingMode == EPatternToolAxisSpacingMode::ByCount"))
+	UPROPERTY(EditAnywhere, Category = LinearPattern, meta = (UIMin = 1, UIMax = 25, ClampMin = 1, EditCondition = "SpacingMode == EPatternToolAxisSpacingMode::ByCount"))
 	int32 Count = 10;
 
 	/** Fixed Increment used to place Pattern Elements */
@@ -163,7 +163,7 @@ public:
 	EPatternToolAxisSpacingMode SpacingX = EPatternToolAxisSpacingMode::ByCount;
 
 	/** Number of  Pattern Elements to place along the Main axis */
-	UPROPERTY(EditAnywhere, Category = GridPatternX, meta = (UIMin = 1, UIMax = 25, EditCondition = "SpacingX == EPatternToolAxisSpacingMode::ByCount"))
+	UPROPERTY(EditAnywhere, Category = GridPatternX, meta = (UIMin = 1, UIMax = 25, ClampMin = 1, EditCondition = "SpacingX == EPatternToolAxisSpacingMode::ByCount"))
 	int32 CountX = 10;
 
 	/** Fixed Increment used to place Pattern Elements along the Main axis */
@@ -183,7 +183,7 @@ public:
 	EPatternToolAxisSpacingMode SpacingY = EPatternToolAxisSpacingMode::ByCount;
 
 	/** Number of  Pattern Elements to place along the Secondary axis */
-	UPROPERTY(EditAnywhere, Category = GridPatternY, meta = (UIMin = 1, UIMax = 25, EditCondition = "SpacingY == EPatternToolAxisSpacingMode::ByCount"))
+	UPROPERTY(EditAnywhere, Category = GridPatternY, meta = (UIMin = 1, UIMax = 25, ClampMin = 1, EditCondition = "SpacingY == EPatternToolAxisSpacingMode::ByCount"))
 	int32 CountY = 10;
 
 	/** Fixed Increment used to place Pattern Elements along the Secondary axis */
@@ -216,7 +216,7 @@ public:
 	EPatternToolAxisSpacingMode SpacingMode = EPatternToolAxisSpacingMode::ByCount;
 
 	/** Number of  Pattern Elements to place */
-	UPROPERTY(EditAnywhere, Category = RadialPattern, meta = (UIMin = 1, UIMax = 25, EditCondition = "SpacingMode == EPatternToolAxisSpacingMode::ByCount"))
+	UPROPERTY(EditAnywhere, Category = RadialPattern, meta = (UIMin = 1, UIMax = 25, ClampMin = 1, EditCondition = "SpacingMode == EPatternToolAxisSpacingMode::ByCount"))
 	int32 Count = 10;
 
 	/** Fixed Increment (in Degrees) used to position Pattern Elements around the Circle/Arc */
