Approximate Actors: Added missing setup of bake options from FMeshApproximationSettings

diff --git a/Engine/Plugins/Experimental/MeshModelingToolsetExp/Source/GeometryProcessingAdapters/Private/ApproximateActorsImpl.cpp b/Engine/Plugins/Experimental/MeshModelingToolsetExp/Source/GeometryProcessingAdapters/Private/ApproximateActorsImpl.cpp
--- a/Engine/Plugins/Experimental/MeshModelingToolsetExp/Source/GeometryProcessingAdapters/Private/ApproximateActorsImpl.cpp
+++ b/Engine/Plugins/Experimental/MeshModelingToolsetExp/Source/GeometryProcessingAdapters/Private/ApproximateActorsImpl.cpp
@@ -992,6 +992,14 @@ IGeometryProcessing_ApproximateActors::FOptions FApproximateActorsImpl::Construc
 	// Ray tracing
 	Options.bSupportRayTracing = UseSettings.bSupportRayTracing;
 
+	// Material properties baking
+	Options.bBakeBaseColor = true;
+	Options.bBakeRoughness = UseSettings.MaterialSettings.bRoughnessMap;
+	Options.bBakeMetallic = UseSettings.MaterialSettings.bMetallicMap;
+	Options.bBakeSpecular = UseSettings.MaterialSettings.bSpecularMap;
+	Options.bBakeEmissive = UseSettings.MaterialSettings.bEmissiveMap;
+	Options.bBakeNormalMap = UseSettings.MaterialSettings.bNormalMap;
+
 	return Options;
 }
 
