Fixed a crash when Displace tool tries to read weightmap data from a deleted StaticMeshAsset.

#rnx none

diff --git a/Engine/Plugins/Experimental/MeshModelingToolsetExp/Source/MeshModelingToolsExp/Private/DisplaceMeshTool.cpp b/Engine/Plugins/Experimental/MeshModelingToolsetExp/Source/MeshModelingToolsExp/Private/DisplaceMeshTool.cpp
--- a/Engine/Plugins/Experimental/MeshModelingToolsetExp/Source/MeshModelingToolsExp/Private/DisplaceMeshTool.cpp
+++ b/Engine/Plugins/Experimental/MeshModelingToolsetExp/Source/MeshModelingToolsExp/Private/DisplaceMeshTool.cpp
@@ -1167,6 +1167,12 @@ void UDisplaceMeshTool::ValidateSubdivisions()
 void UDisplaceMeshTool::OnPropertyModified(UObject* PropertySet, FProperty* Property)
 {
 	using namespace DisplaceMeshToolLocals;
+	
+	if (AreAllTargetsValid() == false)
+	{
+		GetToolManager()->DisplayMessage(LOCTEXT("InvalidTargets", "Target mesh is no longer valid"), EToolMessageLevel::UserWarning);
+		return;
+	}
 
 	if (PropertySet && Property)
 	{
