Fixed a bug in the Displace tool where undoing changes to the perlin noise properties would not update the mesh in the viewport.

#rnx

diff --git a/Engine/Plugins/Experimental/MeshModelingToolsetExp/Source/MeshModelingToolsExp/Private/DisplaceMeshTool.cpp b/Engine/Plugins/Experimental/MeshModelingToolsetExp/Source/MeshModelingToolsExp/Private/DisplaceMeshTool.cpp
--- a/Engine/Plugins/Experimental/MeshModelingToolsetExp/Source/MeshModelingToolsExp/Private/DisplaceMeshTool.cpp
+++ b/Engine/Plugins/Experimental/MeshModelingToolsetExp/Source/MeshModelingToolsExp/Private/DisplaceMeshTool.cpp
@@ -1283,7 +1283,7 @@ void UDisplaceMeshTool::OnPropertyModified(UObject* PropertySet, FProperty* Prop
 		{
 			DisplacerDownCast->SetFilterFalloffWidth(DirectionalFilterProperties->FilterWidth);
 		}
-		else if ((PropName == GET_MEMBER_NAME_CHECKED(FPerlinLayerProperties, Frequency)) || (PropName == GET_MEMBER_NAME_CHECKED(FPerlinLayerProperties, Intensity)))
+		else if (PropertySet == NoiseProperties)
 		{
 			DisplacerDownCast->SetPerlinNoiseLayerProperties(NoiseProperties->PerlinLayerProperties);
 		}
