fix lattice tool for negative scales

diff --git a/Engine/Plugins/Experimental/MeshModelingToolsetExp/Source/MeshModelingToolsExp/Private/LatticeDeformerTool.cpp b/Engine/Plugins/Experimental/MeshModelingToolsetExp/Source/MeshModelingToolsExp/Private/LatticeDeformerTool.cpp
--- a/Engine/Plugins/Experimental/MeshModelingToolsetExp/Source/MeshModelingToolsExp/Private/LatticeDeformerTool.cpp
+++ b/Engine/Plugins/Experimental/MeshModelingToolsetExp/Source/MeshModelingToolsExp/Private/LatticeDeformerTool.cpp
@@ -173,6 +173,9 @@ void ULatticeDeformerTool::Setup()
 	FMeshDescriptionToDynamicMesh Converter;
 	Converter.Convert(UE::ToolTarget::GetMeshDescription(Target), *OriginalMesh);
 
+	// Note: Mesh will be implicitly transformed to world space by transforming the lattice; we account for whether that would invert the mesh here
+	MeshTransforms::ReverseOrientationIfNeeded(*OriginalMesh, (Cast<IPrimitiveComponentBackedTarget>(Target)->GetWorldTransform()));
+
 	Settings = NewObject<ULatticeDeformerToolProperties>(this, TEXT("Lattice Deformer Tool Settings"));
 	Settings->Initialize(this);
 	Settings->RestoreProperties(this);
@@ -358,7 +361,7 @@ void ULatticeDeformerTool::OnShutdown(EToolShutdownType ShutdownType)
 			// The lattice and its output mesh are in world space, so get them in local space.
 			// TODO: Would it make more sense to do all the lattice computation in local space?
 			FTransform3d LocalToWorld(TargetComponent->GetWorldTransform());
-			MeshTransforms::ApplyTransformInverse(*DynamicMeshResult, LocalToWorld);
+			MeshTransforms::ApplyTransformInverse(*DynamicMeshResult, LocalToWorld, true);
 
 			UE::ToolTarget::CommitMeshDescriptionUpdateViaDynamicMesh(Target, *DynamicMeshResult, true);
 
diff --git a/Engine/Source/Runtime/GeometryCore/Private/DynamicMesh/MeshTransforms.cpp b/Engine/Source/Runtime/GeometryCore/Private/DynamicMesh/MeshTransforms.cpp
--- a/Engine/Source/Runtime/GeometryCore/Private/DynamicMesh/MeshTransforms.cpp
+++ b/Engine/Source/Runtime/GeometryCore/Private/DynamicMesh/MeshTransforms.cpp
@@ -199,6 +199,15 @@ void MeshTransforms::ApplyTransformInverse(FDynamicMesh3& Mesh, const FTransform
 }
 
 
+void MeshTransforms::ReverseOrientationIfNeeded(FDynamicMesh3& Mesh, const FTransformSRT3d& Transform)
+{
+	if (Transform.GetDeterminant() < 0)
+	{
+		Mesh.ReverseOrientation(false);
+	}
+}
+
+
 
 void MeshTransforms::ApplyTransform(FDynamicMesh3& Mesh,
 	TFunctionRef<FVector3d(const FVector3d&)> PositionTransform,
diff --git a/Engine/Source/Runtime/GeometryCore/Public/DynamicMesh/MeshTransforms.h b/Engine/Source/Runtime/GeometryCore/Public/DynamicMesh/MeshTransforms.h
--- a/Engine/Source/Runtime/GeometryCore/Public/DynamicMesh/MeshTransforms.h
+++ b/Engine/Source/Runtime/GeometryCore/Public/DynamicMesh/MeshTransforms.h
@@ -67,6 +67,10 @@ namespace MeshTransforms
 		TFunctionRef<FVector3d(const FVector3d&)> PositionTransform,
 		TFunctionRef<FVector3f(const FVector3f&)> NormalTransform);
 
-
+	/**
+	 * If applying Transform would invert Mesh w/ a negative scale, then invert Mesh's triangle orientations.
+	 * Note: Does not apply the transform.
+	 */
+	GEOMETRYCORE_API void ReverseOrientationIfNeeded(FDynamicMesh3& Mesh, const FTransformSRT3d& Transform);
 
 };
\ No newline at end of file
