ModelingTools: move ISMEditorTool from CitySample to Main. Delete the CitySampleModelingTools plugin as all it did was add the ISMEditor Tool.

diff --git a/Engine/Plugins/Editor/ModelingToolsEditorMode/Source/ModelingToolsEditorMode/Private/ModelingToolsEditorMode.cpp b/Engine/Plugins/Editor/ModelingToolsEditorMode/Source/ModelingToolsEditorMode/Private/ModelingToolsEditorMode.cpp
--- a/Engine/Plugins/Editor/ModelingToolsEditorMode/Source/ModelingToolsEditorMode/Private/ModelingToolsEditorMode.cpp
+++ b/Engine/Plugins/Editor/ModelingToolsEditorMode/Source/ModelingToolsEditorMode/Private/ModelingToolsEditorMode.cpp
@@ -102,6 +102,7 @@
 // asset tools
 #include "Tools/GenerateStaticMeshLODAssetTool.h"
 #include "Tools/LODManagerTool.h"
+#include "ISMEditorTool.h"
 
 // commands
 #include "Commands/DeleteGeometrySelectionCommand.h"
@@ -561,6 +562,9 @@ void UModelingToolsEditorMode::Enter()
 	UGenerateStaticMeshLODAssetToolBuilder* GenerateSMLODToolBuilder = NewObject<UGenerateStaticMeshLODAssetToolBuilder>();
 	RegisterTool(ToolManagerCommands.BeginGenerateStaticMeshLODAssetTool, TEXT("BeginGenerateStaticMeshLODAssetTool"), GenerateSMLODToolBuilder);
 
+	UISMEditorToolBuilder* ISMEditorToolBuilder = NewObject<UISMEditorToolBuilder>();
+	RegisterTool(ToolManagerCommands.BeginISMEditorTool, TEXT("BeginISMEditorTool"), ISMEditorToolBuilder);
+
 
 	// edit tools
 
diff --git a/Engine/Plugins/Editor/ModelingToolsEditorMode/Source/ModelingToolsEditorMode/Private/ModelingToolsEditorModeToolkit.cpp b/Engine/Plugins/Editor/ModelingToolsEditorMode/Source/ModelingToolsEditorMode/Private/ModelingToolsEditorModeToolkit.cpp
--- a/Engine/Plugins/Editor/ModelingToolsEditorMode/Source/ModelingToolsEditorMode/Private/ModelingToolsEditorModeToolkit.cpp
+++ b/Engine/Plugins/Editor/ModelingToolsEditorMode/Source/ModelingToolsEditorMode/Private/ModelingToolsEditorModeToolkit.cpp
@@ -686,6 +686,7 @@ void FModelingToolsEditorModeToolkit::BuildToolPalette(FName PaletteIndex, class
 	{
 		ToolbarBuilder.AddToolBarButton(Commands.BeginLODManagerTool);
 		ToolbarBuilder.AddToolBarButton(Commands.BeginGenerateStaticMeshLODAssetTool);
+		ToolbarBuilder.AddToolBarButton(Commands.BeginISMEditorTool);
 	}
 	else if (PaletteIndex == VoxToolsTabName)
 	{
diff --git a/Engine/Plugins/Editor/ModelingToolsEditorMode/Source/ModelingToolsEditorMode/Private/ModelingToolsManagerActions.cpp b/Engine/Plugins/Editor/ModelingToolsEditorMode/Source/ModelingToolsEditorMode/Private/ModelingToolsManagerActions.cpp
--- a/Engine/Plugins/Editor/ModelingToolsEditorMode/Source/ModelingToolsEditorMode/Private/ModelingToolsManagerActions.cpp
+++ b/Engine/Plugins/Editor/ModelingToolsEditorMode/Source/ModelingToolsEditorMode/Private/ModelingToolsManagerActions.cpp
@@ -154,6 +154,7 @@ void FModelingToolsManagerCommands::RegisterCommands()
 	// LODs
 	REGISTER_MODELING_TOOL_COMMAND(BeginLODManagerTool, "LODMgr", "Static mesh asset LOD manager");
 	REGISTER_MODELING_TOOL_COMMAND(BeginGenerateStaticMeshLODAssetTool, "AutoLOD", "Generate static mesh LOD asset");
+	REGISTER_MODELING_TOOL_COMMAND(BeginISMEditorTool, "ISMEd", "Edit Instaces in InstancedStaticMeshComponents");
 
 	REGISTER_MODELING_TOOL_COMMAND(BeginAddPatchTool, "Patch", "Add Patch");
 	REGISTER_MODELING_TOOL_COMMAND(BeginShapeSprayTool, "Spray", "Shape Spray");
diff --git a/Engine/Plugins/Editor/ModelingToolsEditorMode/Source/ModelingToolsEditorMode/Public/ModelingToolsManagerActions.h b/Engine/Plugins/Editor/ModelingToolsEditorMode/Source/ModelingToolsEditorMode/Public/ModelingToolsManagerActions.h
--- a/Engine/Plugins/Editor/ModelingToolsEditorMode/Source/ModelingToolsEditorMode/Public/ModelingToolsManagerActions.h
+++ b/Engine/Plugins/Editor/ModelingToolsEditorMode/Source/ModelingToolsEditorMode/Public/ModelingToolsManagerActions.h
@@ -130,6 +130,7 @@ public:
 	
 	TSharedPtr<FUICommandInfo> BeginLODManagerTool;
 	TSharedPtr<FUICommandInfo> BeginGenerateStaticMeshLODAssetTool;
+	TSharedPtr<FUICommandInfo> BeginISMEditorTool;
 
 	TSharedPtr<FUICommandInfo> BeginPolyModelTool_FaceSelect;
 	TSharedPtr<FUICommandInfo> BeginPolyModelTool_EdgeSelect;
diff --git a/Engine/Plugins/Experimental/MeshModelingToolsetExp/Source/MeshModelingToolsEditorOnlyExp/Private/ISMEditorTool.cpp b/Engine/Plugins/Experimental/MeshModelingToolsetExp/Source/MeshModelingToolsEditorOnlyExp/Private/ISMEditorTool.cpp
new file mode 100644
--- /dev/null
+++ b/Engine/Plugins/Experimental/MeshModelingToolsetExp/Source/MeshModelingToolsEditorOnlyExp/Private/ISMEditorTool.cpp
@@ -0,0 +1,927 @@
+// Copyright Epic Games, Inc. All Rights Reserved.
+
+#include "ISMEditorTool.h"
+#include "InteractiveToolManager.h"
+#include "InteractiveGizmoManager.h"
+#include "Mechanics/DragAlignmentMechanic.h"
+#include "ToolBuilderUtil.h"
+#include "ToolSetupUtil.h"
+#include "DynamicMesh/DynamicMesh3.h"
+#include "BaseBehaviors/ClickDragBehavior.h"
+#include "ToolSceneQueriesUtil.h"
+#include "ModelingToolTargetUtil.h"
+#include "Scene/MeshSceneAdapter.h"
+#include "ToolDataVisualizer.h"
+#include "Drawing/PreviewGeometryActor.h"
+
+#include "BaseGizmos/GizmoComponents.h"
+#include "BaseGizmos/TransformGizmoUtil.h"
+
+#include "Components/PrimitiveComponent.h"
+#include "Components/InstancedStaticMeshComponent.h"
+#include "Engine/World.h"
+
+#include "Editor/UnrealEdEngine.h"
+#include "UnrealEdGlobals.h"
+
+using namespace UE::Geometry;
+
+#define LOCTEXT_NAMESPACE "UISMEditorTool"
+
+
+/*
+ * ToolBuilder
+ */
+
+
+template<typename IterFunc>
+void ComponentsIteration(const FToolBuilderState& SceneState, IterFunc Func)
+{
+	if (SceneState.SelectedComponents.Num() > 0)
+	{
+		for (UActorComponent* Component : SceneState.SelectedComponents)
+		{
+			if (Func(Component) == false)
+			{
+				return;
+			}
+		}
+	}
+	else if (SceneState.SelectedActors.Num() > 0)
+	{
+		for (AActor* Actor : SceneState.SelectedActors)
+		{
+			for (UActorComponent* Component : Actor->GetComponents())
+			{
+				if (Func(Component) == false)
+				{
+					return;
+				}
+			}
+
+			// not currently handling child actor components...
+		}
+	}
+}
+
+
+bool UISMEditorToolBuilder::CanBuildTool(const FToolBuilderState& SceneState) const
+{
+	bool bFoundISM = false;
+	ComponentsIteration(SceneState, [&bFoundISM](UActorComponent* Component)
+	{
+		if (Cast<UInstancedStaticMeshComponent>(Component) != nullptr)
+		{
+			bFoundISM = true;
+			return false;
+		}
+		return true;
+	});
+	return bFoundISM;
+}
+
+UInteractiveTool* UISMEditorToolBuilder::BuildTool(const FToolBuilderState& SceneState) const
+{
+	UISMEditorTool* NewTool = NewObject<UISMEditorTool>(SceneState.ToolManager);
+
+	TArray<UInstancedStaticMeshComponent*> Targets;
+	ComponentsIteration(SceneState, [&Targets](UActorComponent* Component)
+	{
+		if (Cast<UInstancedStaticMeshComponent>(Component) != nullptr)
+		{
+			Targets.Add(Cast<UInstancedStaticMeshComponent>(Component));
+		}
+		return true;
+	});
+
+	NewTool->SetTargets(MoveTemp(Targets));
+	return NewTool;
+}
+
+
+
+void UISMEditorToolActionPropertySetBase::PostAction(EISMEditorToolActions Action)
+{
+	ParentTool->RequestAction(Action);
+}
+
+
+/*
+ * Tool
+ */
+
+void UISMEditorTool::SetTargets(TArray<UInstancedStaticMeshComponent*> Components)
+{
+	TargetComponents = Components;
+}
+
+
+void UISMEditorTool::Setup()
+{
+	UInteractiveTool::Setup();
+
+	bMeshSceneDirty = true;
+
+	// set up rectangle marquee mechanic
+	RectangleMarqueeMechanic = NewObject<URectangleMarqueeMechanic>(this);
+	RectangleMarqueeMechanic->bUseExternalClickDragBehavior = true;
+	RectangleMarqueeMechanic->Setup(this);
+	RectangleMarqueeMechanic->SetIsEnabled(true);
+	RectangleMarqueeMechanic->OnDragRectangleFinished.AddUObject(this, &UISMEditorTool::OnMarqueeRectangleFinished);
+
+	// clickdrag behavior clicks-to-select or does marquee on drag
+	ClickOrDragBehavior = NewObject<USingleClickOrDragInputBehavior>();
+	ClickOrDragBehavior->Initialize(this, RectangleMarqueeMechanic);
+	ClickOrDragBehavior->Modifiers.RegisterModifier(1, FInputDeviceState::IsShiftKeyDown);
+	ClickOrDragBehavior->Modifiers.RegisterModifier(2, FInputDeviceState::IsCtrlKeyDown);
+	AddInputBehavior(ClickOrDragBehavior);
+
+	TransformProps = NewObject<UISMEditorToolProperties>();
+	AddToolPropertySource(TransformProps);
+	TransformProps->RestoreProperties(this);
+	TransformProps->WatchProperty(TransformProps->TransformMode, [this](EISMEditorTransformMode NewMode) { UpdateTransformMode(NewMode); });
+	TransformProps->WatchProperty(TransformProps->bSetPivotMode, [this](bool bNewValue) { UpdateSetPivotModes(bNewValue); });
+
+	UpdateTransformMode(TransformProps->TransformMode);
+
+	ToolActions = NewObject<UISMEditorToolActionPropertySet>();
+	ToolActions->Initialize(this);
+	AddToolPropertySource(ToolActions);
+
+
+	TArray<AActor*> UniqueActors;
+	for (UInstancedStaticMeshComponent* Component : TargetComponents)
+	{
+		UniqueActors.AddUnique(Component->GetOwner());
+	}
+	if (UniqueActors.Num() == 1)
+	{
+		ReplaceAction = NewObject<UISMEditorToolReplacePropertySet>();
+		ReplaceAction->Initialize(this);
+		AddToolPropertySource(ReplaceAction);
+	}
+
+	PreviewGeometry = NewObject<UPreviewGeometry>(this);
+	PreviewGeometry->CreateInWorld(TargetComponents[0]->GetWorld(), FTransform::Identity);
+
+
+	SetToolDisplayName(LOCTEXT("ToolName", "Edit ISMs"));
+	GetToolManager()->DisplayMessage(
+		LOCTEXT("OnStartISMEditorTool", "Edit the selected ISMs"),
+		EToolMessageLevel::UserNotification);
+}
+
+
+void UISMEditorTool::Shutdown(EToolShutdownType ShutdownType)
+{
+	TransformProps->SaveProperties(this);
+
+	GetToolManager()->GetPairedGizmoManager()->DestroyAllGizmosByOwner(this);
+
+	if (PreviewGeometry)
+	{
+		PreviewGeometry->Disconnect();
+		PreviewGeometry = nullptr;
+	}
+
+	// Release the FMeshSceneAdapter, this will clear out any MeshDescriptions that were loaded unnecessarily.
+	// This needs to happen on game thread, if we don't do it here, it may happen during GC later
+	MeshScene.Reset();
+}
+
+
+void UISMEditorTool::OnUpdateModifierState(int ModifierID, bool bIsOn)
+{
+	if (ModifierID == 1)
+	{
+		bShiftModifier = bIsOn;
+	}
+	if (ModifierID == 2)
+	{
+		bCtrlModifier = bIsOn;
+	}
+}
+
+
+
+void UISMEditorTool::OnMarqueeRectangleFinished(const FCameraRectangle& Rectangle, bool bCancelled)
+{
+	if (!bCancelled)
+	{
+		TArray<const FActorChildMesh*> HitMeshes;
+		FCriticalSection HitMeshLock;
+		
+		MeshScene->ParallelMeshVertexEnumeration(
+			[&](int32 NumMeshes) { return true; },
+			[&](int32 MeshIndex, AActor* SourceActor, const FActorChildMesh* ChildMeshInfo, const FAxisAlignedBox3d& WorldBounds) { return true; },
+			[&](int32 MeshIndex, AActor* SourceActor, const FActorChildMesh* ChildMeshInfo, const FVector3d& WorldPos) 
+			{ 
+				if ( Rectangle.IsProjectedPointInRectangle( (FVector)WorldPos ) )
+				{
+					HitMeshLock.Lock();
+					HitMeshes.Add(ChildMeshInfo);
+					HitMeshLock.Unlock();
+					return false;
+				}
+				return true; 
+			});
+
+
+		TArray<FSelectedInstance> NewSelection;
+		if (bShiftModifier || bCtrlModifier)
+		{
+			NewSelection = CurrentSelection;
+		}
+		bool bSubtract = bCtrlModifier;
+
+		for (const FActorChildMesh* HitMesh : HitMeshes)
+		{
+			FSelectedInstance NewItem{ 
+				(UInstancedStaticMeshComponent*)HitMesh->SourceComponent, 
+				HitMesh->ComponentIndex,
+				HitMesh->MeshSpatial->GetWorldBounds(
+					[&](const FVector3d& P) { return HitMesh->WorldTransform.TransformPosition(P); })
+			};
+
+			if (bSubtract)
+			{
+				int32 ExistingIndex = NewSelection.Find(NewItem);
+				if (ExistingIndex != INDEX_NONE)
+				{
+					NewSelection.RemoveAt(ExistingIndex);
+				}
+			}
+			else
+			{
+				NewSelection.AddUnique(NewItem);
+			}
+		}
+
+		if (NewSelection != CurrentSelection)
+		{
+			UpdateSelectionInternal(NewSelection, true);
+		}
+	}
+}
+
+
+void UISMEditorTool::Render(IToolsContextRenderAPI* RenderAPI)
+{
+	if (bInActiveDrag && TransformProps->bHideWhenDragging)
+	{
+		return;
+	}
+
+	if (RectangleMarqueeMechanic)
+	{
+		RectangleMarqueeMechanic->Render(RenderAPI);
+	}
+
+	FToolDataVisualizer LineRenderer;
+	LineRenderer.BeginFrame(RenderAPI);
+
+	//if (TransformProps->bShowSelectable)
+	//{
+	//	LineRenderer.SetLineParameters(FLinearColor(0.2, 0.9, 0.2), 2.0);
+	//	for (const FAxisAlignedBox3d& Box : AllMeshBoundingBoxes)
+	//	{
+	//		LineRenderer.DrawWireBox((FBox)Box);
+	//	}
+	//}
+
+	if (TransformProps->bShowSelected)
+	{
+		LineRenderer.SetLineParameters(FLinearColor(0.95, 0.05, 0.05), 4.0);
+		for (FSelectedInstance Instance : CurrentSelection)
+		{
+			if (Instance.Component->IsValidInstance(Instance.Index))
+			{
+				LineRenderer.DrawWireBox((FBox)Instance.WorldBounds);
+			}
+		}
+	}
+
+	LineRenderer.EndFrame();
+}
+
+
+
+void UISMEditorTool::DrawHUD(FCanvas* Canvas, IToolsContextRenderAPI* RenderAPI)
+{
+	if (RectangleMarqueeMechanic)
+	{
+		RectangleMarqueeMechanic->DrawHUD(Canvas, RenderAPI);
+	}
+}
+
+
+
+void UISMEditorTool::UpdatePreviewGeometry()
+{
+	FColor LinesColor(50, 225, 50);
+	constexpr int BoxFaces[6][4] =
+	{
+		{ 0, 1, 2, 3 },     // back, -z
+		{ 5, 4, 7, 6 },     // front, +z
+		{ 4, 0, 3, 7 },     // left, -x
+		{ 1, 5, 6, 2 },     // right, +x,
+		{ 4, 5, 1, 0 },     // bottom, -y
+		{ 3, 2, 6, 7 }      // top, +y
+	};
+
+	PreviewGeometry->CreateOrUpdateLineSet(TEXT("SelectableBoxes"), AllMeshBoundingBoxes.Num(),
+										   [&](int32 Index, TArray<FRenderableLine>& LinesOut)
+	{
+		FBox Box = (FBox)AllMeshBoundingBoxes[Index];
+		FVector Corners[8] =
+		{
+			Box.Min,
+			FVector(Box.Max.X, Box.Min.Y, Box.Min.Z),
+			FVector(Box.Max.X, Box.Max.Y, Box.Min.Z),
+			FVector(Box.Min.X, Box.Max.Y, Box.Min.Z),
+			FVector(Box.Min.X, Box.Min.Y, Box.Max.Z),
+			FVector(Box.Max.X, Box.Min.Y, Box.Max.Z),
+			Box.Max,
+			FVector(Box.Min.X, Box.Max.Y, Box.Max.Z)
+		};
+		for (int FaceIdx = 0; FaceIdx < 6; FaceIdx++)
+		{
+			for (int Last = 3, Cur = 0; Cur < 4; Last = Cur++)
+			{
+				LinesOut.Add(FRenderableLine(Corners[BoxFaces[FaceIdx][Last]], Corners[BoxFaces[FaceIdx][Cur]],
+											 LinesColor, 2.0, 0.0));
+			}
+		}
+
+	}, 12);
+
+}
+
+
+void UISMEditorTool::OnTick(float DeltaTime)
+{
+	if (PendingAction != EISMEditorToolActions::NoAction)
+	{
+		switch (PendingAction)
+		{
+		case EISMEditorToolActions::ClearSelection:
+			OnClearSelection();
+			break;
+		case EISMEditorToolActions::Delete:
+			OnDeleteSelection();
+			break;
+		case EISMEditorToolActions::Duplicate:
+			OnDuplicateSelection();
+			break;
+		case EISMEditorToolActions::Replace:
+			OnReplaceSelection();
+			break;
+		}
+		PendingAction = EISMEditorToolActions::NoAction;
+	}
+
+	if (bMeshSceneDirty)
+	{
+		RebuildMeshScene();
+		bMeshSceneDirty = false;
+
+		// when we dirty mesh scene on undo/redo, selection may have been invalid
+		OnSelectionUpdated();
+	}
+
+	if (PreviewGeometry != nullptr)
+	{
+		PreviewGeometry->SetLineSetVisibility(TEXT("SelectableBoxes"), TransformProps->bShowSelectable);
+	}
+}
+
+
+void UISMEditorTool::UpdateSetPivotModes(bool bEnableSetPivot)
+{
+	for (FISMEditorTarget& Target : ActiveGizmos)
+	{
+		Target.TransformProxy->bSetPivotMode = bEnableSetPivot;
+	}
+}
+
+
+
+void UISMEditorTool::RegisterActions(FInteractiveToolActionSet& ActionSet)
+{
+	ActionSet.RegisterAction(this, (int32)EStandardToolActions::BaseClientDefinedActionID + 1,
+		TEXT("ToggleSetPivot"),
+		LOCTEXT("TransformToggleSetPivot", "Toggle Set Pivot"),
+		LOCTEXT("TransformToggleSetPivotTooltip", "Toggle Set Pivot on and off"),
+		EModifierKey::None, EKeys::S,
+		[this]() { this->TransformProps->bSetPivotMode = !this->TransformProps->bSetPivotMode; });
+
+	ActionSet.RegisterAction(this, (int32)EStandardToolActions::BaseClientDefinedActionID + 3,
+		TEXT("CycleTransformMode"),
+		LOCTEXT("TransformCycleTransformMode", "Next Transform Mode"),
+		LOCTEXT("TransformCycleTransformModeTooltip", "Cycle through available Transform Modes"),
+		EModifierKey::None, EKeys::A,
+		[this]() { this->TransformProps->TransformMode = (EISMEditorTransformMode)(((uint8)TransformProps->TransformMode+1) % (uint8)EISMEditorTransformMode::LastValue); });
+
+}
+
+
+
+
+
+FInputRayHit UISMEditorTool::IsHitByClick(const FInputDeviceRay& ClickPos)
+{
+	if (CurrentSelection.Num() > 0)
+	{
+		return FInputRayHit(0);
+	}
+	else
+	{
+		FMeshSceneRayHit HitResult;
+		if (MeshScene->FindNearestRayIntersection(ClickPos.WorldRay, HitResult))
+		{
+			return FInputRayHit(HitResult.RayDistance);
+		}
+		return FInputRayHit();
+	}
+}
+
+
+void UISMEditorTool::OnClicked(const FInputDeviceRay& ClickPos)
+{
+	TArray<FSelectedInstance> NewSelection;
+
+	FMeshSceneRayHit HitResult;
+	if (MeshScene->FindNearestRayIntersection(ClickPos.WorldRay, HitResult))
+	{
+		UE_LOG(LogTemp, Warning, TEXT("Hit Component %s, Element Index %d, Triangle %d"),
+			   *HitResult.HitComponent->GetName(), HitResult.HitComponentElementIndex, HitResult.HitMeshTriIndex);
+
+		FSelectedInstance NewItem{ 
+			(UInstancedStaticMeshComponent*)HitResult.HitComponent, 
+			HitResult.HitComponentElementIndex,
+			HitResult.HitMeshSpatialWrapper->GetWorldBounds(
+				[&](const FVector3d& P) { return HitResult.LocalToWorld.TransformPosition(P); })
+		};
+
+		if (bShiftModifier)
+		{
+			NewSelection = CurrentSelection;
+			NewSelection.AddUnique(NewItem);
+		}
+		else if (bCtrlModifier)
+		{
+			int32 ExistingIndex = NewSelection.Find(NewItem);
+			if (ExistingIndex != INDEX_NONE)
+			{
+				NewSelection.RemoveAt(ExistingIndex);
+			}
+			else 
+			{
+				NewSelection.Add(NewItem);
+			}
+		}
+		else
+		{
+			NewSelection.Add(NewItem);
+		}
+	}
+
+	if (NewSelection != CurrentSelection)
+	{
+		UpdateSelectionInternal(NewSelection, true);
+	}
+}
+
+
+void UISMEditorTool::OnSelectionUpdated()
+{
+	TArray<FString> SelectionStrings;
+	for (FSelectedInstance& SelectedItem : CurrentSelection)
+	{
+		if (TargetComponents.Num() > 1)
+		{
+			SelectionStrings.Add(FString::Printf(TEXT("%s : %d"), *SelectedItem.Component->GetName(), SelectedItem.Index));
+		}
+		else
+		{
+			SelectionStrings.Add(FString::Printf(TEXT("%d"), SelectedItem.Index));
+		}
+	}
+	TransformProps->SelectedInstances = SelectionStrings;
+
+	UpdateTransformMode(TransformProps->TransformMode);
+}
+
+
+
+void UISMEditorTool::UpdateTransformMode(EISMEditorTransformMode NewMode)
+{
+	ResetActiveGizmos();
+
+	if (CurrentSelection.Num() == 0)
+	{
+		return;
+	}
+
+	switch (NewMode)
+	{
+		default:
+		case EISMEditorTransformMode::SharedGizmo:
+			SetActiveGizmos_Single(false);
+			break;
+
+		case EISMEditorTransformMode::SharedGizmoLocal:
+			SetActiveGizmos_Single(true);
+			break;
+
+		case EISMEditorTransformMode::PerObjectGizmo:
+			SetActiveGizmos_PerObject();
+			break;
+	}
+
+	CurTransformMode = NewMode;
+}
+
+
+
+namespace UE {
+namespace Local {
+
+static void AddInstancedComponentInstance(UInstancedStaticMeshComponent* ISMC, int32 Index, UTransformProxy* TransformProxy, bool bModifyOnTransform)
+{
+	TransformProxy->AddComponentCustom(ISMC,
+		[ISMC, Index]() {
+			FTransform Tmp;
+			ISMC->GetInstanceTransform(Index, Tmp, true);
+			return Tmp;
+		},
+		[ISMC, Index](FTransform NewTransform) {
+			ISMC->UpdateInstanceTransform(Index, NewTransform, true, true, true);
+		},
+		Index, bModifyOnTransform
+	);
+}
+
+}
+}
+
+
+
+void UISMEditorTool::SetActiveGizmos_Single(bool bLocalRotations)
+{
+	check(ActiveGizmos.Num() == 0);
+
+	FISMEditorTarget Transformable;
+	Transformable.TransformProxy = NewObject<UTransformProxy>(this);
+	Transformable.TransformProxy->bRotatePerObject = bLocalRotations;
+	Transformable.TransformProxy->OnBeginTransformEdit.AddLambda([this](UTransformProxy*) { bInActiveDrag = true; });
+	Transformable.TransformProxy->OnEndTransformEdit.AddLambda([this](UTransformProxy*) { OnTransformCompleted(); bInActiveDrag = false; });
+	Transformable.TransformProxy->OnTransformChangedUndoRedo.AddLambda([this](UTransformProxy*, FTransform) { OnTransformCompleted(); });
+
+	for (FSelectedInstance Instance : CurrentSelection)
+	{
+		if (Instance.Component->IsValidInstance(Instance.Index))
+		{
+			UE::Local::AddInstancedComponentInstance(Instance.Component, Instance.Index, Transformable.TransformProxy, true);
+		}
+	}
+
+	// leave out nonuniform scale if we have multiple objects in non-local mode
+	bool bCanNonUniformScale = TargetComponents.Num() == 1 || bLocalRotations;
+	ETransformGizmoSubElements GizmoElements = (bCanNonUniformScale) ?
+		ETransformGizmoSubElements::FullTranslateRotateScale : ETransformGizmoSubElements::TranslateRotateUniformScale;
+	Transformable.TransformGizmo = UE::TransformGizmoUtil::CreateCustomRepositionableTransformGizmo(GetToolManager(), GizmoElements, this);
+	Transformable.TransformGizmo->SetActiveTarget(Transformable.TransformProxy, GetToolManager());
+
+	ActiveGizmos.Add(Transformable);
+}
+
+void UISMEditorTool::SetActiveGizmos_PerObject()
+{
+	check(ActiveGizmos.Num() == 0);
+
+	for (FSelectedInstance Instance : CurrentSelection)
+	{
+		if (Instance.Component->IsValidInstance(Instance.Index))
+		{
+			FISMEditorTarget Transformable;
+			Transformable.TransformProxy = NewObject<UTransformProxy>(this);
+			Transformable.TransformProxy->OnBeginTransformEdit.AddLambda([this](UTransformProxy*) { bInActiveDrag = true; });
+			Transformable.TransformProxy->OnEndTransformEdit.AddLambda([this](UTransformProxy*) { OnTransformCompleted(); bInActiveDrag = false; });
+			Transformable.TransformProxy->OnTransformChangedUndoRedo.AddLambda([this](UTransformProxy*, FTransform) { OnTransformCompleted(); });
+
+			UE::Local::AddInstancedComponentInstance(Instance.Component, Instance.Index, Transformable.TransformProxy, true);
+
+			ETransformGizmoSubElements GizmoElements = ETransformGizmoSubElements::FullTranslateRotateScale;
+			Transformable.TransformGizmo = UE::TransformGizmoUtil::CreateCustomRepositionableTransformGizmo(GetToolManager(), GizmoElements, this);
+			Transformable.TransformGizmo->SetActiveTarget(Transformable.TransformProxy, GetToolManager());
+
+			ActiveGizmos.Add(Transformable);
+		}
+	}
+}
+
+
+
+void UISMEditorTool::ResetActiveGizmos()
+{
+	GetToolManager()->GetPairedGizmoManager()->DestroyAllGizmosByOwner(this);
+	ActiveGizmos.Reset();
+}
+
+
+void UISMEditorTool::RequestAction(EISMEditorToolActions ActionType)
+{
+	PendingAction = ActionType;
+}
+
+
+void UISMEditorTool::OnClearSelection()
+{
+	UpdateSelectionInternal(TArray<FSelectedInstance>(), true);
+}
+
+void UISMEditorTool::OnDeleteSelection()
+{
+	GetToolManager()->BeginUndoTransaction(LOCTEXT("DeleteSelected", "Delete Selected"));
+
+	TArray<UInstancedStaticMeshComponent*> UniqueInstances;
+	for (FSelectedInstance Instance : CurrentSelection)
+	{
+		if (Instance.Component->IsValidInstance(Instance.Index))
+		{
+			UniqueInstances.AddUnique(Instance.Component);
+		}
+	}
+
+	for (UInstancedStaticMeshComponent* DeleteFromComponent : UniqueInstances)
+	{
+		TArray<int32> InstanceList;
+		for (FSelectedInstance Instance : CurrentSelection)
+		{
+			if (Instance.Component == DeleteFromComponent && Instance.Component->IsValidInstance(Instance.Index))
+			{
+				InstanceList.Add(Instance.Index);
+			}
+		}
+
+		DeleteFromComponent->Modify();
+		DeleteFromComponent->RemoveInstances(InstanceList);
+	}
+
+	bMeshSceneDirty = true;
+	GetToolManager()->EmitObjectChange(this, MakeUnique<FISMEditorSceneChange>(), LOCTEXT("SceneChange", "Scene Update"));
+	UpdateSelectionInternal(TArray<FSelectedInstance>(), true);
+
+	GetToolManager()->EndUndoTransaction();
+}
+
+
+void UISMEditorTool::OnDuplicateSelection()
+{
+	GetToolManager()->BeginUndoTransaction(LOCTEXT("DuplicateSelected", "Duplicate Selected"));
+
+	TArray<FSelectedInstance> NewSelection;
+	for (FSelectedInstance Instance : CurrentSelection)
+	{
+		if (Instance.Component->IsValidInstance(Instance.Index))
+		{
+			FTransform CurInstanceTransform;
+			Instance.Component->GetInstanceTransform(Instance.Index, CurInstanceTransform, true);
+
+			Instance.Component->Modify();
+			int32 NewIndex = Instance.Component->AddInstance(CurInstanceTransform, true);
+
+			// copy per-instance custom data if it is defined
+			int32 NumCustomDataFloats = Instance.Component->NumCustomDataFloats;
+			if (NumCustomDataFloats > 0)
+			{
+				const TArray<float>& CurCustomData = Instance.Component->PerInstanceSMCustomData;
+				int32 BaseIndex = Instance.Index * NumCustomDataFloats;
+				if (CurCustomData.Num() >= (BaseIndex + NumCustomDataFloats))
+				{
+					TArray<float> NewCustomData;
+					for (int32 k = 0; k < NumCustomDataFloats; ++k)
+					{
+						NewCustomData.Add(CurCustomData[BaseIndex + k]);
+					}
+					Instance.Component->SetCustomData(NewIndex, NewCustomData);
+				}
+			}
+
+			NewSelection.Add(FSelectedInstance{ Instance.Component, NewIndex, Instance.WorldBounds });
+		}
+	}
+
+	bMeshSceneDirty = true;
+	GetToolManager()->EmitObjectChange(this, MakeUnique<FISMEditorSceneChange>(), LOCTEXT("SceneChange", "Scene Update"));
+	UpdateSelectionInternal(NewSelection, true);
+
+	GetToolManager()->EndUndoTransaction();
+}
+
+
+
+void UISMEditorTool::OnReplaceSelection()
+{
+	if (ReplaceAction->ReplaceWith == nullptr)
+	{
+		GetToolManager()->DisplayMessage(LOCTEXT("OnReplaceFailed", "Must set valid replacement StaticMesh"), EToolMessageLevel::UserError);
+		return;
+	}
+	if (CurrentSelection.Num() == 0)
+	{
+		GetToolManager()->DisplayMessage(LOCTEXT("OnReplaceFailed2", "Must select at least one Instance"), EToolMessageLevel::UserError);
+		return;
+	}
+
+	GetToolManager()->BeginUndoTransaction(LOCTEXT("ReplaceSelected", "Replace Selected"));
+
+	UInstancedStaticMeshComponent* DupeComponent = CurrentSelection[0].Component;
+	FTransform ExistingTransform = DupeComponent->GetComponentToWorld();
+
+	AActor* ParentActor = DupeComponent->GetOwner();
+	ParentActor->Modify();
+
+	// duplicate the component
+	TArray<UActorComponent*> SourceComponents;
+	SourceComponents.Add(DupeComponent);
+	TArray<UActorComponent*> NewComponents;
+	GUnrealEd->DuplicateComponents(SourceComponents, NewComponents);
+	UInstancedStaticMeshComponent* NewComponent = Cast<UInstancedStaticMeshComponent>(NewComponents[0]);
+
+	NewComponent->Modify();
+	NewComponent->SetStaticMesh(ReplaceAction->ReplaceWith);
+	NewComponent->SetComponentToWorld(ExistingTransform);
+	NewComponent->UnregisterComponent();
+
+	NewComponent->ClearInstances();
+
+	TArray<UInstancedStaticMeshComponent*> UniqueInstances;
+	for (FSelectedInstance Instance : CurrentSelection)
+	{
+		if (Instance.Component->IsValidInstance(Instance.Index))
+		{
+			UniqueInstances.AddUnique(Instance.Component);
+		}
+	}
+	for (UInstancedStaticMeshComponent* DeleteFromComponent : UniqueInstances)
+	{
+		TArray<int32> InstanceList;
+		TArray<FTransform> WorldTransforms;
+		for (FSelectedInstance Instance : CurrentSelection)
+		{
+			if (Instance.Component == DeleteFromComponent && Instance.Component->IsValidInstance(Instance.Index))
+			{
+				FTransform InstanceTransform;
+				Instance.Component->GetInstanceTransform(Instance.Index, InstanceTransform, true);
+				WorldTransforms.Add(InstanceTransform);
+				InstanceList.Add(Instance.Index);
+			}
+		}
+
+		DeleteFromComponent->Modify();
+		DeleteFromComponent->RemoveInstances(InstanceList);
+
+		NewComponent->AddInstances(WorldTransforms, false, true);
+	}
+
+	NewComponent->RegisterComponent();
+
+	TargetComponents.Add(NewComponent);
+	bMeshSceneDirty = true;
+
+	TUniquePtr<FISMEditorSceneChange> SceneChange = MakeUnique<FISMEditorSceneChange>();
+	SceneChange->Components.Add(NewComponent);
+	SceneChange->bAdded = true;
+	GetToolManager()->EmitObjectChange(this, MoveTemp(SceneChange), LOCTEXT("SceneChange", "Scene Update"));
+	UpdateSelectionInternal(TArray<FSelectedInstance>(), true);
+
+	GetToolManager()->EndUndoTransaction();
+
+	NewComponent->PostEditChange();
+	ParentActor->PostEditChange();
+}
+
+
+void UISMEditorTool::RebuildMeshScene()
+{
+	MeshScene = MakeShared<FMeshSceneAdapter>();
+	TArray<UActorComponent*> TempComponents;
+	for (UInstancedStaticMeshComponent* ISMC : TargetComponents)
+	{
+		TempComponents.Add(ISMC);
+	}
+	MeshScene->AddComponents(TempComponents);
+	FMeshSceneAdapterBuildOptions BuildOptions;
+	BuildOptions.bThickenThinMeshes = false;
+	BuildOptions.bFilterTinyObjects = false;
+	BuildOptions.bOnlySurfaceMaterials = false;
+	BuildOptions.bBuildSpatialDataStructures = true;
+	MeshScene->Build(BuildOptions);
+	MeshScene->BuildSpatialEvaluationCache();
+
+	AllMeshBoundingBoxes.Reset();
+	MeshScene->GetMeshBoundingBoxes(AllMeshBoundingBoxes);
+	UpdatePreviewGeometry();
+}
+
+
+void UISMEditorTool::OnTransformCompleted()
+{
+	if (bMeshSceneDirty == false)
+	{
+		MeshScene->FastUpdateTransforms(true);
+		
+		AllMeshBoundingBoxes.Reset();
+		MeshScene->GetMeshBoundingBoxes(AllMeshBoundingBoxes);
+		UpdatePreviewGeometry();
+
+		ParallelFor(CurrentSelection.Num(), [&](int32 k)
+		{
+			FSelectedInstance& Instance = CurrentSelection[k];
+			Instance.WorldBounds = MeshScene->GetMeshBoundingBox(Instance.Component, Instance.Index);
+		});
+	}
+}
+
+
+
+void UISMEditorTool::NotifySceneModified()
+{
+	bMeshSceneDirty = true;
+}
+
+void UISMEditorTool::InternalNotifySceneModified(const TArray<UInstancedStaticMeshComponent*>& ComponentList, bool bAddToScene)
+{
+	if (bAddToScene)
+	{
+		for (UInstancedStaticMeshComponent* Component : ComponentList)
+		{
+			TargetComponents.Add(Component);
+		}
+	}
+	else
+	{
+		for (UInstancedStaticMeshComponent* Component : ComponentList)
+		{
+			TargetComponents.Remove(Component);
+		}
+	}
+
+	NotifySceneModified();
+}
+
+
+void UISMEditorTool::UpdateSelectionInternal(const TArray<FSelectedInstance>& NewSelection, bool bEmitChange)
+{
+	if (bEmitChange)
+	{
+		TUniquePtr<FISMEditorSelectionChange> SelectionChange = MakeUnique<FISMEditorSelectionChange>();
+		SelectionChange->OldSelection = CurrentSelection;
+		SelectionChange->NewSelection = NewSelection;
+		GetToolManager()->EmitObjectChange(this, MoveTemp(SelectionChange), LOCTEXT("SelectionChange", "Selection Change"));
+	}
+
+	CurrentSelection = NewSelection;
+	OnSelectionUpdated();
+}
+
+void UISMEditorTool::UpdateSelectionFromUndoRedo(const TArray<FSelectedInstance>& NewSelection)
+{
+	CurrentSelection = NewSelection;
+	OnSelectionUpdated();
+}
+
+
+void FISMEditorSelectionChange::Apply(UObject* Object)
+{
+	Cast<UISMEditorTool>(Object)->UpdateSelectionFromUndoRedo(NewSelection);
+}
+
+void FISMEditorSelectionChange::Revert(UObject* Object)
+{
+	Cast<UISMEditorTool>(Object)->UpdateSelectionFromUndoRedo(OldSelection);
+}
+
+FString FISMEditorSelectionChange::ToString() const { return TEXT("FISMEditorSceneChange"); }
+
+
+
+
+
+void FISMEditorSceneChange::Apply(UObject* Object)
+{
+	Cast<UISMEditorTool>(Object)->InternalNotifySceneModified(Components, bAdded);
+}
+
+void FISMEditorSceneChange::Revert(UObject* Object)
+{
+	Cast<UISMEditorTool>(Object)->InternalNotifySceneModified(Components, !bAdded);
+}
+
+FString FISMEditorSceneChange::ToString() const { return TEXT("FISMEditorSceneChange"); }
+
+
+
+#undef LOCTEXT_NAMESPACE
diff --git a/Engine/Plugins/Experimental/MeshModelingToolsetExp/Source/MeshModelingToolsEditorOnlyExp/Public/ISMEditorTool.h b/Engine/Plugins/Experimental/MeshModelingToolsetExp/Source/MeshModelingToolsEditorOnlyExp/Public/ISMEditorTool.h
new file mode 100644
--- /dev/null
+++ b/Engine/Plugins/Experimental/MeshModelingToolsetExp/Source/MeshModelingToolsEditorOnlyExp/Public/ISMEditorTool.h
@@ -0,0 +1,302 @@
+// Copyright Epic Games, Inc. All Rights Reserved.
+
+#pragma once
+
+#include "CoreMinimal.h"
+#include "MultiSelectionTool.h"
+#include "InteractiveToolBuilder.h"
+#include "BaseBehaviors/BehaviorTargetInterfaces.h"
+#include "Changes/TransformChange.h"
+#include "FrameTypes.h"
+#include "BoxTypes.h"
+#include "BaseBehaviors/SingleClickBehavior.h"
+#include "Components/InstancedStaticMeshComponent.h"
+#include "BaseBehaviors/SingleClickOrDragBehavior.h"
+#include "Mechanics/RectangleMarqueeMechanic.h"
+#include "BaseGizmos/CombinedTransformGizmo.h"
+#include "ISMEditorTool.generated.h"
+
+class UBaseAxisTranslationGizmo;
+class UAxisAngleGizmo;
+class UDragAlignmentMechanic;
+class UTransformGizmo;
+class UTransformProxy;
+class UPreviewGeometry;
+PREDECLARE_GEOMETRY(class FMeshSceneAdapter);
+
+/**
+ *
+ */
+UCLASS()
+class MESHMODELINGTOOLSEDITORONLYEXP_API UISMEditorToolBuilder : public UInteractiveToolBuilder
+{
+	GENERATED_BODY()
+
+public:
+	virtual bool CanBuildTool(const FToolBuilderState& SceneState) const override;
+	virtual UInteractiveTool* BuildTool(const FToolBuilderState& SceneState) const override;
+};
+
+
+
+/** Mesh Transform modes */
+UENUM()
+enum class EISMEditorTransformMode : uint8
+{
+	/** Single Gizmo for all Instances */
+	SharedGizmo = 0 UMETA(DisplayName = "Shared Gizmo"),
+
+	/** Single Gizmo for all Instances, Rotations applied per-Instance */
+	SharedGizmoLocal = 1 UMETA(DisplayName = "Shared Gizmo (Local)"),
+
+	/** Separate Gizmo for each Instance */
+	PerObjectGizmo = 2 UMETA(DisplayName = "Multi-Gizmo"),
+
+	LastValue UMETA(Hidden)
+};
+
+
+/**
+ * Standard properties of the Transform Meshes operation
+ */
+UCLASS()
+class MESHMODELINGTOOLSEDITORONLYEXP_API UISMEditorToolProperties : public UInteractiveToolPropertySet
+{
+	GENERATED_BODY()
+public:
+
+	/** Transformation Mode controls the overall behavior of the Gizmos in the Tool */
+	UPROPERTY(EditAnywhere, Category = Options)
+	EISMEditorTransformMode TransformMode = EISMEditorTransformMode::SharedGizmo;
+
+	/** When true, the Gizmo can be moved independently without affecting objects. This allows the Gizmo to be repositioned before transforming. */
+	UPROPERTY(EditAnywhere, Category = Pivot, meta = (TransientToolProperty, EditCondition = "TransformMode != EISMEditorTransformMode::PerObjectGizmo") )
+	bool bSetPivotMode = false;
+
+	/** Show a highlight around all selectable instances */
+	UPROPERTY(EditAnywhere, Category = Visualization)
+	bool bShowSelectable = true;
+
+	/** Show a highlight around the selected instances */
+	UPROPERTY(EditAnywhere, Category = Visualization)
+	bool bShowSelected = true;
+
+	/** Hide the Selectable and Selected highlights when dragging with the Gizmos */
+	UPROPERTY(EditAnywhere, Category = Visualization)
+	bool bHideWhenDragging = true;
+
+	/** List of selected Component/Instance values */
+	UPROPERTY(VisibleAnywhere, Category = Selection)
+	TArray<FString> SelectedInstances;
+};
+
+
+USTRUCT()
+struct FISMEditorTarget
+{
+	GENERATED_BODY()
+
+	UPROPERTY()
+	TObjectPtr<UTransformProxy> TransformProxy = nullptr;
+
+	UPROPERTY()
+	TObjectPtr<UCombinedTransformGizmo> TransformGizmo = nullptr;
+};
+
+
+
+UENUM()
+enum class EISMEditorToolActions
+{
+	NoAction,
+	ClearSelection,
+	Delete,
+	Duplicate,
+	Replace
+};
+
+
+UCLASS()
+class MESHMODELINGTOOLSEDITORONLYEXP_API UISMEditorToolActionPropertySetBase : public UInteractiveToolPropertySet
+{
+	GENERATED_BODY()
+public:
+	TWeakObjectPtr<UISMEditorTool> ParentTool;
+
+	void Initialize(UISMEditorTool* ParentToolIn) { ParentTool = ParentToolIn; }
+
+	void PostAction(EISMEditorToolActions Action);
+};
+
+
+
+UCLASS()
+class MESHMODELINGTOOLSEDITORONLYEXP_API UISMEditorToolActionPropertySet : public UISMEditorToolActionPropertySetBase
+{
+	GENERATED_BODY()
+public:
+	/** Clear the active instance selection */
+	UFUNCTION(CallInEditor, Category = SelectionEdits, meta = (DisplayPriority = 0))
+	void ClearSelection() { PostAction(EISMEditorToolActions::ClearSelection); }
+
+	/** Delete the selected instances */
+	UFUNCTION(CallInEditor, Category = SelectionEdits, meta = (DisplayPriority = 1))
+	void Delete() { PostAction(EISMEditorToolActions::Delete); }
+
+	/** Duplicate the selected instances */
+	UFUNCTION(CallInEditor, Category = SelectionEdits, meta = (DisplayPriority = 2))
+	void Duplicate() { PostAction(EISMEditorToolActions::Duplicate); }
+};
+
+
+
+UCLASS()
+class MESHMODELINGTOOLSEDITORONLYEXP_API UISMEditorToolReplacePropertySet : public UISMEditorToolActionPropertySetBase
+{
+	GENERATED_BODY()
+public:
+	UPROPERTY(EditAnywhere, Category = Replace)
+	TObjectPtr<UStaticMesh> ReplaceWith = nullptr;
+
+	/** Clear the active instance selection */
+	UFUNCTION(CallInEditor, Category = Replace, meta = (DisplayPriority = 0))
+	void Replace() { PostAction(EISMEditorToolActions::Replace); }
+};
+
+
+
+/**
+ *
+ */
+UCLASS()
+class MESHMODELINGTOOLSEDITORONLYEXP_API UISMEditorTool : public UInteractiveTool, public IInteractiveToolCameraFocusAPI, public IClickBehaviorTarget
+{
+	GENERATED_BODY()
+
+public:
+	virtual void RegisterActions(FInteractiveToolActionSet& ActionSet) override;
+
+	virtual void SetTargets(TArray<UInstancedStaticMeshComponent*> Components);
+
+	virtual void Setup() override;
+	virtual void Shutdown(EToolShutdownType ShutdownType) override;
+
+	virtual void Render(IToolsContextRenderAPI* RenderAPI) override;
+	virtual void DrawHUD(FCanvas* Canvas, IToolsContextRenderAPI* RenderAPI) override;
+	virtual void OnTick(float DeltaTime) override;
+
+	virtual bool HasCancel() const override { return false; }
+	virtual bool HasAccept() const override { return false; }
+	virtual bool CanAccept() const override { return false; }
+
+
+	virtual FInputRayHit IsHitByClick(const FInputDeviceRay& ClickPos) override;
+	virtual void OnClicked(const FInputDeviceRay& ClickPos) override;
+
+	virtual void RequestAction(EISMEditorToolActions ActionType);
+
+	virtual void NotifySceneModified();
+
+
+	//
+	// Marquee Support
+	//
+public:
+	UPROPERTY()
+	TObjectPtr<USingleClickOrDragInputBehavior> ClickOrDragBehavior;
+
+	UPROPERTY()
+	TObjectPtr<URectangleMarqueeMechanic> RectangleMarqueeMechanic;
+
+protected:
+	void OnMarqueeRectangleFinished(const FCameraRectangle& Rectangle, bool bCancelled);
+
+protected:
+	UPROPERTY()
+	TObjectPtr<UISMEditorToolProperties> TransformProps;
+
+	UPROPERTY()
+	TObjectPtr<UISMEditorToolActionPropertySet> ToolActions;
+
+	UPROPERTY()
+	TObjectPtr<UISMEditorToolReplacePropertySet> ReplaceAction;
+
+	UPROPERTY()
+	TArray<TObjectPtr<UInstancedStaticMeshComponent>> TargetComponents;
+
+	UPROPERTY()
+	TArray<FISMEditorTarget> ActiveGizmos;
+
+	UPROPERTY()
+	TObjectPtr<UPreviewGeometry> PreviewGeometry;
+
+
+	TSharedPtr<UE::Geometry::FMeshSceneAdapter> MeshScene;
+	TArray<UE::Geometry::FAxisAlignedBox3d> AllMeshBoundingBoxes;
+	bool bMeshSceneDirty = false;
+	void RebuildMeshScene();
+
+	struct FSelectedInstance
+	{
+		UInstancedStaticMeshComponent* Component;
+		int32 Index;
+		UE::Geometry::FAxisAlignedBox3d WorldBounds;
+		bool operator==(const FSelectedInstance& Other) const { return Component == Other.Component && Index == Other.Index; }
+	};
+	TArray<FSelectedInstance> CurrentSelection;
+	void OnSelectionUpdated();
+	void UpdateSelectionInternal(const TArray<FSelectedInstance>& NewSelection, bool bEmitChange);
+	void UpdateSelectionFromUndoRedo(const TArray<FSelectedInstance>& NewSelection);
+
+	bool bInActiveDrag = false;
+	bool bShiftModifier = false;
+	bool bCtrlModifier = false;
+	void OnUpdateModifierState(int ModifierID, bool bIsOn);
+
+	EISMEditorTransformMode CurTransformMode;
+	void UpdateTransformMode(EISMEditorTransformMode NewMode);
+	void UpdateSetPivotModes(bool bEnableSetPivot);
+	void OnTransformCompleted();
+
+	void SetActiveGizmos_Single(bool bLocalRotations);
+	void SetActiveGizmos_PerObject();
+	void ResetActiveGizmos();
+
+	EISMEditorToolActions PendingAction = EISMEditorToolActions::NoAction;
+	virtual void OnClearSelection();
+	virtual void OnDeleteSelection();
+	virtual void OnDuplicateSelection();
+	virtual void OnReplaceSelection();
+
+	void UpdatePreviewGeometry();
+
+
+	virtual void InternalNotifySceneModified(const TArray<UInstancedStaticMeshComponent*>& ComponentList, bool bAddToScene);
+
+	friend class FISMEditorSelectionChange;
+	friend class FISMEditorSceneChange;
+};
+
+
+
+class MESHMODELINGTOOLSEDITORONLYEXP_API FISMEditorSelectionChange : public FToolCommandChange
+{
+public:
+	TArray<UISMEditorTool::FSelectedInstance> OldSelection;
+	TArray<UISMEditorTool::FSelectedInstance> NewSelection;
+	virtual void Apply(UObject* Object) override;
+	virtual void Revert(UObject* Object) override;
+	virtual FString ToString() const override;
+};
+
+
+
+class MESHMODELINGTOOLSEDITORONLYEXP_API FISMEditorSceneChange : public FToolCommandChange
+{
+public:
+	TArray<UInstancedStaticMeshComponent*> Components;
+	bool bAdded;
+	virtual void Apply(UObject* Object) override;
+	virtual void Revert(UObject* Object) override;
+	virtual FString ToString() const override;
+};
