ModelingTools: hide reset-to-default arrows on properties where it is not appropriate in AttributeEditorTool

diff --git a/Engine/Plugins/Experimental/MeshModelingToolsetExp/Source/MeshModelingToolsEditorOnlyExp/Public/AttributeEditorTool.h b/Engine/Plugins/Experimental/MeshModelingToolsetExp/Source/MeshModelingToolsEditorOnlyExp/Public/AttributeEditorTool.h
--- a/Engine/Plugins/Experimental/MeshModelingToolsetExp/Source/MeshModelingToolsEditorOnlyExp/Public/AttributeEditorTool.h
+++ b/Engine/Plugins/Experimental/MeshModelingToolsetExp/Source/MeshModelingToolsEditorOnlyExp/Public/AttributeEditorTool.h
@@ -70,22 +70,22 @@ class MESHMODELINGTOOLSEDITORONLYEXP_API UAttributeEditorAttribProperties : publ
 {
 	GENERATED_BODY()
 public:
-	UPROPERTY(VisibleAnywhere, Category = AttributesInspector)
+	UPROPERTY(VisibleAnywhere, Category = AttributesInspector, meta=(NoResetToDefault))
 	TArray<FString> VertexAttributes;
 
-	UPROPERTY(VisibleAnywhere, Category = AttributesInspector)
+	UPROPERTY(VisibleAnywhere, Category = AttributesInspector, meta=(NoResetToDefault))
 	TArray<FString> InstanceAttributes;
 
-	UPROPERTY(VisibleAnywhere, Category = AttributesInspector)
+	UPROPERTY(VisibleAnywhere, Category = AttributesInspector, meta=(NoResetToDefault))
 	TArray<FString> TriangleAttributes;
 
-	UPROPERTY(VisibleAnywhere, Category = AttributesInspector)
+	UPROPERTY(VisibleAnywhere, Category = AttributesInspector, meta=(NoResetToDefault))
 	TArray<FString> PolygonAttributes;
 
-	UPROPERTY(VisibleAnywhere, Category = AttributesInspector)
+	UPROPERTY(VisibleAnywhere, Category = AttributesInspector, meta=(NoResetToDefault))
 	TArray<FString> EdgeAttributes;
 
-	UPROPERTY(VisibleAnywhere, Category = AttributesInspector)
+	UPROPERTY(VisibleAnywhere, Category = AttributesInspector, meta=(NoResetToDefault))
 	TArray<FString> GroupAttributes;
 };
 
@@ -152,7 +152,7 @@ class MESHMODELINGTOOLSEDITORONLYEXP_API UAttributeEditorUVActions : public UAtt
 	GENERATED_BODY()
 public:
 
-	UPROPERTY(EditAnywhere, Category = UVs, meta = (GetOptions = GetUVLayerNamesFunc))
+	UPROPERTY(EditAnywhere, Category = UVs, meta = (GetOptions = GetUVLayerNamesFunc, NoResetToDefault))
 	FString UVLayer;
 
 	UFUNCTION()
@@ -203,11 +203,11 @@ public:
 	bool bGenerateLightmapUVs;
 
 	/** Source UV channel used to compute Lightmap UVs. Use the Static Mesh Editor to change this value. */
-	UPROPERTY(VisibleAnywhere, Category = LightmapUVs, meta = (DisplayName = "Source Channel"))
+	UPROPERTY(VisibleAnywhere, Category = LightmapUVs, meta = (DisplayName = "Source Channel", NoResetToDefault))
 	int32 SourceUVIndex;
 
 	/** Lightmap UVs are stored in this UV Channel. Use the Static Mesh Editor to change this value. */
-	UPROPERTY(VisibleAnywhere, Category = LightmapUVs, meta = (DisplayName = "Dest Channel"))
+	UPROPERTY(VisibleAnywhere, Category = LightmapUVs, meta = (DisplayName = "Dest Channel", NoResetToDefault))
 	int32 DestinationUVIndex;
 
 	UFUNCTION(CallInEditor, Category = LightmapUVs, meta = (DisplayPriority = 1))
@@ -275,7 +275,7 @@ class MESHMODELINGTOOLSEDITORONLYEXP_API UAttributeEditorModifyAttributeActions
 {
 	GENERATED_BODY()
 public:
-	UPROPERTY(EditAnywhere, Category = ModifyAttribute, meta = (GetOptions = GetAttributeNamesFunc))
+	UPROPERTY(EditAnywhere, Category = ModifyAttribute, meta = (GetOptions = GetAttributeNamesFunc, NoResetToDefault))
 	FString Attribute;
 
 	UFUNCTION()
