ModelingTools: emit warning in MeshToVolumeTool if input mesh is very large or has holes

diff --git a/Engine/Plugins/Experimental/MeshModelingToolsetExp/Source/MeshModelingToolsEditorOnlyExp/Private/MeshToVolumeTool.cpp b/Engine/Plugins/Experimental/MeshModelingToolsetExp/Source/MeshModelingToolsEditorOnlyExp/Private/MeshToVolumeTool.cpp
--- a/Engine/Plugins/Experimental/MeshModelingToolsetExp/Source/MeshModelingToolsEditorOnlyExp/Private/MeshToVolumeTool.cpp
+++ b/Engine/Plugins/Experimental/MeshModelingToolsetExp/Source/MeshModelingToolsEditorOnlyExp/Private/MeshToVolumeTool.cpp
@@ -89,6 +89,20 @@ void UMeshToVolumeTool::Setup()
 	GetToolManager()->DisplayMessage(
 		LOCTEXT("OnStartTool", "Convert a Static Mesh to a Volume, or update an existing Volume"),
 		EToolMessageLevel::UserNotification);
+
+	// check for errors in input mesh
+	bool bFoundBoundaryEdges = false;
+	for (int32 BoundaryEdgeID : InputMesh.BoundaryEdgeIndicesItr())
+	{
+		bFoundBoundaryEdges = true;
+		break;
+	}
+	if (bFoundBoundaryEdges)
+	{
+		GetToolManager()->DisplayMessage(
+			LOCTEXT("OpenBoundaryEdges", "Input Mesh is non-Closed and may produce a broken Volume"),
+			EToolMessageLevel::UserWarning);
+	}
 }
 
 void UMeshToVolumeTool::OnShutdown(EToolShutdownType ShutdownType)
@@ -188,6 +202,9 @@ void UMeshToVolumeTool::UpdateLineSet()
 
 void UMeshToVolumeTool::RecalculateVolume()
 {
+	UE::Conversion::FMeshToVolumeOptions DefaultOptions;
+	bool bShowTooLargeWarning = false;
+
 	if (Settings->ConversionMode == EMeshToVolumeMode::MinimalPolygons)
 	{
 		// Since this tool is likely to be a sink, there isn't much reason to keep
@@ -201,10 +218,25 @@ void UMeshToVolumeTool::RecalculateVolume()
 		PlanarSimplifier.SimplifyToMinimalPlanar(0.1);		// angle tolerance in degrees
 
 		UE::Conversion::GetPolygonFaces(LocalMesh, Faces, bRespectGroupBoundaries);
+
+		bShowTooLargeWarning = (LocalMesh.TriangleCount() > DefaultOptions.MaxTriangles);
 	}
 	else
 	{
 		UE::Conversion::GetTriangleFaces(InputMesh, Faces);
+
+		bShowTooLargeWarning = (InputMesh.TriangleCount() > DefaultOptions.MaxTriangles);
+	}
+
+	if (bShowTooLargeWarning)
+	{
+		GetToolManager()->DisplayMessage(
+			LOCTEXT("LargeFaceCount", "Mesh has large face count, output Volume representation may be automatically simplified"),
+			EToolMessageLevel::UserWarning);
+	}
+	else
+	{
+		GetToolManager()->DisplayMessage({}, EToolMessageLevel::UserWarning);
 	}
 
 	UpdateLineSet();
