ModelingTools: change RemeshMeshTool and SimplifyMeshTool to use UPreviewMesh OverrideMaterial to show group colors, instead of just setting a new material set

diff --git a/Engine/Plugins/Experimental/MeshModelingToolsetExp/Source/MeshModelingToolsEditorOnlyExp/Private/SimplifyMeshTool.cpp b/Engine/Plugins/Experimental/MeshModelingToolsetExp/Source/MeshModelingToolsEditorOnlyExp/Private/SimplifyMeshTool.cpp
--- a/Engine/Plugins/Experimental/MeshModelingToolsetExp/Source/MeshModelingToolsEditorOnlyExp/Private/SimplifyMeshTool.cpp
+++ b/Engine/Plugins/Experimental/MeshModelingToolsetExp/Source/MeshModelingToolsEditorOnlyExp/Private/SimplifyMeshTool.cpp
@@ -240,10 +240,9 @@ void USimplifyMeshTool::OnPropertyModified(UObject* PropertySet, FProperty* Prop
 
 void USimplifyMeshTool::UpdateVisualization()
 {
-	FComponentMaterialSet MaterialSet;
 	if (SimplifyProperties->bShowGroupColors)
 	{
-		MaterialSet.Materials = {ToolSetupUtil::GetSelectionMaterial(GetToolManager())};
+		Preview->OverrideMaterial = ToolSetupUtil::GetSelectionMaterial(GetToolManager());
 		Preview->PreviewMesh->SetTriangleColorFunction([this](const FDynamicMesh3* Mesh, int TriangleID)
 		{
 			return LinearColors::SelectFColor(Mesh->GetTriangleGroup(TriangleID));
@@ -252,14 +251,13 @@ void USimplifyMeshTool::UpdateVisualization()
 	}
 	else
 	{
-		MaterialSet = UE::ToolTarget::GetMaterialSet(Target);
+		Preview->OverrideMaterial = nullptr;
 		Preview->PreviewMesh->ClearTriangleColorFunction(UPreviewMesh::ERenderUpdateMode::FastUpdate);
 	}
-	Preview->ConfigureMaterials(MaterialSet.Materials,
-								ToolSetupUtil::GetDefaultWorkingMaterial(GetToolManager()));
 }
 
 
 
 
+
 #undef LOCTEXT_NAMESPACE
diff --git a/Engine/Plugins/Experimental/MeshModelingToolsetExp/Source/MeshModelingToolsExp/Private/RemeshMeshTool.cpp b/Engine/Plugins/Experimental/MeshModelingToolsetExp/Source/MeshModelingToolsExp/Private/RemeshMeshTool.cpp
--- a/Engine/Plugins/Experimental/MeshModelingToolsetExp/Source/MeshModelingToolsExp/Private/RemeshMeshTool.cpp
+++ b/Engine/Plugins/Experimental/MeshModelingToolsetExp/Source/MeshModelingToolsExp/Private/RemeshMeshTool.cpp
@@ -251,10 +251,9 @@ void URemeshMeshTool::OnPropertyModified(UObject* PropertySet, FProperty* Proper
 
 void URemeshMeshTool::UpdateVisualization()
 {
-	FComponentMaterialSet MaterialSet;
 	if (BasicProperties->bShowGroupColors)
 	{
-		MaterialSet.Materials = {ToolSetupUtil::GetSelectionMaterial(GetToolManager())};
+		Preview->OverrideMaterial = ToolSetupUtil::GetSelectionMaterial(GetToolManager());
 		Preview->PreviewMesh->SetTriangleColorFunction([this](const FDynamicMesh3* Mesh, int TriangleID)
 		{
 			return LinearColors::SelectFColor(Mesh->GetTriangleGroup(TriangleID));
@@ -263,11 +262,9 @@ void URemeshMeshTool::UpdateVisualization()
 	}
 	else
 	{
-		MaterialSet = UE::ToolTarget::GetMaterialSet(Targets[0]);
+		Preview->OverrideMaterial = nullptr;
 		Preview->PreviewMesh->ClearTriangleColorFunction(UPreviewMesh::ERenderUpdateMode::FastUpdate);
 	}
-	Preview->ConfigureMaterials(MaterialSet.Materials,
-								ToolSetupUtil::GetDefaultWorkingMaterial(GetToolManager()));
 }
 
 bool URemeshMeshTool::CanAccept() const
