DynamicMesh: add function MakeUniqueGroupLayerName that will make a unique group name given a base name
ModelingTools: add support for writing to selected group layer, as well as new group layer, in ConvertToPolygonsTool (GrpGen in UI). Also add support for setting groups from existing layer, which (combined) allows copying groups between layers, duplicating existing layer, and so on.

diff --git a/Engine/Plugins/Experimental/MeshModelingToolsetExp/Source/MeshModelingToolsExp/Private/ConvertToPolygonsTool.cpp b/Engine/Plugins/Experimental/MeshModelingToolsetExp/Source/MeshModelingToolsExp/Private/ConvertToPolygonsTool.cpp
--- a/Engine/Plugins/Experimental/MeshModelingToolsetExp/Source/MeshModelingToolsExp/Private/ConvertToPolygonsTool.cpp
+++ b/Engine/Plugins/Experimental/MeshModelingToolsetExp/Source/MeshModelingToolsExp/Private/ConvertToPolygonsTool.cpp
@@ -58,6 +58,9 @@ public:
 
 	// input mesh
 	TSharedPtr<FDynamicMesh3, ESPMode::ThreadSafe> OriginalMesh;
+	// input polygroups, if available
+	TSharedPtr<FPolygroupSet, ESPMode::ThreadSafe> SourcePolyGroups;
+
 	// result
 	FPolygroupsGenerator Generator;
 
@@ -70,12 +73,21 @@ public:
 
 		ResultMesh->Copy(*OriginalMesh, true, true, true, true);
 	
-
 		if (Progress && Progress->Cancelled())
 		{
 			return;
 		}
 
+		if (ConversionMode == EConvertToPolygonsMode::CopyFromLayer)
+		{
+			if (!ensure(SourcePolyGroups.IsValid())) return;
+			for (int32 tid : ResultMesh->TriangleIndicesItr())
+			{
+				ResultMesh->SetTriangleGroup(tid, SourcePolyGroups->GetTriangleGroup(tid));
+			}
+			return;
+		}
+
 		Generator = FPolygroupsGenerator(ResultMesh.Get());
 		Generator.MinGroupSize = MinGroupSize;
 
@@ -239,16 +251,11 @@ void UConvertToPolygonsTool::Setup()
 			UpdateVisualization();
 		});
 
-		// start the compute
-		PreviewCompute->InvalidateResult();
 	}
 	
 	PreviewGeometry = NewObject<UPreviewGeometry>(this);
 	PreviewGeometry->CreateInWorld(GetTargetWorld(), MeshTransform);
 
-	// updates the triangle color visualization
-	UpdateVisualization();
-
 	Settings->WatchProperty(Settings->ConversionMode, [this](EConvertToPolygonsMode) { OnSettingsModified(); });
 	Settings->WatchProperty(Settings->bShowGroupColors, [this](bool) { UpdateVisualization(); });
 	Settings->WatchProperty(Settings->AngleTolerance, [this](float) { OnSettingsModified(); });
@@ -263,6 +270,36 @@ void UConvertToPolygonsTool::Setup()
 	Settings->WatchProperty(Settings->bRespectUVSeams, [this](int) { OnSettingsModified(); });
 	Settings->WatchProperty(Settings->bRespectHardNormals, [this](int) { OnSettingsModified(); });
 
+	// add group layer picking property set for source groups
+	CopyFromLayerProperties = NewObject<UPolygroupLayersProperties>(this);
+	CopyFromLayerProperties->InitializeGroupLayers(OriginalDynamicMesh.Get());
+	CopyFromLayerProperties->WatchProperty(CopyFromLayerProperties->ActiveGroupLayer, [&](FName) { OnSelectedFromGroupLayerChanged(); });
+	AddToolPropertySource(CopyFromLayerProperties);
+	UpdateFromGroupLayer();
+	SetToolPropertySourceEnabled(CopyFromLayerProperties, Settings->ConversionMode == EConvertToPolygonsMode::CopyFromLayer);
+
+	// add picker for output group layer
+	OutputProperties = NewObject<UOutputPolygroupLayerProperties>(this);
+	AddToolPropertySource(OutputProperties);
+	OutputProperties->OptionsList.Reset();
+	OutputProperties->OptionsList.Add(TEXT("Default"));		// always have standard group
+	if (OriginalDynamicMesh->Attributes())
+	{
+		for (int32 k = 0; k < OriginalDynamicMesh->Attributes()->NumPolygroupLayers(); k++)
+		{
+			FName Name = OriginalDynamicMesh->Attributes()->GetPolygroupLayer(k)->GetName();
+			OutputProperties->OptionsList.Add(Name.ToString());
+		}
+	}
+	OutputProperties->OptionsList.Add(TEXT("Create New..."));
+	OutputProperties->WatchProperty(OutputProperties->GroupLayer, [&](FName NewName) { OutputProperties->bShowNewLayerName = (NewName == TEXT("Create New...")); });
+
+
+	// start the compute
+	PreviewCompute->InvalidateResult();
+
+	// updates the triangle color visualization
+	UpdateVisualization();
 
 	GetToolManager()->DisplayMessage(
 		LOCTEXT("OnStartTool", "Cluster triangles of the Mesh into PolyGroups using various strategies"),
@@ -286,7 +323,11 @@ void UConvertToPolygonsTool::UpdateOpParameters(FConvertToPolygonsOp& ConvertToP
 	ConvertToPolygonsOp.bRespectHardNormals = Settings->bRespectHardNormals;
 
 	ConvertToPolygonsOp.OriginalMesh = OriginalDynamicMesh;
-	
+	if (ActiveFromGroupSet.IsValid())
+	{
+		ConvertToPolygonsOp.SourcePolyGroups = ActiveFromGroupSet;
+	}
+
 	FTransform LocalToWorld = (FTransform)UE::ToolTarget::GetLocalToWorldTransform(Target);
 	ConvertToPolygonsOp.SetTransform(LocalToWorld);
 }
@@ -310,8 +351,51 @@ void UConvertToPolygonsTool::OnShutdown(EToolShutdownType ShutdownType)
 			FDynamicMesh3* DynamicMeshResult = Result.Mesh.Get();
 			if (ensure(DynamicMeshResult != nullptr))
 			{
-				// todo: have not actually modified topology here, but groups-only update is not supported yet
-				UE::ToolTarget::CommitDynamicMeshUpdate(Target, *DynamicMeshResult, true);
+				if (OutputProperties->GroupLayer != TEXT("Default"))
+				{
+					// if we want to write to any layer other than default, we have to find or create it
+					FDynamicMesh3 UseResultMesh = *OriginalDynamicMesh;
+					if (UseResultMesh.HasAttributes() == false)
+					{
+						UseResultMesh.EnableAttributes();
+					}
+
+					FDynamicMeshPolygroupAttribute* UseAttribLayer = nullptr;
+					if (OutputProperties->GroupLayer == TEXT("Create New..."))
+					{
+						// append new group layer and set it's name
+						int32 TargetLayerIdx = UseResultMesh.Attributes()->NumPolygroupLayers();
+						UseResultMesh.Attributes()->SetNumPolygroupLayers(TargetLayerIdx + 1);
+						UseAttribLayer = UseResultMesh.Attributes()->GetPolygroupLayer(TargetLayerIdx);
+						FString UseUniqueName = UE::Geometry::MakeUniqueGroupLayerName(UseResultMesh, OutputProperties->NewLayerName);
+						UseAttribLayer->SetName( FName(UseUniqueName) );
+					}
+					else
+					{
+						UseAttribLayer = UE::Geometry::FindPolygroupLayerByName(UseResultMesh, OutputProperties->GroupLayer);
+					}
+
+					if (UseAttribLayer)
+					{
+						// copy the generated groups from the Op output mesh (stored in primary groups) to the target layer
+						for (int32 tid : UseResultMesh.TriangleIndicesItr())
+						{
+							UseAttribLayer->SetValue(tid, DynamicMeshResult->GetTriangleGroup(tid));
+						}
+						UE::ToolTarget::CommitDynamicMeshUpdate(Target, UseResultMesh, true);
+					}
+					else
+					{
+						// If we can't find or create the layer (which should not be possible) the tool is going to do nothing,
+						// this is the safest option at this point
+						UE_LOG(LogGeometry, Warning, TEXT("Output group layer missing?"));
+					}
+				}
+				else
+				{
+					// todo: have not actually modified topology here, but groups-only update is not supported yet
+					UE::ToolTarget::CommitDynamicMeshUpdate(Target, *DynamicMeshResult, true);
+				}
 			}
 			GetToolManager()->EndUndoTransaction();
 		}
@@ -320,6 +404,8 @@ void UConvertToPolygonsTool::OnShutdown(EToolShutdownType ShutdownType)
 
 void UConvertToPolygonsTool::OnSettingsModified()
 {
+	SetToolPropertySourceEnabled(CopyFromLayerProperties, Settings->ConversionMode == EConvertToPolygonsMode::CopyFromLayer);
+
 	PreviewCompute->InvalidateResult();
 }
 
@@ -330,6 +416,26 @@ void UConvertToPolygonsTool::OnTick(float DeltaTime)
 }
 
 
+void UConvertToPolygonsTool::OnSelectedFromGroupLayerChanged()
+{
+	UpdateFromGroupLayer();
+	PreviewCompute->InvalidateResult();
+}
+
+void UConvertToPolygonsTool::UpdateFromGroupLayer()
+{
+	if (CopyFromLayerProperties->HasSelectedPolygroup() == false)
+	{
+		ActiveFromGroupSet = MakeShared<UE::Geometry::FPolygroupSet, ESPMode::ThreadSafe>(OriginalDynamicMesh.Get());
+	}
+	else
+	{
+		FName SelectedName = CopyFromLayerProperties->ActiveGroupLayer;
+		FDynamicMeshPolygroupAttribute* FoundAttrib = UE::Geometry::FindPolygroupLayerByName(*OriginalDynamicMesh, SelectedName);
+		ensureMsgf(FoundAttrib, TEXT("Selected Attribute Not Found! Falling back to Default group layer."));
+		ActiveFromGroupSet = MakeShared<UE::Geometry::FPolygroupSet, ESPMode::ThreadSafe>(OriginalDynamicMesh.Get(), FoundAttrib);
+	}
+}
 
 
 void UConvertToPolygonsTool::UpdateVisualization()
diff --git a/Engine/Plugins/Experimental/MeshModelingToolsetExp/Source/MeshModelingToolsExp/Public/ConvertToPolygonsTool.h b/Engine/Plugins/Experimental/MeshModelingToolsetExp/Source/MeshModelingToolsExp/Public/ConvertToPolygonsTool.h
--- a/Engine/Plugins/Experimental/MeshModelingToolsetExp/Source/MeshModelingToolsExp/Public/ConvertToPolygonsTool.h
+++ b/Engine/Plugins/Experimental/MeshModelingToolsetExp/Source/MeshModelingToolsExp/Public/ConvertToPolygonsTool.h
@@ -7,6 +7,7 @@
 #include "PreviewMesh.h"
 #include "ModelingOperators.h"
 #include "MeshOpPreviewHelpers.h"
+#include "PropertySets/PolygroupLayersProperties.h"
 #include "ConvertToPolygonsTool.generated.h"
 
 // predeclaration
@@ -43,7 +44,9 @@ enum class EConvertToPolygonsMode
 	/** Create Polygroups based on Connected Triangles */
 	FromConnectedTris UMETA(DisplayName = "From Connected Tris"),
 	/** Create Polygroups centered on well-spaced sample points, approximating a surface Voronoi diagram */
-	FromFurthestPointSampling UMETA(DisplayName = "Furthest Point Sampling")
+	FromFurthestPointSampling UMETA(DisplayName = "Furthest Point Sampling"),
+	/** Copy from existing Polygroup Layer */
+	CopyFromLayer UMETA(DisplayName = "Copy From Layer"),
 };
 
 
@@ -101,12 +104,12 @@ public:
 
 
 	/** group filtering */
-	UPROPERTY(EditAnywhere, Category = Filtering, meta = (UIMin = "1", UIMax = "100", ClampMin = "1", ClampMax = "10000"))
+	UPROPERTY(EditAnywhere, Category = Filtering, meta = (UIMin = "1", UIMax = "100", ClampMin = "1", ClampMax = "10000", EditCondition = "ConversionMode != EConvertToPolygonsMode::CopyFromLayer"))
 	int32 MinGroupSize = 2;
 
 
 	/** If true, normals are recomputed per-group, with hard edges at group boundaries */
-	UPROPERTY(EditAnywhere, Category = Output)
+	UPROPERTY(EditAnywhere, Category = Output, meta=(EditCondition = "ConversionMode != EConvertToPolygonsMode::CopyFromLayer") )
 	bool bCalculateNormals = false;
 	
 	/** Display each group with a different auto-generated color */
@@ -114,6 +117,42 @@ public:
 	bool bShowGroupColors = true;
 };
 
+
+
+
+
+
+
+UCLASS()
+class MESHMODELINGTOOLSEXP_API UOutputPolygroupLayerProperties : public UInteractiveToolPropertySet
+{
+	GENERATED_BODY()
+public:
+
+	/** Select PolyGroup layer to use. */
+	UPROPERTY(EditAnywhere, Category = "Output", meta = (DisplayName = "Output Layer", GetOptions = GetGroupOptionsList, NoResetToDefault))
+	FName GroupLayer = "Default";
+
+	// Provides set of available group layers
+	UFUNCTION()
+	TArray<FString> GetGroupOptionsList() { return OptionsList; }
+
+	// internal list used to implement above
+	UPROPERTY(meta = (TransientToolProperty))
+	TArray<FString> OptionsList;
+
+	UPROPERTY(meta = (TransientToolProperty))
+	bool bShowNewLayerName = false;
+
+	/** Name of the new Group Layer */
+	UPROPERTY(EditAnywhere, Category = "Output", meta = (TransientToolProperty, DisplayName = "New Layer Name",
+		EditCondition = "bShowNewLayerName", HideEditConditionToggle, NoResetToDefault))
+	FString NewLayerName = TEXT("polygroups");
+};
+
+
+
+
 UCLASS()
 class MESHMODELINGTOOLSEXP_API UConvertToPolygonsOperatorFactory : public UObject, public UE::Geometry::IDynamicMeshOperatorFactory
 {
@@ -158,6 +197,13 @@ protected:
 	UPROPERTY()
 	TObjectPtr<UConvertToPolygonsToolProperties> Settings;
 
+	UPROPERTY()
+	TObjectPtr<UPolygroupLayersProperties> CopyFromLayerProperties = nullptr;
+
+	UPROPERTY()
+	TObjectPtr<UOutputPolygroupLayerProperties> OutputProperties = nullptr;
+
+
 	UPROPERTY()
 	TObjectPtr<UMeshOpPreviewWithBackgroundCompute> PreviewCompute = nullptr;
 
@@ -171,4 +217,10 @@ protected:
 	TArray<int> PolygonEdges;
 	
 	void UpdateVisualization();
+
+
+	TSharedPtr<UE::Geometry::FPolygroupSet, ESPMode::ThreadSafe> ActiveFromGroupSet;
+	void OnSelectedFromGroupLayerChanged();
+	void UpdateFromGroupLayer();
+
 };
diff --git a/Engine/Plugins/Runtime/GeometryProcessing/Source/DynamicMesh/Private/Polygroups/PolygroupUtil.cpp b/Engine/Plugins/Runtime/GeometryProcessing/Source/DynamicMesh/Private/Polygroups/PolygroupUtil.cpp
--- a/Engine/Plugins/Runtime/GeometryProcessing/Source/DynamicMesh/Private/Polygroups/PolygroupUtil.cpp
+++ b/Engine/Plugins/Runtime/GeometryProcessing/Source/DynamicMesh/Private/Polygroups/PolygroupUtil.cpp
@@ -74,4 +74,49 @@ int32 UE::Geometry::ComputeGroupIDBound(const FDynamicMesh3& Mesh, const FDynami
 		Bound = FMath::Max(Bound, Layer->GetValue(TID) + 1);
 	}
 	return Bound;
+}
+
+
+FString UE::Geometry::MakeUniqueGroupLayerName(const FDynamicMesh3& Mesh, FString BaseName)
+{
+	if (BaseName.Len() == 0)
+	{
+		BaseName = TEXT("group");
+	}
+
+	if (Mesh.HasAttributes() == false)
+	{
+		return BaseName;
+	}
+	const FDynamicMeshAttributeSet* AttribSet = Mesh.Attributes();
+	int32 NumGroupLayers = AttribSet->NumPolygroupLayers();
+	if (NumGroupLayers == 0)
+	{
+		return BaseName;
+	}
+
+	// TODO: would be nice to detect if BaseName has a numeric suffix in either basename_0 or basename0 style
+	// and increment the number. But we should have a general utility function for that because it is so common...
+
+	FString UniqueName = BaseName;
+	int32 NumberCounter = 0;
+	while (true)
+	{
+		bool bFoundDuplicate = false;
+		for (int32 k = 0; k < NumGroupLayers; ++k)
+		{
+			if ( AttribSet->GetPolygroupLayer(k)->GetName() == FName(UniqueName) )
+			{
+				bFoundDuplicate = true;
+			}
+		}
+		if (!bFoundDuplicate)
+		{
+			return UniqueName;
+		}
+		UniqueName = FString::Printf(TEXT("%s_%d"), *BaseName, NumberCounter++);
+	}
+	ensureMsgf(false, TEXT("Failed to create unique name, returning base name"));
+	return BaseName;
+
 }
\ No newline at end of file
diff --git a/Engine/Plugins/Runtime/GeometryProcessing/Source/DynamicMesh/Public/Polygroups/PolygroupUtil.h b/Engine/Plugins/Runtime/GeometryProcessing/Source/DynamicMesh/Public/Polygroups/PolygroupUtil.h
--- a/Engine/Plugins/Runtime/GeometryProcessing/Source/DynamicMesh/Public/Polygroups/PolygroupUtil.h
+++ b/Engine/Plugins/Runtime/GeometryProcessing/Source/DynamicMesh/Public/Polygroups/PolygroupUtil.h
@@ -38,5 +38,10 @@ DYNAMICMESH_API int32 FindPolygroupLayerIndex(const FDynamicMesh3& Mesh, const F
  */
 DYNAMICMESH_API int32 ComputeGroupIDBound(const FDynamicMesh3& Mesh, const FDynamicMeshPolygroupAttribute* Layer);
 
+/**
+* @return unique group layer name derived from BaseName by appending a numeric suffix in the form _0, _1, etc
+*/
+DYNAMICMESH_API FString MakeUniqueGroupLayerName(const FDynamicMesh3& Mesh, FString BaseName);
+
 }	// end namespace Geometry
 }	// end namespace UE
\ No newline at end of file
