ModelingTools: add support for specifying polygroup layer in EditNormalsTool

diff --git a/Engine/Plugins/Experimental/MeshModelingToolsetExp/Source/MeshModelingToolsExp/Private/EditNormalsTool.cpp b/Engine/Plugins/Experimental/MeshModelingToolsetExp/Source/MeshModelingToolsExp/Private/EditNormalsTool.cpp
--- a/Engine/Plugins/Experimental/MeshModelingToolsetExp/Source/MeshModelingToolsExp/Private/EditNormalsTool.cpp
+++ b/Engine/Plugins/Experimental/MeshModelingToolsetExp/Source/MeshModelingToolsExp/Private/EditNormalsTool.cpp
@@ -6,16 +6,13 @@
 
 #include "ToolSetupUtil.h"
 #include "ModelingToolTargetUtil.h"
+#include "Polygroups/PolygroupUtil.h"
 
 #include "DynamicMesh/DynamicMesh3.h"
-#include "BaseBehaviors/MultiClickSequenceInputBehavior.h"
-#include "Selection/SelectClickedAction.h"
 
 #include "MeshDescriptionToDynamicMesh.h"
 #include "DynamicMeshToMeshDescription.h"
 
-#include "InteractiveGizmoManager.h"
-
 #include "AssetUtils/MeshDescriptionUtil.h"
 #include "Engine/StaticMesh.h"
 #include "Components/StaticMeshComponent.h"
@@ -86,7 +83,22 @@ void UEditNormalsTool::Setup()
 	// initialize the PreviewMesh+BackgroundCompute object
 	UpdateNumPreviews();
 
-	
+	// if editing normals on a single object, user can select from available polygroup layers in from-polygroup mode
+	if (Previews.Num() == 1)
+	{
+		PolygroupLayerProperties = NewObject<UPolygroupLayersProperties>(this);
+		PolygroupLayerProperties->RestoreProperties(this, TEXT("EditNormalsTool"));
+		PolygroupLayerProperties->InitializeGroupLayers(OriginalDynamicMeshes[0].Get());
+		PolygroupLayerProperties->WatchProperty(PolygroupLayerProperties->ActiveGroupLayer, [&](FName) { OnSelectedGroupLayerChanged(); });
+		AddToolPropertySource(PolygroupLayerProperties);
+		UpdateActiveGroupLayer();
+
+		BasicProperties->WatchProperty(BasicProperties->SplitNormalMethod, [&](ESplitNormalMethod NewMethod)
+		{
+			SetToolPropertySourceEnabled(PolygroupLayerProperties, NewMethod == ESplitNormalMethod::FaceGroupID);
+		});
+		SetToolPropertySourceEnabled(PolygroupLayerProperties, BasicProperties->SplitNormalMethod == ESplitNormalMethod::FaceGroupID);
+	}
 	
 	for (UMeshOpPreviewWithBackgroundCompute* Preview : Previews)
 	{
@@ -145,6 +157,7 @@ void UEditNormalsTool::OnShutdown(EToolShutdownType ShutdownType)
 {
 	BasicProperties->SaveProperties(this);
 	AdvancedProperties->SaveProperties(this);
+	PolygroupLayerProperties->SaveProperties(this, TEXT("EditNormalsTool"));
 
 	// Restore (unhide) the source meshes
 	for (int32 ComponentIdx = 0, NumTargets = Targets.Num(); ComponentIdx < NumTargets; ComponentIdx++)
@@ -176,6 +189,12 @@ TUniquePtr<FDynamicMeshOperator> UEditNormalsOperatorFactory::MakeNewOperator()
 
 	const FTransform LocalToWorld = (FTransform) UE::ToolTarget::GetLocalToWorldTransform(Tool->Targets[ComponentIndex]);
 	NormalsOp->OriginalMesh = Tool->OriginalDynamicMeshes[ComponentIndex];
+
+	// use custom polygroup layer if available
+	if (ComponentIndex == 0 && Tool->OriginalDynamicMeshes.Num() == 1 && Tool->ActiveGroupSet.IsValid())
+	{
+		NormalsOp->MeshPolygroups = Tool->ActiveGroupSet;
+	};
 	
 	NormalsOp->SetTransform(LocalToWorld);
 
@@ -238,6 +257,36 @@ bool UEditNormalsTool::CanAccept() const
 }
 
 
+void UEditNormalsTool::OnSelectedGroupLayerChanged()
+{
+	UpdateActiveGroupLayer();
+	for (UMeshOpPreviewWithBackgroundCompute* Preview : Previews)
+	{
+		Preview->InvalidateResult();
+	}
+}
+
+
+void UEditNormalsTool::UpdateActiveGroupLayer()
+{
+	// do not update if more than one mesh
+	if (OriginalDynamicMeshes.Num() > 1) return;
+
+	if (PolygroupLayerProperties->HasSelectedPolygroup() == false)
+	{
+		ActiveGroupSet = MakeShared<UE::Geometry::FPolygroupSet, ESPMode::ThreadSafe>(OriginalDynamicMeshes[0].Get());
+	}
+	else
+	{
+		FName SelectedName = PolygroupLayerProperties->ActiveGroupLayer;
+		FDynamicMeshPolygroupAttribute* FoundAttrib = UE::Geometry::FindPolygroupLayerByName(*OriginalDynamicMeshes[0], SelectedName);
+		ensureMsgf(FoundAttrib, TEXT("Selected Attribute Not Found! Falling back to Default group layer."));
+		ActiveGroupSet = MakeShared<UE::Geometry::FPolygroupSet, ESPMode::ThreadSafe>(OriginalDynamicMeshes[0].Get(), FoundAttrib);
+	}
+}
+
+
+
 void UEditNormalsTool::GenerateAsset(const TArray<FDynamicMeshOpResult>& Results)
 {
 	GetToolManager()->BeginUndoTransaction(LOCTEXT("EditNormalsToolTransactionName", "Edit Normals Tool"));
diff --git a/Engine/Plugins/Experimental/MeshModelingToolsetExp/Source/MeshModelingToolsExp/Public/EditNormalsTool.h b/Engine/Plugins/Experimental/MeshModelingToolsetExp/Source/MeshModelingToolsExp/Public/EditNormalsTool.h
--- a/Engine/Plugins/Experimental/MeshModelingToolsetExp/Source/MeshModelingToolsExp/Public/EditNormalsTool.h
+++ b/Engine/Plugins/Experimental/MeshModelingToolsetExp/Source/MeshModelingToolsExp/Public/EditNormalsTool.h
@@ -3,13 +3,13 @@
 #pragma once
 
 #include "CoreMinimal.h"
-#include "UObject/NoExportTypes.h"
 #include "BaseTools/MultiSelectionMeshEditingTool.h"
 #include "InteractiveToolBuilder.h"
 #include "MeshOpPreviewHelpers.h"
 #include "CleaningOps/EditNormalsOp.h"
 #include "DynamicMesh/DynamicMesh3.h"
-#include "BaseTools/SingleClickTool.h"
+#include "PropertySets/PolygroupLayersProperties.h"
+#include "Polygroups/PolygroupSet.h"
 #include "EditNormalsTool.generated.h"
 
 
@@ -148,11 +148,13 @@ public:
 protected:
 
 	UPROPERTY()
-	TObjectPtr<UEditNormalsToolProperties> BasicProperties;
+	TObjectPtr<UEditNormalsToolProperties> BasicProperties = nullptr;
 
 	UPROPERTY()
-	TObjectPtr<UEditNormalsAdvancedProperties> AdvancedProperties;
+	TObjectPtr<UEditNormalsAdvancedProperties> AdvancedProperties = nullptr;
 
+	UPROPERTY()
+	TObjectPtr<UPolygroupLayersProperties> PolygroupLayerProperties = nullptr;
 
 	UPROPERTY()
 	TArray<TObjectPtr<UMeshOpPreviewWithBackgroundCompute>> Previews;
@@ -166,4 +168,9 @@ protected:
 	void UpdateNumPreviews();
 
 	void GenerateAsset(const TArray<FDynamicMeshOpResult>& Results);
+
+
+	TSharedPtr<UE::Geometry::FPolygroupSet, ESPMode::ThreadSafe> ActiveGroupSet;
+	void OnSelectedGroupLayerChanged();
+	void UpdateActiveGroupLayer();
 };
diff --git a/Engine/Plugins/Runtime/MeshModelingToolset/Source/ModelingOperators/Private/CleaningOps/EditNormalsOp.cpp b/Engine/Plugins/Runtime/MeshModelingToolset/Source/ModelingOperators/Private/CleaningOps/EditNormalsOp.cpp
--- a/Engine/Plugins/Runtime/MeshModelingToolset/Source/ModelingOperators/Private/CleaningOps/EditNormalsOp.cpp
+++ b/Engine/Plugins/Runtime/MeshModelingToolset/Source/ModelingOperators/Private/CleaningOps/EditNormalsOp.cpp
@@ -103,9 +103,12 @@ void FEditNormalsOp::CalculateResult(FProgressCancel* Progress)
 		}
 		else // SplitNormalMethod == ESplitNormalMethod::FaceGroupID
 		{
-			ResultMesh->Attributes()->PrimaryNormals()->CreateFromPredicate([this](int VID, int TA, int TB)
+			FPolygroupSet DefaultGroups(ResultMesh.Get());
+			FPolygroupSet& UsePolygroups = (MeshPolygroups.IsValid()) ? *MeshPolygroups : DefaultGroups;
+
+			ResultMesh->Attributes()->PrimaryNormals()->CreateFromPredicate([this, &UsePolygroups](int VID, int TA, int TB)
 			{
-				return ResultMesh->GetTriangleGroup(TA) == ResultMesh->GetTriangleGroup(TB);
+				return UsePolygroups.GetTriangleGroup(TA) == UsePolygroups.GetTriangleGroup(TB);
 			}, 0);
 		}
 	}
diff --git a/Engine/Plugins/Runtime/MeshModelingToolset/Source/ModelingOperators/Public/CleaningOps/EditNormalsOp.h b/Engine/Plugins/Runtime/MeshModelingToolset/Source/ModelingOperators/Public/CleaningOps/EditNormalsOp.h
--- a/Engine/Plugins/Runtime/MeshModelingToolset/Source/ModelingOperators/Public/CleaningOps/EditNormalsOp.h
+++ b/Engine/Plugins/Runtime/MeshModelingToolset/Source/ModelingOperators/Public/CleaningOps/EditNormalsOp.h
@@ -5,6 +5,7 @@
 #include "CoreMinimal.h"
 #include "Util/ProgressCancel.h"
 #include "ModelingOperators.h"
+#include "Polygroups/PolygroupSet.h"
 
 #include "EditNormalsOp.generated.h"
 
@@ -48,6 +49,7 @@ public:
 
 	// inputs
 	TSharedPtr<FDynamicMesh3, ESPMode::ThreadSafe> OriginalMesh;
+	TSharedPtr<UE::Geometry::FPolygroupSet, ESPMode::ThreadSafe> MeshPolygroups;
 
 	bool bFixInconsistentNormals;
 	bool bInvertNormals;
