Fix VS2022 lambda error by marking constexpr variable as static

#rnx

diff --git a/Engine/Plugins/Experimental/MeshModelingToolsetExp/Source/MeshModelingToolsEditorOnlyExp/Private/ISMEditorTool.cpp b/Engine/Plugins/Experimental/MeshModelingToolsetExp/Source/MeshModelingToolsEditorOnlyExp/Private/ISMEditorTool.cpp
--- a/Engine/Plugins/Experimental/MeshModelingToolsetExp/Source/MeshModelingToolsEditorOnlyExp/Private/ISMEditorTool.cpp
+++ b/Engine/Plugins/Experimental/MeshModelingToolsetExp/Source/MeshModelingToolsEditorOnlyExp/Private/ISMEditorTool.cpp
@@ -319,7 +319,7 @@ void UISMEditorTool::DrawHUD(FCanvas* Canvas, IToolsContextRenderAPI* RenderAPI)
 void UISMEditorTool::UpdatePreviewGeometry()
 {
 	FColor LinesColor(50, 225, 50);
-	constexpr int BoxFaces[6][4] =
+	static constexpr int BoxFaces[6][4] =
 	{
 		{ 0, 1, 2, 3 },     // back, -z
 		{ 5, 4, 7, 6 },     // front, +z
