PatternTool: Add random transform functionality to the pattern tool.

The translation and rotation jitter are relatively intuitive, but for scaling there are a couple of ways it could be done. From the user's perspective, scale jittering is accomplished in the same way as translation and rotation, but the scale along any axis cannot be negative.

The randomness is seeded so that the pattern is not completely different every time any parameter is changed (as long as the seed isn't changed, changing the scale along the z axis will just result in changes to the elements' scales along the z-axis, no other changes occur) and to allow the user to cycle through various patterns generated using a given set of parameters for the tool without needing to slightly move the pattern or alter any of the parameters to indirectly cause a call to MarkPatternDirty(). Instead, MarkPatternDirty() is called any time the seed changes.

diff --git a/Engine/Plugins/Experimental/MeshModelingToolsetExp/Source/MeshModelingToolsExp/Private/PatternTool.cpp b/Engine/Plugins/Experimental/MeshModelingToolsetExp/Source/MeshModelingToolsExp/Private/PatternTool.cpp
--- a/Engine/Plugins/Experimental/MeshModelingToolsetExp/Source/MeshModelingToolsExp/Private/PatternTool.cpp
+++ b/Engine/Plugins/Experimental/MeshModelingToolsetExp/Source/MeshModelingToolsExp/Private/PatternTool.cpp
@@ -74,17 +74,27 @@ public:
 
 public:
 	// input settings
+	FRandomStream RotationRandomStream;
+	FRandomStream ScaleRandomStream;
+	FRandomStream TranslationRandomStream;
+	
 	FQuaterniond StartRotation = FQuaterniond::Identity();
 	FQuaterniond EndRotation = FQuaterniond::Identity();
-	bool bInterpolateRotation = false;			// if false, only StartRotation is used
-
+	FRotator RotationJitterRange = FRotator::ZeroRotator;	// using an FRotator so that the sampling can be done about 3 axes individually and converted to a quat when needed
+	bool bInterpolateRotation = false;						// if false, only StartRotation is used
+	bool bJitterRotation = false;
+	
 	FVector3d StartTranslation = FVector3d::Zero();
 	FVector3d EndTranslation = FVector3d::Zero();
-	bool bInterpolateTranslation = false;		// if false, only StartTranslation is used
+	FVector3d TranslationJitterRange = FVector3d::Zero();
+	bool bInterpolateTranslation = false;					// if false, only StartTranslation is used
+	bool bJitterTranslation = false;
 
 	FVector3d StartScale = FVector3d::One();
 	FVector3d EndScale = FVector3d::One();
-	bool bInterpolateScale = false;				// if false, only StartScale is used
+	FVector3d ScaleJitterRange = FVector3d::Zero();
+	bool bInterpolateScale = false;							// if false, only StartScale is used
+	bool bJitterScale = false;
 
 	FAxisAlignedBox3d Dimensions = FAxisAlignedBox3d(FVector3d::Zero(), 10.0);
 
@@ -135,6 +145,46 @@ public:
 		{
 			Transform.SetTranslation( Transform.GetTranslation() + StartTranslation );
 		}
+		
+		// Lerp operations are performed on a vector-component basis because each component should be able to independently
+		// vary without respect for the other two components. This is true for rotation, scale, and translation.
+		if (bJitterRotation)
+		{
+			FRotator RotationJitter;
+
+			// TODO: maybe RotationJitterRange should have user-definable lower and upper bound like scale.
+			RotationJitter.Pitch = FMath::Lerp(-RotationJitterRange.Pitch, RotationJitterRange.Pitch, RotationRandomStream.GetFraction());
+			RotationJitter.Roll  = FMath::Lerp(-RotationJitterRange.Roll,  RotationJitterRange.Roll,  RotationRandomStream.GetFraction());
+			RotationJitter.Yaw   = FMath::Lerp(-RotationJitterRange.Yaw,   RotationJitterRange.Yaw,   RotationRandomStream.GetFraction());
+
+			Transform.SetRotation(Transform.GetRotation() * FQuaterniond(RotationJitter));
+		}
+
+		if (bJitterScale)
+		{
+			FVector3d ScaleJitter;
+			FVector3d TransformScale = Transform.GetScale();
+
+			constexpr double MIN_SCALE = 0.001f;
+			ScaleJitter.X = FMath::Max(MIN_SCALE, TransformScale.X + FMath::Lerp(-ScaleJitterRange.X, ScaleJitterRange.X, ScaleRandomStream.GetFraction()));
+			ScaleJitter.Y = FMath::Max(MIN_SCALE, TransformScale.Y + FMath::Lerp(-ScaleJitterRange.Y, ScaleJitterRange.Y, ScaleRandomStream.GetFraction()));
+			ScaleJitter.Z = FMath::Max(MIN_SCALE, TransformScale.Z + FMath::Lerp(-ScaleJitterRange.Z, ScaleJitterRange.Z, ScaleRandomStream.GetFraction()));
+				
+			Transform.SetScale( ScaleJitter );
+		}
+
+
+		if (bJitterTranslation)
+		{
+			FVector3d TranslationJitter;
+
+			// TODO: maybe TranslationJitterRange should have user-definable lower and upper bound like scale.
+			TranslationJitter.X = FMath::Lerp(-TranslationJitterRange.X, TranslationJitterRange.X, TranslationRandomStream.GetFraction());
+			TranslationJitter.Y = FMath::Lerp(-TranslationJitterRange.Y, TranslationJitterRange.Y, TranslationRandomStream.GetFraction());
+			TranslationJitter.Z = FMath::Lerp(-TranslationJitterRange.Z, TranslationJitterRange.Z, TranslationRandomStream.GetFraction());
+			
+			Transform.SetTranslation( Transform.GetTranslation() + TranslationJitter);
+		}
 	}
 
 
@@ -457,6 +507,7 @@ void UPatternTool::Setup()
 	Settings->WatchProperty(Settings->SingleAxis, [this](EPatternToolSingleAxis) { OnParametersUpdated(); } );
 	Settings->WatchProperty(Settings->SinglePlane, [this](EPatternToolSinglePlane) { OnParametersUpdated(); } );
 	Settings->WatchProperty(Settings->bHideSources, [this](bool bNewValue) { OnSourceVisibilityToggled(!bNewValue); } );
+	Settings->WatchProperty(Settings->Seed, [this](int32 NewSeed) { MarkPatternDirty(); } );
 
 	LinearSettings = NewObject<UPatternTool_LinearSettings>();
 	AddToolPropertySource(LinearSettings);
@@ -729,7 +780,15 @@ void UPatternTool::ResetPreviews()
 
 static void InitializeGenerator(FPatternGenerator& Generator, UPatternTool* Tool)
 {
+	// The way this is seeded is kind of arbitrary but is reliable for the purpose of deterministically
+	// providing each type of jitter a random stream independent of the others given a single input seed.
+	Generator.RotationRandomStream.Initialize(Tool->Settings->Seed);
+	Generator.ScaleRandomStream.Initialize(Generator.RotationRandomStream.GetUnsignedInt());
+	Generator.TranslationRandomStream.Initialize(Generator.RotationRandomStream.GetUnsignedInt());
+	
 	Generator.bInterpolateRotation = Tool->RotationSettings->bInterpolate;
+	Generator.bJitterRotation = Tool->RotationSettings->bJitter;
+	Generator.RotationJitterRange = Tool->RotationSettings->RotationJitterRange;
 	if (Generator.bInterpolateRotation)
 	{
 		Generator.StartRotation = FQuaterniond(Tool->RotationSettings->StartRotation);
@@ -741,6 +800,8 @@ static void InitializeGenerator(FPatternGenerator& Generator, UPatternTool* Tool
 	}
 
 	Generator.bInterpolateTranslation = Tool->TranslationSettings->bInterpolate;
+	Generator.bJitterTranslation = Tool->TranslationSettings->bJitter;
+	Generator.TranslationJitterRange = Tool->TranslationSettings->TranslationJitterRange;
 	if (Generator.bInterpolateTranslation)
 	{
 		Generator.StartTranslation = Tool->TranslationSettings->StartTranslation;
@@ -752,6 +813,8 @@ static void InitializeGenerator(FPatternGenerator& Generator, UPatternTool* Tool
 	}
 
 	Generator.bInterpolateScale = Tool->ScaleSettings->bInterpolate;
+	Generator.bJitterScale = Tool->ScaleSettings->bJitter;
+	Generator.ScaleJitterRange = (Tool->ScaleSettings->bUniform) ? (FVector3d::One()*Tool->ScaleSettings->ScaleJitterRange) : Tool->ScaleSettings->ScaleJitterRangeNonUniform;
 	if (Generator.bInterpolateScale)
 	{
 		Generator.StartScale = (Tool->ScaleSettings->bUniform) ? (FVector3d::One()*Tool->ScaleSettings->StartScale) : Tool->ScaleSettings->StartScaleNonUniform;
@@ -880,7 +943,7 @@ void UPatternTool::GetPatternTransforms_Radial(TArray<UE::Geometry::FTransformSR
 void UPatternTool::UpdatePattern()
 {
 	TArray<FTransformSRT3d> Pattern;
-
+	
 	if (Settings->Shape == EPatternToolShape::Line)
 	{
 		GetPatternTransforms_Linear(Pattern);
@@ -893,7 +956,7 @@ void UPatternTool::UpdatePattern()
 	{
 		GetPatternTransforms_Radial(Pattern);
 	}
-
+	
 	// Return all current preview components in use to the preview component pool
 	ResetPreviews();
 
@@ -1028,7 +1091,7 @@ UDynamicMeshComponent* UPatternTool::GetPreviewDynamicMesh(FPatternElement& Elem
 
 	FDynamicMesh3 ElementMeshCopy(Element.SourceDynamicMesh->GetMeshRef());
 	DynamicMeshComp->SetMesh(MoveTemp(ElementMeshCopy));
-
+	
 	AllComponents.Add(DynamicMeshComp);
 
 	return DynamicMeshComp;
diff --git a/Engine/Plugins/Experimental/MeshModelingToolsetExp/Source/MeshModelingToolsExp/Public/PatternTool.h b/Engine/Plugins/Experimental/MeshModelingToolsetExp/Source/MeshModelingToolsExp/Public/PatternTool.h
--- a/Engine/Plugins/Experimental/MeshModelingToolsetExp/Source/MeshModelingToolsExp/Public/PatternTool.h
+++ b/Engine/Plugins/Experimental/MeshModelingToolsetExp/Source/MeshModelingToolsExp/Public/PatternTool.h
@@ -114,6 +114,9 @@ public:
 	UPROPERTY(EditAnywhere, Category = Shape)
 	bool bHideSources = true;
 
+	/** The seed used to introduce random transform variations when enabled */
+	UPROPERTY(EditAnywhere, Category = Shape)
+	int32 Seed = FMath::Rand();
 };
 
 
@@ -259,6 +262,10 @@ public:
 	UPROPERTY(EditAnywhere, Category = Rotation)
 	bool bInterpolate = false;
 
+	/** If true, Rotation at each Pattern Element is offset by a uniformly chosen random value in the range of [-RotationJitterRange, RotationJitterRange] */
+	UPROPERTY(EditAnywhere, Category = Rotation)
+	bool bJitter = false;
+
 	/** Rotation at first Pattern Element */
 	UPROPERTY(EditAnywhere, Category = Rotation, meta = (EditCondition = "bInterpolate", EditConditionHides, HideEditConditionToggle))
 	FRotator StartRotation = FRotator::ZeroRotator;
@@ -266,6 +273,10 @@ public:
 	/** Rotation applied to all Pattern Elements, or at Last Pattern Element for Interpolated rotations */
 	UPROPERTY(EditAnywhere, Category = Rotation)
 	FRotator Rotation = FRotator::ZeroRotator;
+
+	/** Upper bound of the range which is sampled to randomly rotate each Pattern Element if bJitter is true */
+	UPROPERTY(EditAnywhere, Category = Rotation, meta = (ClampMin = 0, EditCondition = "bJitter", EditConditionHides, HideEditConditionToggle))
+	FRotator RotationJitterRange = FRotator::ZeroRotator;
 };
 
 
@@ -281,6 +292,10 @@ public:
 	UPROPERTY(EditAnywhere, Category = Translation)
 	bool bInterpolate = false;
 
+	/** If true, Translation at each Pattern Element is offset by a uniformly chosen random value in the range of [-TranslationJitterRange, TranslationJitterRange] */
+	UPROPERTY(EditAnywhere, Category = Translation)
+	bool bJitter = false;
+
 	/** Translation at first Pattern Element */
 	UPROPERTY(EditAnywhere, Category = Translation, meta = (EditCondition = "bInterpolate", EditConditionHides, HideEditConditionToggle))
 	FVector StartTranslation = FVector::ZeroVector;
@@ -288,6 +303,10 @@ public:
 	/** Translation applied to all Pattern Elements, or at Last Pattern Element for Interpolated translations */
 	UPROPERTY(EditAnywhere, Category = Translation)
 	FVector Translation = FVector::ZeroVector;
+
+	/** Upper bound of the range which is sampled to randomly rotate each Pattern Element if bJitter is true */
+	UPROPERTY(EditAnywhere, Category = Translation, meta = (ClampMin = 0, EditCondition = "bJitter", EditConditionHides, HideEditConditionToggle))
+	FVector TranslationJitterRange = FVector::ZeroVector;
 };
 
 
@@ -307,6 +326,10 @@ public:
 	UPROPERTY(EditAnywhere, Category = Scale)
 	bool bInterpolate = false;
 
+	/** If true, Scale at each Pattern Element is offset by a uniformly chosen random value in the range of [-ScaleJitterRange, ScaleJitterRange] */
+	UPROPERTY(EditAnywhere, Category = Scale)
+	bool bJitter = false;
+	
 	/** Uniform Scale at first Pattern Element */
 	UPROPERTY(EditAnywhere, Category = Scale, meta = (EditCondition = "bInterpolate && bUniform", EditConditionHides, HideEditConditionToggle))
 	float StartScale = 1;
@@ -315,6 +338,10 @@ public:
 	UPROPERTY(EditAnywhere, Category = Scale, meta = (EditCondition = "bUniform", EditConditionHides, HideEditConditionToggle))
 	float Scale = 1;
 
+	/** Upper bound of the range which is sampled to randomly scale each Pattern Element if bJitter is true */
+	UPROPERTY(EditAnywhere, Category = Scale, meta = (ClampMin = 0, EditCondition = "bJitter && bUniform", EditConditionHides, HideEditConditionToggle))
+	float ScaleJitterRange = 0.0f;
+
 	/** Non-Uniform Scale at first Pattern Element */
 	UPROPERTY(EditAnywhere, Category = Scale, meta = (DisplayName = "StartScale", EditCondition = "bInterpolate && bUniform == false", EditConditionHides, HideEditConditionToggle))
 	FVector StartScaleNonUniform = FVector(1,1,1);
@@ -322,6 +349,10 @@ public:
 	/** Non-Uniform Scale applied to all Pattern Elements, or at Last Pattern Element for Interpolated scales */
 	UPROPERTY(EditAnywhere, Category = Scale, meta = (DisplayName = "Scale", EditCondition = "bUniform == false", EditConditionHides, HideEditConditionToggle))
 	FVector ScaleNonUniform = FVector(1,1,1);
+
+	/** Upper bound of the range which is sampled to randomly scale each Pattern Element if bJitter is true (Non-Uniform) */
+	UPROPERTY(EditAnywhere, Category = Scale, meta = (ClampMin = 0, DisplayName = "Scale Jitter Range", EditCondition = "bJitter && bUniform == false", EditConditionHides, HideEditConditionToggle))
+	FVector ScaleJitterRangeNonUniform = FVector(0, 0, 0);
 };
 
 
@@ -438,7 +469,7 @@ protected:
 	void GetPatternTransforms_Grid(TArray<UE::Geometry::FTransformSRT3d>& TransformsOut);
 	void GetPatternTransforms_Radial(TArray<UE::Geometry::FTransformSRT3d>& TransformsOut);
 
-
+	
 	struct FPatternElement
 	{
 		int32 TargetIndex = 0;
