Modeling Mode: SubdividePolyTool - make sure the number of subdivisions respects the min of 1

diff --git a/Engine/Plugins/Experimental/MeshModelingToolsetExp/Source/MeshModelingToolsEditorOnlyExp/Private/SubdividePolyTool.cpp b/Engine/Plugins/Experimental/MeshModelingToolsetExp/Source/MeshModelingToolsEditorOnlyExp/Private/SubdividePolyTool.cpp
--- a/Engine/Plugins/Experimental/MeshModelingToolsetExp/Source/MeshModelingToolsEditorOnlyExp/Private/SubdividePolyTool.cpp
+++ b/Engine/Plugins/Experimental/MeshModelingToolsetExp/Source/MeshModelingToolsEditorOnlyExp/Private/SubdividePolyTool.cpp
@@ -211,7 +211,10 @@ void USubdividePolyTool::Setup()
 		return;
 	}
 
-	check(Properties->SubdivisionLevel >= 1);	// Should be enforced by UPROPERTY meta tags
+	if (!ensure(Properties->SubdivisionLevel >= 1)) // Should be enforced by UPROPERTY meta tags
+	{
+		Properties->SubdivisionLevel = 1;
+	}
 	
 	CapSubdivisionLevel(Properties->SubdivisionScheme, Properties->SubdivisionLevel);
 	PreviewDynamicMeshComponent->SetRenderMeshPostProcessor(MakeUnique<SubdivPostProcessor>(Properties->SubdivisionLevel,
