fix handling of non-uniform scaling in mirror tool by baking scales before mirroring + update warning language to just suggest baking scale.  Also add new warning for zero scales, because that case is actually impossible to handle well in general.

diff --git a/Engine/Plugins/Experimental/MeshModelingToolsetExp/Source/MeshModelingToolsExp/Private/MirrorTool.cpp b/Engine/Plugins/Experimental/MeshModelingToolsetExp/Source/MeshModelingToolsExp/Private/MirrorTool.cpp
--- a/Engine/Plugins/Experimental/MeshModelingToolsetExp/Source/MeshModelingToolsExp/Private/MirrorTool.cpp
+++ b/Engine/Plugins/Experimental/MeshModelingToolsetExp/Source/MeshModelingToolsExp/Private/MirrorTool.cpp
@@ -14,6 +14,7 @@
 #include "Misc/MessageDialog.h"
 #include "ToolBuilderUtil.h"
 #include "ToolSetupUtil.h"
+#include "DynamicMesh/MeshTransforms.h"
 
 #include "TargetInterfaces/MaterialProvider.h"
 #include "TargetInterfaces/PrimitiveComponentBackedTarget.h"
@@ -23,6 +24,20 @@ using namespace UE::Geometry;
 
 #define LOCTEXT_NAMESPACE "UMirrorTool"
 
+namespace MirrorTool_Local
+{
+	FTransform WithoutScale(FTransform T)
+	{
+		T.SetScale3D(FVector::One());
+		return T;
+	}
+	FTransform OnlyScale(const FTransform& TransformIn)
+	{
+		FTransform T = FTransform::Identity;
+		T.SetScale3D(TransformIn.GetScale3D());
+		return T;
+	}
+}
 
 // Tool builder functions
 
@@ -46,22 +61,9 @@ TUniquePtr<FDynamicMeshOperator> UMirrorOperatorFactory::MakeNewOperator()
 	MirrorOp->PlaneTolerance = MirrorTool->Settings->PlaneTolerance;
 	MirrorOp->bAllowBowtieVertexCreation = MirrorTool->Settings->bAllowBowtieVertexCreation;
 
-	FTransform LocalToWorld = (FTransform) UE::ToolTarget::GetLocalToWorldTransform(MirrorTool->Targets[ComponentIndex]);
+	FTransform LocalToWorld = MirrorTool_Local::WithoutScale((FTransform) UE::ToolTarget::GetLocalToWorldTransform(MirrorTool->Targets[ComponentIndex]));
 	MirrorOp->SetTransform(LocalToWorld);
 
-	// We also need WorldToLocal. Threshold the LocalToWorld scaling transform so we can get the inverse.
-	FVector LocalToWorldScale = LocalToWorld.GetScale3D();
-	for (int i = 0; i < 3; i++)
-	{
-		float DimScale = FMathf::Abs(LocalToWorldScale[i]);
-		float Tolerance = KINDA_SMALL_NUMBER;
-		if (DimScale < Tolerance)
-		{
-			LocalToWorldScale[i] = Tolerance * FMathf::SignNonZero(LocalToWorldScale[i]);
-		}
-	}
-	LocalToWorld.SetScale3D(LocalToWorldScale);
-
 	// Now we can get the plane parameters in local space.
 	MirrorOp->LocalPlaneOrigin = LocalToWorld.InverseTransformPosition(MirrorTool->MirrorPlaneOrigin);;
 
@@ -177,6 +179,9 @@ void UMirrorTool::Setup()
 		TSharedPtr<FDynamicMesh3, ESPMode::ThreadSafe> DynamicMesh = MakeShared<FDynamicMesh3, ESPMode::ThreadSafe>();
 		FMeshDescriptionToDynamicMesh Converter;
 		Converter.Convert(UE::ToolTarget::GetMeshDescription(Targets[i]), *DynamicMesh);
+		// Bake the scale part of the transform
+		FTransform Transform = (FTransform)UE::ToolTarget::GetLocalToWorldTransform(Targets[i]);
+		MeshTransforms::ApplyTransform(*DynamicMesh, MirrorTool_Local::OnlyScale(Transform), true);
 
 		// Wrap the dynamic mesh in a replacement change target
 		UDynamicMeshReplacementChangeTarget* WrappedTarget = MeshesToMirror.Add_GetRef(NewObject<UDynamicMeshReplacementChangeTarget>());
@@ -269,7 +274,8 @@ void UMirrorTool::SetupPreviews()
 
 		// Set initial preview to unprocessed mesh, so that things don't disappear initially
 		Preview->PreviewMesh->UpdatePreview(MeshesToMirror[PreviewIndex]->GetMesh().Get());
-		Preview->PreviewMesh->SetTransform((FTransform) UE::ToolTarget::GetLocalToWorldTransform(Targets[PreviewIndex]));
+		FTransform Transform = (FTransform)UE::ToolTarget::GetLocalToWorldTransform(Targets[PreviewIndex]);
+		Preview->PreviewMesh->SetTransform(MirrorTool_Local::WithoutScale(Transform));
 		Preview->SetVisibility(Settings->bShowPreview);
 	}
 }
@@ -278,7 +284,7 @@ void UMirrorTool::CheckAndDisplayWarnings()
 {
 	// We can have more than one warning, which makes this a bit more work.
 	FText SameSourceWarning;
-	FText NonUniformScaleWarning;
+	FText ScaleWarning;
 
 	// See if any of the selected components have the same source.
 	TArray<int32> MapToFirstOccurrences;
@@ -293,37 +299,49 @@ void UMirrorTool::CheckAndDisplayWarnings()
 
 	// See if any of the selected components have a nonuniform scaling transform.
 	IPrimitiveComponentBackedTarget* NonUniformScalingTarget = nullptr;
+	IPrimitiveComponentBackedTarget* ZeroScalingTarget = nullptr;
 	for (int32 i = 0; i < Targets.Num(); ++i)
 	{
 		IPrimitiveComponentBackedTarget* Component = Cast<IPrimitiveComponentBackedTarget>(Targets[i]);
 		const FVector Scaling = Component->GetWorldTransform().GetScale3D();
+		if (Scaling.X == 0 || Scaling.Y == 0 || Scaling.Z == 0)
+		{
+			ZeroScalingTarget = Component;
+			break;
+		}
 		if (Scaling.X != Scaling.Y || Scaling.Y != Scaling.Z)
 		{
 			NonUniformScalingTarget = Component;
-			break;
+			// don't break; continue in case we see a ZeroScalingTarget (which is worse)
 		}
 	}
 
-	if (NonUniformScalingTarget)
+	if (ZeroScalingTarget)
+	{
+		ScaleWarning = FText::Format(
+			LOCTEXT("MirrorZeroScaledAsset", "WARNING: The item \"{0}\" has a zero-scale on at least one axis. Mirroring cannot be correctly applied in this case. Consider instead baking the scale before mirroring."),
+			FText::FromString(ZeroScalingTarget->GetOwnerActor()->GetName()));
+	}
+	else if (NonUniformScalingTarget) // Only show the non-uniform-scale warning if the more-severe zero-scale warning does not apply
 	{
-		NonUniformScaleWarning = FText::Format(
-			LOCTEXT("MirrorNonUniformScaledAsset", "WARNING: The item \"{0}\" has a non-uniform scaling transform. This is not supported because mirroring acts on the underlying mesh, and mirroring is not commutative with non-uniform scaling. Consider deforming the mesh rather than scaling it non-uniformly."),
+		ScaleWarning = FText::Format(
+			LOCTEXT("MirrorNonUniformScaledAsset", "WARNING: The item \"{0}\" has a non-uniform scaling transform. The mirror will be applied in world-space, so the underlying asset will not have mirror symmetry. Consider instead baking the scale before mirroring."),
 			FText::FromString(NonUniformScalingTarget->GetOwnerActor()->GetName()));
 	}
 
-	if (bAnyHaveSameSource && NonUniformScalingTarget)
+	if (bAnyHaveSameSource && (NonUniformScalingTarget || ZeroScalingTarget))
 	{
 		// Concatenates the two warnings with an extra line in between.
 		GetToolManager()->DisplayMessage(FText::Format(LOCTEXT("CombinedWarnings", "{0}\n\n{1}"),
-			SameSourceWarning, NonUniformScaleWarning), EToolMessageLevel::UserWarning);
+			SameSourceWarning, ScaleWarning), EToolMessageLevel::UserWarning);
 	}
 	else if (bAnyHaveSameSource)
 	{
 		GetToolManager()->DisplayMessage(SameSourceWarning, EToolMessageLevel::UserWarning);
 	}
-	else if (NonUniformScalingTarget)
+	else if (NonUniformScalingTarget || ZeroScalingTarget)
 	{
-		GetToolManager()->DisplayMessage(NonUniformScaleWarning, EToolMessageLevel::UserWarning);
+		GetToolManager()->DisplayMessage(ScaleWarning, EToolMessageLevel::UserWarning);
 	}
 }
 
@@ -344,9 +362,12 @@ void UMirrorTool::OnShutdown(EToolShutdownType ShutdownType)
 	{
 		// Gather results
 		TArray<FDynamicMeshOpResult> Results;
-		for (UMeshOpPreviewWithBackgroundCompute* Preview : Previews)
+		for (int32 PreviewIndex = 0; PreviewIndex < Previews.Num(); ++PreviewIndex)
 		{
+			UMeshOpPreviewWithBackgroundCompute* Preview = Previews[PreviewIndex];
+			FTransform Transform = (FTransform)UE::ToolTarget::GetLocalToWorldTransform(Targets[PreviewIndex]);
 			Results.Emplace(Preview->Shutdown());
+			MeshTransforms::ApplyTransformInverse(*(Results.Last().Mesh), MirrorTool_Local::OnlyScale(Transform), true);
 		}
 
 		// Convert to output. This will also edit the selection.
