make attribute editor's attribute lists sorted, so the order doesn't randomly shift during editing
+ minor refactor to expose less internal stuff to the header

diff --git a/Engine/Plugins/Experimental/MeshModelingToolsetExp/Source/MeshModelingToolsEditorOnlyExp/Private/AttributeEditorTool.cpp b/Engine/Plugins/Experimental/MeshModelingToolsetExp/Source/MeshModelingToolsEditorOnlyExp/Private/AttributeEditorTool.cpp
--- a/Engine/Plugins/Experimental/MeshModelingToolsetExp/Source/MeshModelingToolsEditorOnlyExp/Private/AttributeEditorTool.cpp
+++ b/Engine/Plugins/Experimental/MeshModelingToolsetExp/Source/MeshModelingToolsEditorOnlyExp/Private/AttributeEditorTool.cpp
@@ -126,6 +126,14 @@ void UAttributeEditorTool::Setup()
 }
 
 
+struct FAttributeEditorAttribInfo
+{
+	FName Name;
+	EAttributeEditorElementType ElementType;
+	EAttributeEditorAttribType DataType;
+	bool bIsAutoGenerated;
+};
+
 template<typename AttribSetType>
 void ExtractAttribList(const FMeshDescription* Mesh, AttribSetType& AttribSet, EAttributeEditorElementType ElemType, TArray<FAttributeEditorAttribInfo>& AttribList, TArray<FString>& StringList)
 {
@@ -142,6 +150,7 @@ void ExtractAttribList(const FMeshDescription* Mesh, AttribSetType& AttribSet, E
 		AttribInfo.Name = AttributeName;
 		AttribInfo.ElementType = ElemType;
 		AttribInfo.DataType = EAttributeEditorAttribType::Unknown;
+		AttribInfo.bIsAutoGenerated = (AttributesRef.GetFlags() & EMeshAttributeFlags::AutoGenerated) != EMeshAttributeFlags::None;
 		if (AttribSet.template HasAttributeOfType<int32>(AttributeName))
 		{
 			AttribInfo.DataType = EAttributeEditorAttribType::Int32;
@@ -171,14 +180,19 @@ void ExtractAttribList(const FMeshDescription* Mesh, AttribSetType& AttribSet, E
 			AttribInfo.DataType = EAttributeEditorAttribType::String;
 		}
 		AttribList.Add(AttribInfo);
+	});
+	AttribList.Sort([](const FAttributeEditorAttribInfo& A, const FAttributeEditorAttribInfo& B)
+	{
+		return A.Name.LexicalLess(B.Name);
+	});
 
-		bool bIsAutoGen = (AttributesRef.GetFlags() & EMeshAttributeFlags::AutoGenerated) != EMeshAttributeFlags::None;
-		
-		FString UIString = (bIsAutoGen) ?
+	for (const FAttributeEditorAttribInfo& AttribInfo : AttribList)
+	{	
+		FString UIString = (AttribInfo.bIsAutoGenerated) ?
 			FString::Printf(TEXT("%s - %s (autogen)"), *(AttribInfo.Name.ToString()), *EnumStrings[(int32)AttribInfo.DataType])
 			: FString::Printf(TEXT("%s - %s"), *(AttribInfo.Name.ToString()), *EnumStrings[(int32)AttribInfo.DataType]);
 		StringList.Add(UIString);
-	});
+	}
 }
 
 
@@ -305,12 +319,13 @@ void UAttributeEditorTool::InitializeAttributeLists()
 		}
 	}
 
-	ExtractAttribList(Mesh, Mesh->VertexAttributes(), EAttributeEditorElementType::Vertex, this->VertexAttributes, AttributeProps->VertexAttributes);
-	ExtractAttribList(Mesh, Mesh->VertexInstanceAttributes(), EAttributeEditorElementType::VertexInstance, this->InstanceAttributes, AttributeProps->InstanceAttributes);
-	ExtractAttribList(Mesh, Mesh->TriangleAttributes(), EAttributeEditorElementType::Triangle, this->TriangleAttributes, AttributeProps->TriangleAttributes);
-	ExtractAttribList(Mesh, Mesh->PolygonAttributes(), EAttributeEditorElementType::Polygon, this->PolygonAttributes, AttributeProps->PolygonAttributes);
-	ExtractAttribList(Mesh, Mesh->EdgeAttributes(), EAttributeEditorElementType::Edge, this->EdgeAttributes, AttributeProps->EdgeAttributes);
-	ExtractAttribList(Mesh, Mesh->PolygonGroupAttributes(), EAttributeEditorElementType::PolygonGroup, this->GroupAttributes, AttributeProps->GroupAttributes);
+	TArray<FAttributeEditorAttribInfo> VertexAttributes, InstanceAttributes, TriangleAttributes, PolygonAttributes, EdgeAttributes, GroupAttributes;
+	ExtractAttribList(Mesh, Mesh->VertexAttributes(), EAttributeEditorElementType::Vertex, VertexAttributes, AttributeProps->VertexAttributes);
+	ExtractAttribList(Mesh, Mesh->VertexInstanceAttributes(), EAttributeEditorElementType::VertexInstance, InstanceAttributes, AttributeProps->InstanceAttributes);
+	ExtractAttribList(Mesh, Mesh->TriangleAttributes(), EAttributeEditorElementType::Triangle, TriangleAttributes, AttributeProps->TriangleAttributes);
+	ExtractAttribList(Mesh, Mesh->PolygonAttributes(), EAttributeEditorElementType::Polygon, PolygonAttributes, AttributeProps->PolygonAttributes);
+	ExtractAttribList(Mesh, Mesh->EdgeAttributes(), EAttributeEditorElementType::Edge, EdgeAttributes, AttributeProps->EdgeAttributes);
+	ExtractAttribList(Mesh, Mesh->PolygonGroupAttributes(), EAttributeEditorElementType::PolygonGroup, GroupAttributes, AttributeProps->GroupAttributes);
 
 	TArray<FString> OldAttributeNames = ModifyAttributeProps->AttributeNamesList;
 
@@ -323,7 +338,7 @@ void UAttributeEditorTool::InitializeAttributeLists()
 	//	&this->TriangleAttributes, & this->PolygonAttributes,
 	//	& this->EdgeAttributes, & this->GroupAttributes };
 	TArray<TArray<FAttributeEditorAttribInfo>*> AttribInfos = {
-		&this->VertexAttributes, &this->PolygonAttributes, &this->TriangleAttributes };
+		&VertexAttributes, &PolygonAttributes, &TriangleAttributes };
 
 	for (TArray<FAttributeEditorAttribInfo>* AttribInfoList : AttribInfos)
 	{
diff --git a/Engine/Plugins/Experimental/MeshModelingToolsetExp/Source/MeshModelingToolsEditorOnlyExp/Public/AttributeEditorTool.h b/Engine/Plugins/Experimental/MeshModelingToolsetExp/Source/MeshModelingToolsEditorOnlyExp/Public/AttributeEditorTool.h
--- a/Engine/Plugins/Experimental/MeshModelingToolsetExp/Source/MeshModelingToolsEditorOnlyExp/Public/AttributeEditorTool.h
+++ b/Engine/Plugins/Experimental/MeshModelingToolsetExp/Source/MeshModelingToolsEditorOnlyExp/Public/AttributeEditorTool.h
@@ -57,14 +57,6 @@ enum class EAttributeEditorAttribType : uint8
 };
 
 
-struct FAttributeEditorAttribInfo
-{
-	FName Name;
-	EAttributeEditorElementType ElementType;
-	EAttributeEditorAttribType DataType;
-};
-
-
 UCLASS()
 class MESHMODELINGTOOLSEDITORONLYEXP_API UAttributeEditorAttribProperties : public UInteractiveToolPropertySet
 {
@@ -372,13 +364,6 @@ protected:
 protected:
 	bool bTargetIsStaticMesh = true;
 
-	TArray<FAttributeEditorAttribInfo> VertexAttributes;
-	TArray<FAttributeEditorAttribInfo> InstanceAttributes;
-	TArray<FAttributeEditorAttribInfo> TriangleAttributes;
-	TArray<FAttributeEditorAttribInfo> PolygonAttributes;
-	TArray<FAttributeEditorAttribInfo> EdgeAttributes;
-	TArray<FAttributeEditorAttribInfo> GroupAttributes;
-
 	bool bAttributeListsValid = false;
 	bool bHaveAutoGeneratedLightmapUVSet = false;
 	void InitializeAttributeLists();
