ModelingTools: Fix undo/redo when Ctrl+middle clicking in CubeGr. Add undo/redo support for actions while in corner mode. Refresh detail panel when initalizing from actor. Fix old selection being visible when clicking away from everything.

diff --git a/Engine/Plugins/Experimental/MeshModelingToolsetExp/Source/MeshModelingToolsExp/Private/CubeGridTool.cpp b/Engine/Plugins/Experimental/MeshModelingToolsetExp/Source/MeshModelingToolsExp/Private/CubeGridTool.cpp
--- a/Engine/Plugins/Experimental/MeshModelingToolsetExp/Source/MeshModelingToolsExp/Private/CubeGridTool.cpp
+++ b/Engine/Plugins/Experimental/MeshModelingToolsetExp/Source/MeshModelingToolsExp/Private/CubeGridTool.cpp
@@ -56,6 +56,7 @@ namespace CubeGridToolLocals
 
 	const FText SelectionChangeTransactionName = LOCTEXT("SelectionChangeTransaction", "Cube Grid Selection Change");
 	const FText ModeChangeTransactionName = LOCTEXT("ModeChangeTransaction", "Cube Grid Mode Change");
+	const FText CornerModeExtrudeAmountChangeTransactionName = LOCTEXT("CornerExtrudeAmountTransaction", "Corner Push/Pull Amount");
 
 	const FString PropertyCacheIdentifier = TEXT("CubeGridTool");
 
@@ -194,6 +195,75 @@ namespace CubeGridToolLocals
 		}
 	};
 
+	class FCornerModeExtrudeAmountChange : public FToolCommandChange
+	{
+	public:
+		FCornerModeExtrudeAmountChange(int32 ExtrudeAmountBeforeIn, int32 ExtrudeAmountAfterIn) 
+			: ExtrudeAmountBefore(ExtrudeAmountBeforeIn)
+			, ExtrudeAmountAfter(ExtrudeAmountAfterIn)
+		
+		{};
+
+		virtual bool HasExpired(UObject* Object) const override
+		{
+			return !Cast<UCubeGridTool>(Object)->IsInCornerMode();
+		}
+
+		virtual void Apply(UObject* Object) override
+		{
+			Cast<UCubeGridTool>(Object)->SetCurrentExtrudeAmount(ExtrudeAmountAfter);
+		}
+		virtual void Revert(UObject* Object) override
+		{
+			Cast<UCubeGridTool>(Object)->SetCurrentExtrudeAmount(ExtrudeAmountBefore);
+		}
+
+		virtual FString ToString() const override
+		{
+			return TEXT("CubeGridToolLocals::FCornerModeExtrudeAmountChange");
+		}
+
+	protected:
+		int32 ExtrudeAmountBefore = 0;
+		int32 ExtrudeAmountAfter = 0;
+	};
+
+	class FCornerModeSelectedCornerChange : public FToolCommandChange
+	{
+	public:
+		FCornerModeSelectedCornerChange(bool CornerFlagsBeforeIn[4], bool CornerFlagsAfterIn[4])
+		{
+			for (int i = 0; i < 4; ++i)
+			{
+				CornerFlagsBefore[i] = CornerFlagsBeforeIn[i];
+				CornerFlagsAfter[i] = CornerFlagsAfterIn[i];
+			}
+		};
+
+		virtual bool HasExpired(UObject* Object) const override
+		{
+			return !Cast<UCubeGridTool>(Object)->IsInCornerMode();
+		}
+
+		virtual void Apply(UObject* Object) override
+		{
+			Cast<UCubeGridTool>(Object)->SetCornerSelection(CornerFlagsAfter);
+		}
+		virtual void Revert(UObject* Object) override
+		{
+			Cast<UCubeGridTool>(Object)->SetCornerSelection(CornerFlagsBefore);
+		}
+
+		virtual FString ToString() const override
+		{
+			return TEXT("CubeGridToolLocals::FCornerModeSelectedCornerChange");
+		}
+
+	protected:
+		bool CornerFlagsBefore[4] = { false,false,false,false };
+		bool CornerFlagsAfter[4] = { false,false,false,false };
+	};
+
 	/** Undoes setting bChangesMade, which is used to determine whether the target needs saving. */
 	class FCubeGridChangesMadeChange : public FToolCommandChange
 	{
@@ -974,15 +1044,15 @@ void UCubeGridTool::Setup()
 
 	GridFrameOriginWatcherIdx = Settings->WatchProperty(Settings->GridFrameOrigin,
 		[this](FVector) {
-			UpdateGridTransform(FTransform(Settings->GridFrameOrientation, Settings->GridFrameOrigin, FVector::One()), 
-				// Update gizmo, don't update detail panel
-				true, false);
+			FTransform Transform(Settings->GridFrameOrientation, Settings->GridFrameOrigin, FVector::One());
+			UpdateGridGizmo(Transform, /* bSilentlyUpdate */ true); // Silent to avoid undo/redo during drag in detail panel
+			UpdateGridTransform(Transform, /*bUpdateDetailPanel */ false);
 		});
 	GridFrameOrientationWatcherIdx = Settings->WatchProperty(Settings->GridFrameOrientation,
 		[this](FRotator) {
-			UpdateGridTransform(FTransform(Settings->GridFrameOrientation, Settings->GridFrameOrigin, FVector::One()), 
-				// Update gizmo, don't update detail panel
-				true, false);
+			FTransform Transform(Settings->GridFrameOrientation, Settings->GridFrameOrigin, FVector::One());
+			UpdateGridGizmo(Transform, /* bSilentlyUpdate */ true); // Silent to avoid undo/redo during drag in detail panel
+			UpdateGridTransform(Transform, /*bUpdateDetailPanel */ false);
 		});
 
 	Settings->SilentUpdateWatched();
@@ -1240,9 +1310,10 @@ void UCubeGridTool::ClearHover()
 
 void  UCubeGridTool::GridGizmoMoved(UTransformProxy* Proxy, FTransform Transform)
 {
-	UpdateGridTransform(Transform, 
-		// Don't update gizmo, update detail panel
-		false, true);
+	UpdateGridTransform(Transform, /* bUpdateDetailPanel */ true,
+		// If we're here due to a drag, then we don't want to trigger the detail panel
+		// rebuild (it will be triggered on drag end).
+		!bInGizmoDrag);
 }
 
 void UCubeGridTool::OnCtrlMiddleClick(const FInputDeviceRay& ClickPos)
@@ -1313,9 +1384,7 @@ void UCubeGridTool::OnCtrlMiddleClick(const FInputDeviceRay& ClickPos)
 	}
 
 	// Move the gizmo to that corner.
-	UpdateGridTransform(FTransform(Settings->GridFrameOrientation, CubeGrid->ToWorldPoint(Corners[ClosestCornerIndex]), FVector::One())
-		// Update both gizmo and detail panel
-		, true, true);
+	UpdateGridGizmo(FTransform(Settings->GridFrameOrientation, CubeGrid->ToWorldPoint(Corners[ClosestCornerIndex]), FVector::One()));
 
 	GetToolManager()->EndUndoTransaction();
 }
@@ -1443,13 +1512,20 @@ void UCubeGridTool::UpdateGizmoVisibility(bool bVisible)
 	Settings->SilentUpdateWatched();
 }
 
-void UCubeGridTool::UpdateGridTransform(FTransform NewTransform, bool bUpdateGizmo, bool bUpdateDetailPanel)
+void UCubeGridTool::UpdateGridGizmo(const FTransform& NewTransform, bool bSilentlyUpdate)
 {
-	if (bUpdateGizmo)
+	if (bSilentlyUpdate)
 	{
 		GridGizmo->ReinitializeGizmoTransform(NewTransform);
 	}
+	else
+	{
+		GridGizmo->SetNewGizmoTransform(NewTransform);
+	}
+}
 
+void UCubeGridTool::UpdateGridTransform(const FTransform& NewTransform, bool bUpdateDetailPanel, bool bTriggerDetailPanelRebuild)
+{
 	if (bUpdateDetailPanel)
 	{
 		Settings->GridFrameOrigin = NewTransform.GetTranslation();
@@ -1457,6 +1533,11 @@ void UCubeGridTool::UpdateGridTransform(FTransform NewTransform, bool bUpdateGiz
 
 		Settings->GridFrameOrientation = NewTransform.GetRotation().Rotator();
 		Settings->SilentUpdateWatcherAtIndex(GridFrameOrientationWatcherIdx);
+
+		if (bTriggerDetailPanelRebuild)
+		{
+			NotifyOfPropertyChangeByTool(Settings);
+		}
 	}
 
 	CubeGrid->SetGridFrame(FFrame3d(NewTransform));
@@ -1468,6 +1549,7 @@ void UCubeGridTool::UpdateGridTransform(FTransform NewTransform, bool bUpdateGiz
 
 		if (!bInGizmoDrag)
 		{
+			// Don't need to update corner hit points while dragging since we can't drag and click corners at the same time.
 			UpdateCornerGeometrySet();
 		}
 	}
@@ -1609,6 +1691,7 @@ void UCubeGridTool::OnClickPress(const FInputDeviceRay& PressPos)
 		// Clear selection (the event emit, if needed, happens on click release) 
 		bHaveSelection = false;
 		Selection.Box = FAxisAlignedBox3d();
+		UpdateSelectionLineSet();
 	}
 }
 
@@ -1661,7 +1744,27 @@ void UCubeGridTool::OnClickDrag(const FInputDeviceRay& DragPos)
 
 void UCubeGridTool::OnClickRelease(const FInputDeviceRay& ReleasePos)
 {
-	if (MouseState == EMouseState::DraggingExtrudeDistance)
+	using namespace CubeGridToolLocals;
+
+	if (MouseState == EMouseState::DraggingCornerSelection)
+	{
+		bool bCornerSelectionChanged = false;
+		for (int i = 0; i < 4; ++i)
+		{
+			if (PreDragCornerSelectedFlags[i] != CornerSelectedFlags[i])
+			{
+				bCornerSelectionChanged = true;
+				break;
+			}
+		}
+		if (bCornerSelectionChanged)
+		{
+			GetToolManager()->EmitObjectChange(this,
+				MakeUnique<FCornerModeSelectedCornerChange>(PreDragCornerSelectedFlags, CornerSelectedFlags),
+				LOCTEXT("CornerSelectionTransaction", "Corner Selection"));
+		}
+	}
+	else if (MouseState == EMouseState::DraggingExtrudeDistance)
 	{
 		// Only apply result if we're not in corner mode, because in corner mode
 		// we apply when exiting corner mode (that behavior is particularly important 
@@ -1673,6 +1776,12 @@ void UCubeGridTool::OnClickRelease(const FInputDeviceRay& ReleasePos)
 			bBlockUntilPreviewUpdate = false;
 			bAdjustSelectionOnPreviewUpdate = true;
 		}
+		else if (Mode == EMode::Corner && CurrentExtrudeAmount != DragStartExtrudeAmount)
+		{
+			GetToolManager()->EmitObjectChange(this, 
+				MakeUnique<FCornerModeExtrudeAmountChange>(DragStartExtrudeAmount, CurrentExtrudeAmount), 
+				CornerModeExtrudeAmountChangeTransactionName);
+		}
 	}
 	else if (MouseState == EMouseState::DraggingRegularSelection)
 	{
@@ -1993,16 +2102,25 @@ void UCubeGridTool::ApplyPushPull(int32 NumBlocks)
 		return;
 	}
 
+	int32 PreviousExtrudeAmount = CurrentExtrudeAmount;
 	CurrentExtrudeAmount += NumBlocks;
 
 	InvalidatePreview();
 
-	if (Mode != EMode::Corner)
+	if (Mode == EMode::PushPull)
 	{
 		bWaitingToApplyPreview = true;
 		bBlockUntilPreviewUpdate = false;
 		bAdjustSelectionOnPreviewUpdate = true;
 	}
+	else if (Mode == EMode::Corner)
+	{
+		// In corner mode, we don't actually appy the change yet, but we transact the extrude amount.
+		// (in regular mode we transact the mesh change once the actual change is applied)
+		GetToolManager()->EmitObjectChange(this,
+			MakeUnique<FCornerModeExtrudeAmountChange>(PreviousExtrudeAmount, CurrentExtrudeAmount),
+			CornerModeExtrudeAmountChangeTransactionName);
+	}
 }
 
 void UCubeGridTool::SetGridPowerClamped(int32 GridPower)
@@ -2147,9 +2265,7 @@ void UCubeGridTool::ApplyAction(ECubeGridToolAction ActionType)
 	case ECubeGridToolAction::ResetFromActor:
 		if (ToolActions->GridSourceActor)
 		{
-			UpdateGridTransform(ToolActions->GridSourceActor->GetTransform(), 
-				// Update gizmo and detail panel
-				true, true);
+			UpdateGridGizmo(ToolActions->GridSourceActor->GetTransform());
 		}
 		break;
 	}
@@ -2394,6 +2510,11 @@ bool UCubeGridTool::IsInDefaultMode() const
 	return Mode == EMode::PushPull;
 }
 
+bool UCubeGridTool::IsInCornerMode() const
+{
+	return Mode == EMode::Corner;
+}
+
 void UCubeGridTool::RevertToDefaultMode()
 {
 	if (Mode == EMode::Corner)
@@ -2411,6 +2532,23 @@ void UCubeGridTool::SetChangesMade(bool bChangesMadeIn)
 	bChangesMade = bChangesMadeIn;
 }
 
+// For use by Undo/Redo during corner mode
+void UCubeGridTool::SetCurrentExtrudeAmount(int32 ExtrudeAmount)
+{
+	CurrentExtrudeAmount = ExtrudeAmount;
+	InvalidatePreview();
+}
+
+// For use by Undo/Redo during corner mode
+void UCubeGridTool::SetCornerSelection(bool CornerSelectedFlagsIn[4])
+{
+	for (int i = 0; i < 4; ++i)
+	{
+		CornerSelectedFlags[i] = CornerSelectedFlagsIn[i];
+	}
+	InvalidatePreview();
+}
+
 bool UCubeGridTool::CanCurrentlyNestedCancel()
 {
 	return Mode == EMode::Corner || bHaveSelection;
diff --git a/Engine/Plugins/Experimental/MeshModelingToolsetExp/Source/MeshModelingToolsExp/Public/CubeGridTool.h b/Engine/Plugins/Experimental/MeshModelingToolsetExp/Source/MeshModelingToolsExp/Public/CubeGridTool.h
--- a/Engine/Plugins/Experimental/MeshModelingToolsetExp/Source/MeshModelingToolsExp/Public/CubeGridTool.h
+++ b/Engine/Plugins/Experimental/MeshModelingToolsetExp/Source/MeshModelingToolsExp/Public/CubeGridTool.h
@@ -337,7 +337,7 @@ public:
 		// Both of these boxes are in the coordinate space of the (unscaled) grid frame.
 		UE::Geometry::FAxisAlignedBox3d Box;
 		UE::Geometry::FAxisAlignedBox3d StartBox; // Box delineating original selected face
-		
+
 		// Direction must be initialized to a valid enum value (0 is not one).
 		FCubeGrid::EFaceDirection Direction = FCubeGrid::EFaceDirection::PositiveX;
 
@@ -379,8 +379,11 @@ public:
 	// Used by undo/redo
 	virtual void UpdateUsingMeshChange(const UE::Geometry::FDynamicMeshChange& MeshChange, bool bRevert);
 	virtual bool IsInDefaultMode() const;
+	virtual bool IsInCornerMode() const;
 	virtual void RevertToDefaultMode();
 	virtual void SetChangesMade(bool bChangesMadeIn);
+	virtual void SetCurrentExtrudeAmount(int32 ExtrudeAmount);
+	virtual void SetCornerSelection(bool CornerSelectedFlags[4]);
 
 	// IClickDragBehaviorTarget
 	virtual FInputRayHit CanBeginClickDragSequence(const FInputDeviceRay& PressPos) override;
@@ -418,7 +421,7 @@ protected:
 
 	UPROPERTY()
 	TObjectPtr<UDragAlignmentMechanic> GridGizmoAlignmentMechanic = nullptr;
-	
+
 	UPROPERTY()
 	TObjectPtr<UTransformProxy> GridGizmoTransformProxy = nullptr;
 
@@ -427,7 +430,7 @@ protected:
 
 	UPROPERTY()
 	TObjectPtr<UClickDragInputBehavior> ClickDragBehavior = nullptr;
-	
+
 	UPROPERTY()
 	TObjectPtr<UMouseHoverBehavior> HoverBehavior = nullptr;
 
@@ -555,9 +558,28 @@ protected:
 	int32 GridFrameOrientationWatcherIdx;
 	void GridGizmoMoved(UTransformProxy* Proxy, FTransform Transform);
 	void UpdateGizmoVisibility(bool bVisible);
-	void UpdateGridTransform(FTransform NewTransform, bool bUpdateGizmo, bool bUpdateDetailPanel);
 	bool bInGizmoDrag = false;
 
+	/*
+	 * Updates the gizmo controlling the cube grid transform. Only updates the cube grid itself if
+	 * bSilentlyUpdate is false.
+	 * 
+	 * @param bSilentlyUpdate If true, gizmo is just repositioned without triggering any callback
+	 *   or emitting an undo transaction. If false, emits transactions and triggers the GridGizmoMoved
+	 *   callback (which will trigger UpdateGridTransform).
+	 */
+	void UpdateGridGizmo(const FTransform& NewTransform, bool bSilentlyUpdate = false);
+	
+	/*
+	* Updates the cube grid.
+	* 
+	 * @param bUpdateDetailPanel Not needed if updating from the detail panel, otherwise should be true.
+	 * @param bTriggerDetailPanelRebuild Should be false if updating from a drag to avoid the costly rebuild, 
+	 * but otherwise should be true to update the edit conditions and "revert to default" arrows. Not relevant
+	 * if bUpdateDetailPanel is false.
+	 */
+	void UpdateGridTransform(const FTransform& NewTransform, bool bUpdateDetailPanel = true, bool bTriggerDetailPanelRebuild = true);
+	
 	FVector3d MiddleClickDragStart;
 	FInputRayHit RayCastSelectionPlane(const FRay3d& WorldRay, 
 		FVector3d& HitPointOut);
@@ -571,6 +593,7 @@ protected:
 	// the current selection, with the Z axis being along selection normal, the 0-3
 	// indices here correspond to the 0-3 corner indices in the box, which are the
 	// bottom corners along Z axis (see TOrientedBox3::GetCorner())
+	// TODO: Might be nice to pack these into one value for ease of copying/resetting/comparing
 	bool CornerSelectedFlags[4] = { false, false, false, false };
 	bool PreDragCornerSelectedFlags[4] = { false, false, false, false };
 
