make bsp convert tool's selection changes have associated transactions

diff --git a/Engine/Plugins/Experimental/MeshModelingToolsetExp/Source/MeshModelingToolsEditorOnlyExp/Private/BspConversionTool.cpp b/Engine/Plugins/Experimental/MeshModelingToolsetExp/Source/MeshModelingToolsEditorOnlyExp/Private/BspConversionTool.cpp
--- a/Engine/Plugins/Experimental/MeshModelingToolsetExp/Source/MeshModelingToolsEditorOnlyExp/Private/BspConversionTool.cpp
+++ b/Engine/Plugins/Experimental/MeshModelingToolsetExp/Source/MeshModelingToolsEditorOnlyExp/Private/BspConversionTool.cpp
@@ -20,6 +20,7 @@
 #include "Tools/EditorComponentSourceFactory.h"
 #include "ToolSetupUtil.h"
 #include "Engine/Selection.h"
+#include "ScopedTransaction.h"
 
 #include "Logging/MessageLog.h"
 #include "Misc/MessageDialog.h"
@@ -667,6 +668,8 @@ void UBspConversionTool::ApplyAction(EBspConversionToolAction ActionType)
 	{
 	case EBspConversionToolAction::SelectAllValidBrushes:
 	{
+		FScopedTransaction Transaction(LOCTEXT("SelectAllValidBrushes", "Select All Valid Brushes"));
+
 		FSelectedOjectsChangeList NewSelection;
 		NewSelection.ModificationType = ESelectedObjectsModificationType::Replace;
 
@@ -685,6 +688,8 @@ void UBspConversionTool::ApplyAction(EBspConversionToolAction ActionType)
 	break;
 	case EBspConversionToolAction::DeselectVolumes:
 	{
+		FScopedTransaction Transaction(LOCTEXT("DeselectVolumes", "Deselect Volumes"));
+
 		FSelectedOjectsChangeList NewSelection;
 		NewSelection.ModificationType = ESelectedObjectsModificationType::Replace;
 
@@ -703,6 +708,8 @@ void UBspConversionTool::ApplyAction(EBspConversionToolAction ActionType)
 	break;
 	case EBspConversionToolAction::DeselectNonValid:
 	{
+		FScopedTransaction Transaction(LOCTEXT("DeselectNonValid", "Deselect Non-Valid"));
+
 		// Normally, "GEditor->SelectNone(true, false, false)" would deselect all brushes, but
 		// it does not deselect volumes, which we want to do depending on the settings.
 
