PatternTool: Fix circle shape not updating in XZ or YZ plane by orienting the CenterFrame of the radial pattern generator and adding logic to handle orientation along radial pattern.

diff --git a/Engine/Plugins/Experimental/MeshModelingToolsetExp/Source/MeshModelingToolsExp/Private/PatternTool.cpp b/Engine/Plugins/Experimental/MeshModelingToolsetExp/Source/MeshModelingToolsExp/Private/PatternTool.cpp
--- a/Engine/Plugins/Experimental/MeshModelingToolsetExp/Source/MeshModelingToolsExp/Private/PatternTool.cpp
+++ b/Engine/Plugins/Experimental/MeshModelingToolsetExp/Source/MeshModelingToolsExp/Private/PatternTool.cpp
@@ -382,6 +382,7 @@ public:
 	double Radius = 100.0;
 
 	bool bOriented = true;
+	int AxisIndexToAlign = 0;	// depends on SinglePlane of tool
 
 	ESpacingMode SpacingMode = ESpacingMode::ByCount;
 
@@ -453,15 +454,14 @@ void FRadialPatternGenerator::UpdatePattern_CircleFill()
 		double FrameX = Radius * FMathd::Cos(Angle * FMathd::DegToRad);
 		double FrameY = Radius * FMathd::Sin(Angle * FMathd::DegToRad);
 
-		FFrame3d PatternFrame = CenterFrame;
-		PatternFrame.Origin = CenterFrame.FromPlaneUV(FVector2d(FrameX, FrameY), 2);
+		FFrame3d PatternFrame(CenterFrame.FromPlaneUV(FVector2d(FrameX, FrameY), 2));
 
 		if (bOriented)
 		{
 			FVector3d Axis = Normalized(PatternFrame.Origin - CenterFrame.Origin);
 			if (IsNormalized(Axis))
 			{
-				PatternFrame.ConstrainedAlignAxis(0, Axis, PatternFrame.Z());
+				PatternFrame.ConstrainedAlignAxis(AxisIndexToAlign, Axis, CenterFrame.Z());
 			}
 		}
 
@@ -917,8 +917,26 @@ void UPatternTool::GetPatternTransforms_Radial(TArray<UE::Geometry::FTransformSR
 	InitializeGenerator(Generator, this);
 	Generator.Dimensions = this->Elements[0].PatternBounds;
 
-	Generator.CenterFrame = FFrame3d();
-
+	// Orient the CenterFrame based on the plane setting of the tool & determine which axis is used if bOriented is true
+	const FVector3d Origin(0, 0, 0);
+	switch (Settings->SinglePlane)
+	{
+	case EPatternToolSinglePlane::XYPlane:
+		Generator.CenterFrame = FFrame3d(Origin, FVector3d(0, 0, 1));
+		Generator.AxisIndexToAlign = 0;
+		break;
+	case EPatternToolSinglePlane::XZPlane:
+		Generator.CenterFrame = FFrame3d(Origin, FVector3d(0, 1, 0));
+		Generator.AxisIndexToAlign = 2;
+		break;
+	case EPatternToolSinglePlane::YZPlane:
+		Generator.CenterFrame = FFrame3d(Origin, FVector3d(1, 0, 0));
+		Generator.AxisIndexToAlign = 2;
+		break;
+	default:
+		Generator.CenterFrame = FFrame3d();
+	}
+	
 	Generator.StartAngleDeg = RadialSettings->StartAngle;
 	Generator.EndAngleDeg = RadialSettings->EndAngle;
 	Generator.AngleShift = RadialSettings->AngleShift;
@@ -930,7 +948,7 @@ void UPatternTool::GetPatternTransforms_Radial(TArray<UE::Geometry::FTransformSR
 	Generator.SpacingMode = (FRadialPatternGenerator::ESpacingMode)(int)RadialSettings->SpacingMode;
 
 	Generator.FillMode = FRadialPatternGenerator::EFillMode::CircleFill;
-
+	
 	Generator.Count = RadialSettings->Count;
 	Generator.StepSizeDeg = RadialSettings->StepSize;
 
