PatternTool: Fix various tooltips and UPROPERTY display names.

diff --git a/Engine/Plugins/Experimental/MeshModelingToolsetExp/Source/MeshModelingToolsExp/Public/PatternTool.h b/Engine/Plugins/Experimental/MeshModelingToolsetExp/Source/MeshModelingToolsExp/Public/PatternTool.h
--- a/Engine/Plugins/Experimental/MeshModelingToolsetExp/Source/MeshModelingToolsExp/Public/PatternTool.h
+++ b/Engine/Plugins/Experimental/MeshModelingToolsetExp/Source/MeshModelingToolsExp/Public/PatternTool.h
@@ -274,7 +274,7 @@ public:
 	UPROPERTY(EditAnywhere, Category = Rotation)
 	FRotator Rotation = FRotator::ZeroRotator;
 
-	/** Upper bound of the range which is sampled to randomly rotate each Pattern Element if bJitter is true */
+	/** Upper bound of the range which is sampled to randomly rotate each Pattern Element if Jitter is true */
 	UPROPERTY(EditAnywhere, Category = Rotation, meta = (ClampMin = 0, EditCondition = "bJitter", EditConditionHides, HideEditConditionToggle))
 	FRotator RotationJitterRange = FRotator::ZeroRotator;
 };
@@ -304,7 +304,7 @@ public:
 	UPROPERTY(EditAnywhere, Category = Translation)
 	FVector Translation = FVector::ZeroVector;
 
-	/** Upper bound of the range which is sampled to randomly rotate each Pattern Element if bJitter is true */
+	/** Upper bound of the range which is sampled to randomly translate each Pattern Element if Jitter is true */
 	UPROPERTY(EditAnywhere, Category = Translation, meta = (ClampMin = 0, EditCondition = "bJitter", EditConditionHides, HideEditConditionToggle))
 	FVector TranslationJitterRange = FVector::ZeroVector;
 };
@@ -338,19 +338,19 @@ public:
 	UPROPERTY(EditAnywhere, Category = Scale, meta = (EditCondition = "bUniform", EditConditionHides, HideEditConditionToggle))
 	float Scale = 1;
 
-	/** Upper bound of the range which is sampled to randomly scale each Pattern Element if bJitter is true */
+	/** Upper bound of the range which is sampled to randomly scale each Pattern Element if Jitter is true */
 	UPROPERTY(EditAnywhere, Category = Scale, meta = (ClampMin = 0, EditCondition = "bJitter && bUniform", EditConditionHides, HideEditConditionToggle))
 	float ScaleJitterRange = 0.0f;
 
 	/** Non-Uniform Scale at first Pattern Element */
-	UPROPERTY(EditAnywhere, Category = Scale, meta = (DisplayName = "StartScale", EditCondition = "bInterpolate && bUniform == false", EditConditionHides, HideEditConditionToggle))
+	UPROPERTY(EditAnywhere, Category = Scale, meta = (DisplayName = "Start Scale", EditCondition = "bInterpolate && bUniform == false", EditConditionHides, HideEditConditionToggle))
 	FVector StartScaleNonUniform = FVector(1,1,1);
 
 	/** Non-Uniform Scale applied to all Pattern Elements, or at Last Pattern Element for Interpolated scales */
 	UPROPERTY(EditAnywhere, Category = Scale, meta = (DisplayName = "Scale", EditCondition = "bUniform == false", EditConditionHides, HideEditConditionToggle))
 	FVector ScaleNonUniform = FVector(1,1,1);
 
-	/** Upper bound of the range which is sampled to randomly scale each Pattern Element if bJitter is true (Non-Uniform) */
+	/** Upper bound of the range which is sampled to randomly scale each Pattern Element if Jitter is true (Non-Uniform) */
 	UPROPERTY(EditAnywhere, Category = Scale, meta = (ClampMin = 0, DisplayName = "Scale Jitter Range", EditCondition = "bJitter && bUniform == false", EditConditionHides, HideEditConditionToggle))
 	FVector ScaleJitterRangeNonUniform = FVector(0, 0, 0);
 };
