LWC fixes for polyedit and polypath tools: Keep preview meshes centered near the origin, with translation on the transform rather than baked, because rendering stores the mesh vertices in floating point

diff --git a/Engine/Plugins/Experimental/MeshModelingToolsetExp/Source/MeshModelingToolsExp/Private/DrawPolyPathTool.cpp b/Engine/Plugins/Experimental/MeshModelingToolsetExp/Source/MeshModelingToolsExp/Private/DrawPolyPathTool.cpp
--- a/Engine/Plugins/Experimental/MeshModelingToolsetExp/Source/MeshModelingToolsExp/Private/DrawPolyPathTool.cpp
+++ b/Engine/Plugins/Experimental/MeshModelingToolsetExp/Source/MeshModelingToolsExp/Private/DrawPolyPathTool.cpp
@@ -41,8 +41,9 @@ namespace DrawPolyPathToolLocals
 		}
 	}
 
-
-	void GeneratePathMesh(FDynamicMesh3& Mesh, 
+	/// Generate path mesh
+	/// @return Offset to location of generated mesh 
+	UE_NODISCARD FVector3d GeneratePathMesh(FDynamicMesh3& Mesh,
 		const TArray<FFrame3d>& InPathPoints, 
 		const TArray<double>& InOffsetScaleFactors, 
 		double OffsetDistance, 
@@ -59,6 +60,22 @@ namespace DrawPolyPathToolLocals
 		TArray<FFrame3d> UsePathPoints = InPathPoints;
 		TArray<double> UseOffsetScaleFactors = InOffsetScaleFactors;
 
+		// re-center the input points at the origin
+		FVector3d Center(0, 0, 0);
+		if (UsePathPoints.Num())
+		{
+			FAxisAlignedBox3d PathBounds;
+			for (const FFrame3d& Point : UsePathPoints)
+			{
+				PathBounds.Contain(Point.Origin);
+			}
+			Center = PathBounds.Center();
+			for (FFrame3d& Point : UsePathPoints)
+			{
+				Point.Origin -= Center;
+			}
+		}
+
 		if (bPathIsClosed && bRampMode)
 		{
 			// Duplicate vertices at the beginning/end of the path when generating a ramp
@@ -122,6 +139,8 @@ namespace DrawPolyPathToolLocals
 				Mesh.SetVertex(NumMeshVertices - 1, Mesh.GetVertex(1));
 			}
 		}
+
+		return Center;
 	}
 
 }	// namespace DrawPolyPathToolLocals
@@ -790,7 +809,8 @@ void UDrawPolyPathTool::UpdatePathPreview()
 	}
 
 	FDynamicMesh3 PathMesh;
-	GeneratePathMesh(PathMesh);
+	FVector3d MeshCenter = GeneratePathMesh(PathMesh);
+	EditPreview->SetTransform(FTransform3d(MeshCenter));
 
 	if (State == EState::SettingHeight)
 	{
@@ -804,14 +824,14 @@ void UDrawPolyPathTool::UpdatePathPreview()
 }
 
 
-void UDrawPolyPathTool::GeneratePathMesh(FDynamicMesh3& Mesh)
+FVector3d UDrawPolyPathTool::GeneratePathMesh(FDynamicMesh3& Mesh) 
 {
 	CurPolyLoop.Reset();
 	SecondPolyLoop.Reset();
 
 	const bool bRampMode = (TransformProps->ExtrudeMode == EDrawPolyPathExtrudeMode::RampFixed) || (TransformProps->ExtrudeMode == EDrawPolyPathExtrudeMode::RampInteractive);
 	constexpr bool bLimitCornerRadius = true;
-	DrawPolyPathToolLocals::GeneratePathMesh(Mesh, 
+	FVector3d MeshCenter = DrawPolyPathToolLocals::GeneratePathMesh(Mesh, 
 		CurPathPoints, 
 		OffsetScaleFactors, 
 		TransformProps->Width/2,
@@ -836,6 +856,8 @@ void UDrawPolyPathTool::GeneratePathMesh(FDynamicMesh3& Mesh)
 			Loops.Loops[1].GetVertices<FVector3d>(SecondPolyLoop);
 		}
 	}
+
+	return MeshCenter;
 }
 
 void UDrawPolyPathTool::BeginSettingHeight()
@@ -873,15 +895,17 @@ void UDrawPolyPathTool::BeginInteractiveExtrudeHeight()
 	ExtrudeHeightMechanic->CurrentHeight = 1.0f;  // initialize to something non-zero...prob should be based on polygon bounds maybe?
 
 	FDynamicMesh3 PathMesh;
-	GeneratePathMesh(PathMesh);
+	FVector3d MeshCenter = GeneratePathMesh(PathMesh);
+	EditPreview->SetTransform(FTransform(MeshCenter));
 	EditPreview->InitializeExtrudeType(MoveTemp(PathMesh), DrawPlaneWorld.Z(), nullptr, false);
 
 	FDynamicMesh3 TmpMesh;
 	EditPreview->MakeExtrudeTypeHitTargetMesh(TmpMesh, false);
 
 	FFrame3d UseFrame = DrawPlaneWorld; 
-	UseFrame.Origin = CurPathPoints.Last().Origin;
-	ExtrudeHeightMechanic->Initialize(MoveTemp(TmpMesh), UseFrame, true);
+	UseFrame.Origin = MeshCenter;
+	FTransform3d MeshToFrame(FQuat(UseFrame.Rotation.Inverse()));
+	ExtrudeHeightMechanic->Initialize(MoveTemp(TmpMesh), UseFrame, MeshToFrame);
 
 	ShowExtrudeMessage();
 }
@@ -910,7 +934,8 @@ void UDrawPolyPathTool::BeginConstantExtrudeHeight()
 	}
 
 	FDynamicMesh3 PathMesh;
-	GeneratePathMesh(PathMesh);
+	FVector3d Center = GeneratePathMesh(PathMesh);
+	EditPreview->SetTransform(FTransform(Center));
 	EditPreview->InitializeExtrudeType(MoveTemp(PathMesh), DrawPlaneWorld.Z(), nullptr, false);
 	UpdateExtrudePreview();
 
@@ -1005,14 +1030,19 @@ void UDrawPolyPathTool::GenerateExtrudeMesh(FDynamicMesh3& PathMesh)
 void UDrawPolyPathTool::EmitNewObject()
 {
 	FDynamicMesh3 PathMesh;
-	GeneratePathMesh(PathMesh);
+	FVector3d MeshCenter = GeneratePathMesh(PathMesh);
 	GenerateExtrudeMesh(PathMesh);
 	PathMesh.DiscardVertexUVs();  // throw away arc lengths
 
-	FFrame3d MeshTransform = DrawPlaneWorld;
-	FVector3d Center = PathMesh.GetBounds().Center();
-	MeshTransform.Origin = MeshTransform.ToPlane(Center, 2);
-	MeshTransforms::WorldToFrameCoords(PathMesh, MeshTransform);
+	FFrame3d MeshTransform = DrawPlaneWorld; // The desired frame for the final output mesh
+	FVector3d WorldCenter = PathMesh.GetBounds().Center() + MeshCenter;
+	MeshTransform.Origin = MeshTransform.ToPlane(WorldCenter, 2);
+	
+	// Transform the mesh from its MeshCenter-offset space to the MeshTransform frame space
+	// The below is equivalent to applying: MeshTransform.Rotation^-1 (Mesh + MeshCenter - MeshTransform.Origin)
+	FFrame3d LocalCenterMeshTransform = MeshTransform;
+	LocalCenterMeshTransform.Origin -= MeshCenter;
+	MeshTransforms::WorldToFrameCoords(PathMesh, LocalCenterMeshTransform);
 	UE::Geometry::FMeshTangentsf::ComputeDefaultOverlayTangents(PathMesh);
 
 	GetToolManager()->BeginUndoTransaction(LOCTEXT("CreatePolyPathTransactionName", "Create PolyPath"));
diff --git a/Engine/Plugins/Experimental/MeshModelingToolsetExp/Source/MeshModelingToolsExp/Public/DrawPolyPathTool.h b/Engine/Plugins/Experimental/MeshModelingToolsetExp/Source/MeshModelingToolsExp/Public/DrawPolyPathTool.h
--- a/Engine/Plugins/Experimental/MeshModelingToolsetExp/Source/MeshModelingToolsExp/Public/DrawPolyPathTool.h
+++ b/Engine/Plugins/Experimental/MeshModelingToolsetExp/Source/MeshModelingToolsExp/Public/DrawPolyPathTool.h
@@ -285,7 +285,7 @@ protected:
 	void InitializePreviewMesh();
 	void ClearPreview();
 
-	void GeneratePathMesh(UE::Geometry::FDynamicMesh3& Mesh);
+	UE_NODISCARD FVector3d GeneratePathMesh(UE::Geometry::FDynamicMesh3& Mesh);
 	void GenerateExtrudeMesh(UE::Geometry::FDynamicMesh3& PathMesh);
 	void EmitNewObject();
 
diff --git a/Engine/Plugins/Runtime/MeshModelingToolset/Source/MeshModelingTools/Private/ToolActivities/PolyEditCutFacesActivity.cpp b/Engine/Plugins/Runtime/MeshModelingToolset/Source/MeshModelingTools/Private/ToolActivities/PolyEditCutFacesActivity.cpp
--- a/Engine/Plugins/Runtime/MeshModelingToolset/Source/MeshModelingTools/Private/ToolActivities/PolyEditCutFacesActivity.cpp
+++ b/Engine/Plugins/Runtime/MeshModelingToolset/Source/MeshModelingTools/Private/ToolActivities/PolyEditCutFacesActivity.cpp
@@ -103,7 +103,9 @@ void UPolyEditCutFacesActivity::BeginCutFaces()
 	FTransform3d WorldTransform(ActivityContext->Preview->PreviewMesh->GetTransform());
 
 	EditPreview = PolyEditActivityUtil::CreatePolyEditPreviewMesh(*ParentTool, *ActivityContext);
-	EditPreview->InitializeStaticType(ActivityContext->CurrentMesh.Get(), ActiveTriangleSelection, &WorldTransform);
+	FTransform3d WorldTranslation, WorldRotateScale;
+	EditPreview->ApplyTranslationToPreview(WorldTransform, WorldTranslation, WorldRotateScale);
+	EditPreview->InitializeStaticType(ActivityContext->CurrentMesh.Get(), ActiveTriangleSelection, &WorldRotateScale);
 
 	FDynamicMesh3 StaticHitTargetMesh;
 	EditPreview->MakeInsetTypeTargetMesh(StaticHitTargetMesh);
diff --git a/Engine/Plugins/Runtime/MeshModelingToolset/Source/MeshModelingTools/Private/ToolActivities/PolyEditInsetOutsetActivity.cpp b/Engine/Plugins/Runtime/MeshModelingToolset/Source/MeshModelingTools/Private/ToolActivities/PolyEditInsetOutsetActivity.cpp
--- a/Engine/Plugins/Runtime/MeshModelingToolset/Source/MeshModelingTools/Private/ToolActivities/PolyEditInsetOutsetActivity.cpp
+++ b/Engine/Plugins/Runtime/MeshModelingToolset/Source/MeshModelingTools/Private/ToolActivities/PolyEditInsetOutsetActivity.cpp
@@ -121,7 +121,9 @@ bool UPolyEditInsetOutsetActivity::BeginInset()
 	FTransform3d WorldTransform(ActivityContext->Preview->PreviewMesh->GetTransform());
 
 	EditPreview = PolyEditActivityUtil::CreatePolyEditPreviewMesh(*ParentTool, *ActivityContext);
-	EditPreview->InitializeInsetType(ActivityContext->CurrentMesh.Get(), ActiveTriangleSelection, &WorldTransform);
+	FTransform3d WorldTranslation, WorldRotateScale;
+	EditPreview->ApplyTranslationToPreview(WorldTransform, WorldTranslation, WorldRotateScale);
+	EditPreview->InitializeInsetType(ActivityContext->CurrentMesh.Get(), ActiveTriangleSelection, &WorldRotateScale);
 
 	// Hide the selected triangles (that are being replaced by the inset/outset portion)
 	ActivityContext->Preview->PreviewMesh->SetSecondaryBuffersVisibility(false);
diff --git a/Engine/Plugins/Runtime/MeshModelingToolset/Source/MeshModelingTools/Private/ToolActivities/PolyEditPlanarProjectionUVActivity.cpp b/Engine/Plugins/Runtime/MeshModelingToolset/Source/MeshModelingTools/Private/ToolActivities/PolyEditPlanarProjectionUVActivity.cpp
--- a/Engine/Plugins/Runtime/MeshModelingToolset/Source/MeshModelingTools/Private/ToolActivities/PolyEditPlanarProjectionUVActivity.cpp
+++ b/Engine/Plugins/Runtime/MeshModelingToolset/Source/MeshModelingTools/Private/ToolActivities/PolyEditPlanarProjectionUVActivity.cpp
@@ -104,7 +104,9 @@ void UPolyEditPlanarProjectionUVActivity::BeginSetUVs()
 	FTransform3d WorldTransform(ActivityContext->Preview->PreviewMesh->GetTransform());
 
 	EditPreview = PolyEditActivityUtil::CreatePolyEditPreviewMesh(*ParentTool, *ActivityContext);
-	EditPreview->InitializeStaticType(ActivityContext->CurrentMesh.Get(), ActiveTriangleSelection, &WorldTransform);
+	FTransform3d WorldTranslation, WorldRotateScale;
+	EditPreview->ApplyTranslationToPreview(WorldTransform, WorldTranslation, WorldRotateScale);
+	EditPreview->InitializeStaticType(ActivityContext->CurrentMesh.Get(), ActiveTriangleSelection, &WorldRotateScale);
 
 	UpdatePolyEditPreviewMaterials(*ParentTool, *ActivityContext, *EditPreview, (SetUVProperties->bShowMaterial) ?
 		EPreviewMaterialType::SourceMaterials : EPreviewMaterialType::UVMaterial);
@@ -140,17 +142,28 @@ void UPolyEditPlanarProjectionUVActivity::BeginSetUVs()
 void UPolyEditPlanarProjectionUVActivity::UpdateSetUVS()
 {
 	// align projection frame to line user is drawing out from plane origin
-	FFrame3d PlanarFrame = SurfacePathMechanic->PreviewPathPoint;
-	double UVScale = 1.0 / ActiveSelectionBounds.MaxDim();
+	FFrame3d WorldPlanarFrame = SurfacePathMechanic->PreviewPathPoint;
+	FFrame3d LocalPlanarFrame;
+	double UVScale = 1.0 / ActiveSelectionBounds.MaxDim(); // Note: EditPreview has baked scale, so UVScale is already in world
+	FTransformSRT3d WorldTransform = EditPreview->GetTransform();
 	if (SurfacePathMechanic->HitPath.Num() == 1)
 	{
-		SurfacePathMechanic->InitializePlaneSurface(PlanarFrame);
+		SurfacePathMechanic->InitializePlaneSurface(WorldPlanarFrame);
 
-		FVector3d Delta = PlanarFrame.Origin - SurfacePathMechanic->HitPath[0].Origin;
+		FVector3d Delta = WorldPlanarFrame.Origin - SurfacePathMechanic->HitPath[0].Origin;
 		double Dist = UE::Geometry::Normalize(Delta);
 		UVScale *= FMathd::Lerp(1.0, 25.0, Dist / ActiveSelectionBounds.MaxDim());
-		PlanarFrame = SurfacePathMechanic->HitPath[0];
-		PlanarFrame.ConstrainedAlignAxis(0, Delta, PlanarFrame.Z());
+		WorldPlanarFrame = SurfacePathMechanic->HitPath[0];
+		FVector3d WorldNormal = WorldPlanarFrame.Z();
+		FVector3d LocalNormal = WorldTransform.InverseTransformNormal(WorldNormal);
+		LocalPlanarFrame = FFrame3d(WorldTransform.InverseTransformPosition(WorldPlanarFrame.Origin), LocalNormal);
+		LocalPlanarFrame.ConstrainedAlignAxis(0, WorldTransform.InverseTransformVector(Delta), LocalPlanarFrame.Z());
+	}
+	else
+	{
+		FVector3d WorldNormal = WorldPlanarFrame.Z();
+		FVector3d LocalNormal = WorldTransform.InverseTransformNormal(WorldNormal);
+		LocalPlanarFrame = FFrame3d(WorldTransform.InverseTransformPosition(WorldPlanarFrame.Origin), LocalNormal);
 	}
 
 	EditPreview->UpdateStaticType([&](FDynamicMesh3& Mesh)
@@ -161,7 +174,7 @@ void UPolyEditPlanarProjectionUVActivity::UpdateSetUVS()
 			{
 				AllTriangles.Add(tid);
 			}
-			Editor.SetTriangleUVsFromProjection(AllTriangles, PlanarFrame, UVScale, FVector2f::Zero(), false);
+			Editor.SetTriangleUVsFromProjection(AllTriangles, LocalPlanarFrame, UVScale, FVector2f::Zero(), false);
 		}, false);
 
 }
diff --git a/Engine/Plugins/Runtime/MeshModelingToolset/Source/ModelingComponents/Private/Drawing/PolyEditPreviewMesh.cpp b/Engine/Plugins/Runtime/MeshModelingToolset/Source/ModelingComponents/Private/Drawing/PolyEditPreviewMesh.cpp
--- a/Engine/Plugins/Runtime/MeshModelingToolset/Source/ModelingComponents/Private/Drawing/PolyEditPreviewMesh.cpp
+++ b/Engine/Plugins/Runtime/MeshModelingToolset/Source/ModelingComponents/Private/Drawing/PolyEditPreviewMesh.cpp
@@ -34,7 +34,9 @@ void UPolyEditPreviewMesh::InitializeExtrudeType(
 
 	// save copy of initial patch
 	InitialEditPatch = EditPatch;
+	PRAGMA_DISABLE_DEPRECATION_WARNINGS
 	InitialEditPatchBVTree.SetMesh(&InitialEditPatch);
+	PRAGMA_ENABLE_DEPRECATION_WARNINGS
 
 	// save initial tri normals
 	InitialTriNormals.SetNum(InitialEditPatch.MaxTriangleID());
@@ -103,7 +105,9 @@ void UPolyEditPreviewMesh::InitializeExtrudeType(FDynamicMesh3&& BaseMesh,
 	}
 
 	// save copy of initial patch
+	PRAGMA_DISABLE_DEPRECATION_WARNINGS
 	InitialEditPatchBVTree.SetMesh(&InitialEditPatch);
+	PRAGMA_ENABLE_DEPRECATION_WARNINGS
 
 	// save initial tri normals
 	InitialTriNormals.SetNum(InitialEditPatch.MaxTriangleID());
@@ -259,7 +263,9 @@ void UPolyEditPreviewMesh::InitializeInsetType(const FDynamicMesh3* SourceMesh,
 
 	// save copy of initial patch
 	InitialEditPatch = EditPatch;
+	PRAGMA_DISABLE_DEPRECATION_WARNINGS
 	InitialEditPatchBVTree.SetMesh(&InitialEditPatch);
+	PRAGMA_ENABLE_DEPRECATION_WARNINGS
 
 	// initialize the preview mesh
 	UpdatePreview(&EditPatch);
@@ -290,6 +296,7 @@ void UPolyEditPreviewMesh::UpdateInsetType(double NewOffset, bool bReproject, do
 void UPolyEditPreviewMesh::MakeInsetTypeTargetMesh(FDynamicMesh3& TargetMesh)
 {
 	TargetMesh = InitialEditPatch;
+	MeshTransforms::ApplyTransform(TargetMesh, GetTransform());
 }
 
 
@@ -315,7 +322,6 @@ void UPolyEditPreviewMesh::InitializeStaticType(const FDynamicMesh3* SourceMesh,
 
 	// save copy of initial patch
 	InitialEditPatch = EditPatch;
-	InitialEditPatchBVTree.SetMesh(&InitialEditPatch);
 
 	// initialize the preview mesh
 	UpdatePreview(&EditPatch);
diff --git a/Engine/Plugins/Runtime/MeshModelingToolset/Source/ModelingComponents/Private/Mechanics/PlaneDistanceFromHitMechanic.cpp b/Engine/Plugins/Runtime/MeshModelingToolset/Source/ModelingComponents/Private/Mechanics/PlaneDistanceFromHitMechanic.cpp
--- a/Engine/Plugins/Runtime/MeshModelingToolset/Source/ModelingComponents/Private/Mechanics/PlaneDistanceFromHitMechanic.cpp
+++ b/Engine/Plugins/Runtime/MeshModelingToolset/Source/ModelingComponents/Private/Mechanics/PlaneDistanceFromHitMechanic.cpp
@@ -40,6 +40,17 @@ void UPlaneDistanceFromHitMechanic::Initialize(FDynamicMesh3&& HitTargetMesh, co
 	PreviewHeightTargetAABB.SetMesh(&PreviewHeightTarget);
 }
 
+// Alternative Initialize() that explicitly specifies what transform to bake
+void UPlaneDistanceFromHitMechanic::Initialize(FDynamicMesh3&& HitTargetMesh, const FFrame3d& PlaneFrameWorld, const FTransform& MeshToPlaneFrame)
+{
+	PreviewHeightFrame = PlaneFrameWorld;
+
+	PreviewHeightTarget = MoveTemp(HitTargetMesh);
+	MeshTransforms::ApplyTransform(PreviewHeightTarget, MeshToPlaneFrame);
+
+	PreviewHeightTargetAABB.SetMesh(&PreviewHeightTarget);
+}
+
 
 
 void UPlaneDistanceFromHitMechanic::UpdateCurrentDistance(const FRay& WorldRay)
diff --git a/Engine/Plugins/Runtime/MeshModelingToolset/Source/ModelingComponents/Public/Drawing/PolyEditPreviewMesh.h b/Engine/Plugins/Runtime/MeshModelingToolset/Source/ModelingComponents/Public/Drawing/PolyEditPreviewMesh.h
--- a/Engine/Plugins/Runtime/MeshModelingToolset/Source/ModelingComponents/Public/Drawing/PolyEditPreviewMesh.h
+++ b/Engine/Plugins/Runtime/MeshModelingToolset/Source/ModelingComponents/Public/Drawing/PolyEditPreviewMesh.h
@@ -57,16 +57,34 @@ public:
 	void InitializeInsetType(const FDynamicMesh3* SourceMesh, const TArray<int32>& Triangles,
 		const FTransform3d* MeshTransform = nullptr);
 	void UpdateInsetType(double NewOffset, bool bReproject = false, double Softness = 0.0, double AreaScaleT = 1.0, bool bBoundaryOnly = false);
-	void MakeInsetTypeTargetMesh(FDynamicMesh3& TargetMesh);
+	
+	/// @param TargetMeshOut Will be populated with the initial edit patch mesh transformed to world space via the preview mesh transform
+	void MakeInsetTypeTargetMesh(FDynamicMesh3& TargetMeshOut);
 
+	UE_DEPRECATED(5.1, "We are removing this accessor because it was unused.")
 	const FDynamicMesh3& GetInitialPatchMesh() const { return InitialEditPatch; }
+	
+	void ApplyTranslationToPreview(const FTransform& TransformIn, FTransform& TranslationOut, FTransform& RotateScaleOut)
+	{
+		TranslationOut = FTransform(TransformIn.GetTranslation());
+		SetTransform(TranslationOut);
+		RotateScaleOut = TransformIn;
+		RotateScaleOut.SetTranslation(FVector::ZeroVector);
+	}
 
 	// must be non-const because query functions are non-const
-	UE::Geometry::FDynamicMeshAABBTree3& GetInitialPatchMeshSpatial() { return InitialEditPatchBVTree; }
+	UE_DEPRECATED(5.1, "We are removing this accessor and the underlying tree, because they were unused.")
+	UE::Geometry::FDynamicMeshAABBTree3& GetInitialPatchMeshSpatial()
+	{
+		PRAGMA_DISABLE_DEPRECATION_WARNINGS
+		return InitialEditPatchBVTree; 
+		PRAGMA_ENABLE_DEPRECATION_WARNINGS
+	}
 
 protected:
 	TUniquePtr<UE::Geometry::FDynamicSubmesh3> ActiveSubmesh;
 	FDynamicMesh3 InitialEditPatch;
+	UE_DEPRECATED(5.1, "We are removing this tree, because it was unused.")
 	UE::Geometry::FDynamicMeshAABBTree3 InitialEditPatchBVTree;
 
 	TArray<int32> EditVertices;
@@ -77,6 +95,6 @@ protected:
 
 	FVector3d InputDirection;
 
-	FTransform3d MeshTransform;
+	FTransform3d MeshTransform; // Transform that has been baked into the FDynamicMesh3
 	bool bHaveMeshTransform;
 };
\ No newline at end of file
diff --git a/Engine/Plugins/Runtime/MeshModelingToolset/Source/ModelingComponents/Public/Mechanics/PlaneDistanceFromHitMechanic.h b/Engine/Plugins/Runtime/MeshModelingToolset/Source/ModelingComponents/Public/Mechanics/PlaneDistanceFromHitMechanic.h
--- a/Engine/Plugins/Runtime/MeshModelingToolset/Source/ModelingComponents/Public/Mechanics/PlaneDistanceFromHitMechanic.h
+++ b/Engine/Plugins/Runtime/MeshModelingToolset/Source/ModelingComponents/Public/Mechanics/PlaneDistanceFromHitMechanic.h
@@ -59,6 +59,12 @@ public:
 	 */
 	virtual void Initialize(UE::Geometry::FDynamicMesh3&& HitTargetMesh, const FFrame3d& PlaneFrameWorld, bool bMeshInWorldCoords);
 
+	/**
+	 * Set the hit target mesh and the plane frame. Distance is measured along Z axis.
+	 * Applies MeshToPlaneFrame to transform the HitTargetMesh to the local space of PlaneFrameWorld
+	 */
+	virtual void Initialize(UE::Geometry::FDynamicMesh3&& HitTargetMesh, const FFrame3d& PlaneFrameWorld, const FTransform& MeshToPlaneFrame);
+
 	/**
 	 * Update the current distance/height based on the input world ray
 	 */
