Fixed a bug where Smooth Borders action in TriSel tool could not be undone.

#rnx

diff --git a/Engine/Plugins/Experimental/MeshModelingToolsetExp/Source/MeshModelingToolsExp/Private/MeshSelectionTool.cpp b/Engine/Plugins/Experimental/MeshModelingToolsetExp/Source/MeshModelingToolsExp/Private/MeshSelectionTool.cpp
--- a/Engine/Plugins/Experimental/MeshModelingToolsetExp/Source/MeshModelingToolsExp/Private/MeshSelectionTool.cpp
+++ b/Engine/Plugins/Experimental/MeshModelingToolsetExp/Source/MeshModelingToolsExp/Private/MeshSelectionTool.cpp
@@ -1665,6 +1665,11 @@ void UMeshSelectionTool::SmoothSelectionBoundary()
 			}
 		}
 
+		for (const TPair<int, FVector3d>& NewVert : NewLoopPositions) 
+		{
+			ChangeTracker.SaveVertexOneRingTriangles(NewVert.Key, true);
+		}
+
 		PreviewMesh->EditMesh([NewLoopPositions](FDynamicMesh3& Mesh)
 		{
 			for (const TPair<int, FVector3d>& NewVert : NewLoopPositions)
