MeshModelingTools: Make XForm tool ignore hotkeys while snap dragging.

diff --git a/Engine/Plugins/Experimental/MeshModelingToolsetExp/Source/MeshModelingToolsExp/Private/TransformMeshesTool.cpp b/Engine/Plugins/Experimental/MeshModelingToolsetExp/Source/MeshModelingToolsExp/Private/TransformMeshesTool.cpp
--- a/Engine/Plugins/Experimental/MeshModelingToolsetExp/Source/MeshModelingToolsExp/Private/TransformMeshesTool.cpp
+++ b/Engine/Plugins/Experimental/MeshModelingToolsetExp/Source/MeshModelingToolsExp/Private/TransformMeshesTool.cpp
@@ -130,11 +130,13 @@ void UTransformMeshesTool::UpdateSetPivotModes(bool bEnableSetPivot)
 
 void UTransformMeshesTool::RegisterActions(FInteractiveToolActionSet& ActionSet)
 {
-	// Note: we should really disallow hotkeys while dragging, but we are discouraged from adding
-	// members to an object in a hotfix, so we can't add a bCurrentlySnapDragging (and we can't use
-	// ActiveSnapDragIndex directly because it can be set even without getting capture). For now, we
-	// just make sure to call CheckAndUpdateWatched() after any property changes to avoid some invalid
-	// states that can happen between a drag call and the next tick.
+	// Note: The CheckAndUpdateWatched() calls below were added to fix some invalid states
+	// that we got into if we used hotkeys while snap dragging (where the handler might be called
+	// before next tick). This should no longer be possible since we check that we're not snap dragging,
+	// but we'll keep them here for safety.
+	// The hotkeys are still usable while dragging the gizmos, but that doesn't cause
+	// the kinds of UI issues that it does for snap dragging, and changing it for consistency
+	// is tedious since we build and rebuild potentially multiple transform proxies.
 
 	ActionSet.RegisterAction(this, (int32)EStandardToolActions::BaseClientDefinedActionID + 1,
 		TEXT("ToggleSetPivot"),
@@ -142,6 +144,10 @@ void UTransformMeshesTool::RegisterActions(FInteractiveToolActionSet& ActionSet)
 		LOCTEXT("TransformToggleSetPivotTooltip", "Toggle Set Pivot on and off"),
 		EModifierKey::None, EKeys::S,
 		[this]() { 
+			if (bCurrentlyDragging)
+			{
+				return;
+			}
 			TransformProps->bSetPivotMode = !TransformProps->bSetPivotMode; 
 			TransformProps->CheckAndUpdateWatched();
 			NotifyOfPropertyChangeByTool(TransformProps);
@@ -153,6 +159,10 @@ void UTransformMeshesTool::RegisterActions(FInteractiveToolActionSet& ActionSet)
 		LOCTEXT("TransformToggleSnapDragTooltip", "Toggle SnapDrag on and off"),
 		EModifierKey::None, EKeys::D,
 		[this]() { 
+			if (bCurrentlyDragging)
+			{
+				return;
+			}
 			TransformProps->bEnableSnapDragging = !TransformProps->bEnableSnapDragging;
 			TransformProps->CheckAndUpdateWatched();
 			NotifyOfPropertyChangeByTool(TransformProps);
@@ -164,6 +174,10 @@ void UTransformMeshesTool::RegisterActions(FInteractiveToolActionSet& ActionSet)
 		LOCTEXT("TransformCycleTransformModeTooltip", "Cycle through available Transform Modes"),
 		EModifierKey::None, EKeys::A,
 		[this]() { 
+			if (bCurrentlyDragging)
+			{
+				return;
+			}
 			TransformProps->TransformMode = (ETransformMeshesTransformMode)(((uint8)TransformProps->TransformMode+1) % (uint8)ETransformMeshesTransformMode::LastValue);
 			TransformProps->CheckAndUpdateWatched();
 			NotifyOfPropertyChangeByTool(TransformProps);
@@ -176,6 +190,10 @@ void UTransformMeshesTool::RegisterActions(FInteractiveToolActionSet& ActionSet)
 		LOCTEXT("TransformCycleSourceModeTooltip", "Cycle through available SnapDrag Source Modes"),
 		EModifierKey::None, EKeys::W,
 		[this]() { 
+			if (bCurrentlyDragging)
+			{
+				return;
+			}
 			TransformProps->SnapDragSource = (ETransformMeshesSnapDragSource)(((uint8)TransformProps->SnapDragSource + 1) % (uint8)ETransformMeshesSnapDragSource::LastValue);
 			TransformProps->CheckAndUpdateWatched();
 			NotifyOfPropertyChangeByTool(TransformProps);
@@ -187,6 +205,10 @@ void UTransformMeshesTool::RegisterActions(FInteractiveToolActionSet& ActionSet)
 		LOCTEXT("TransformCycleRotationModeTooltip", "Cycle through available SnapDrag Rotation Modes"),
 		EModifierKey::None, EKeys::E,
 		[this]() { 
+			if (bCurrentlyDragging)
+			{
+				return;
+			}
 			TransformProps->RotationMode = (ETransformMeshesSnapDragRotationMode)(((uint8)TransformProps->RotationMode + 1) % (uint8)ETransformMeshesSnapDragRotationMode::LastValue); 
 			TransformProps->CheckAndUpdateWatched();
 			NotifyOfPropertyChangeByTool(TransformProps);
@@ -406,6 +428,7 @@ void UTransformMeshesTool::OnClickPress(const FInputDeviceRay& PressPos)
 
 void UTransformMeshesTool::OnClickDrag(const FInputDeviceRay& DragPos)
 {
+	bCurrentlyDragging = true;
 	bool bApplyToPivot = TransformProps->bSetPivotMode;
 
 	TArray<const UPrimitiveComponent*> IgnoreComponents;
@@ -490,6 +513,7 @@ void UTransformMeshesTool::OnClickRelease(const FInputDeviceRay& ReleasePos)
 
 void UTransformMeshesTool::OnTerminateDragSequence()
 {
+	bCurrentlyDragging = false;
 	FTransformMeshesTarget& ActiveTarget =
 		(TransformProps->TransformMode == ETransformMeshesTransformMode::PerObjectGizmo) ?
 		ActiveGizmos[ActiveSnapDragIndex] : ActiveGizmos[0];
diff --git a/Engine/Plugins/Experimental/MeshModelingToolsetExp/Source/MeshModelingToolsExp/Public/TransformMeshesTool.h b/Engine/Plugins/Experimental/MeshModelingToolsetExp/Source/MeshModelingToolsExp/Public/TransformMeshesTool.h
--- a/Engine/Plugins/Experimental/MeshModelingToolsetExp/Source/MeshModelingToolsExp/Public/TransformMeshesTool.h
+++ b/Engine/Plugins/Experimental/MeshModelingToolsetExp/Source/MeshModelingToolsExp/Public/TransformMeshesTool.h
@@ -197,6 +197,8 @@ protected:
 	UE::Geometry::FFrame3d StartDragFrameWorld;
 	FTransform StartDragTransform;
 	int ActiveSnapDragIndex = -1;
+	// Can't just use ActiveSnapDragIndex to replace this bool, because the former gets set while checking for capture
+	bool bCurrentlyDragging = false;
 
 	void OnParametersUpdated();
 };
