Fixed a bug in the Displace tool where we hit ensure after reverting selection of the Active Material (Selective Tessellation Option).

#rnx

diff --git a/Engine/Plugins/Experimental/MeshModelingToolsetExp/Source/MeshModelingToolsExp/Private/DisplaceMeshTool.cpp b/Engine/Plugins/Experimental/MeshModelingToolsetExp/Source/MeshModelingToolsExp/Private/DisplaceMeshTool.cpp
--- a/Engine/Plugins/Experimental/MeshModelingToolsetExp/Source/MeshModelingToolsExp/Private/DisplaceMeshTool.cpp
+++ b/Engine/Plugins/Experimental/MeshModelingToolsetExp/Source/MeshModelingToolsExp/Private/DisplaceMeshTool.cpp
@@ -355,25 +355,12 @@ namespace DisplaceMeshToolLocals{
 		int SubdivisionsCount = Parameters.SubdivisionsCount;
 		if (SubdivisionType == EDisplaceMeshToolSubdivisionType::Flat) 
 		{
-			if (Parameters.SelectionType.IsSet() == false)
-			{
-				return;
-			}
-
-			if (Parameters.SelectionType.GetValue() == EDisplaceMeshToolTriangleSelectionType::None) 
-			{
-				FUniformTessellate Tessellator(ResultMesh.Get());
-				Tessellator.Progress = ProgressCancel;
-				Tessellator.TessellationNum = SubdivisionsCount;
-							
-				if (Tessellator.Validate() == EOperationValidationResult::Ok) 
-				{
-					Tessellator.Compute();
-				}
-			}
-			else if (Parameters.SelectionType.GetValue() == EDisplaceMeshToolTriangleSelectionType::Material)
-			{
-				if (Parameters.ActiveMaterialID.IsSet())
+			if (Parameters.SelectionType.IsSet()) // user wants to only tessellate a subset of the triangles
+			{	
+				EDisplaceMeshToolTriangleSelectionType SelectionType = Parameters.SelectionType.GetValue();
+				if (SelectionType == EDisplaceMeshToolTriangleSelectionType::Material && 
+					Parameters.ActiveMaterialID.IsSet() && 
+					Parameters.ActiveMaterialID.GetValue() != INDEX_NONE)
 				{
 					TUniquePtr<FTessellationPattern> Pattern = FSelectiveTessellate::CreateConcentricRingsPatternFromMaterial(ResultMesh.Get(), SubdivisionsCount, Parameters.ActiveMaterialID.GetValue());;
 					
@@ -393,10 +380,30 @@ namespace DisplaceMeshToolLocals{
 						}
 					}
 				}
+				else
+				{
+					// Fall back to the uniform tessellation if either SelectionType is EDisplaceMeshToolTriangleSelectionType::None 
+					// or not all the parameters were set for the other selection types. 
+					FUniformTessellate Tessellator(ResultMesh.Get());
+					Tessellator.Progress = ProgressCancel;
+					Tessellator.TessellationNum = SubdivisionsCount;
+
+					if (ensureMsgf(Tessellator.Validate() == EOperationValidationResult::Ok, TEXT("The tessellator parameters are invalid.")))
+					{
+						Tessellator.Compute();
+					}
+				}
 			}
 			else
 			{
-				checkNoEntry(); // Unsupported selection type
+				FUniformTessellate Tessellator(ResultMesh.Get());
+				Tessellator.Progress = ProgressCancel;
+				Tessellator.TessellationNum = SubdivisionsCount;
+						
+				if (ensureMsgf(Tessellator.Validate() == EOperationValidationResult::Ok, TEXT("The tessellator parameters are invalid.")))
+				{
+					Tessellator.Compute();
+				}
 			}
 		}
 		else if (SubdivisionType == EDisplaceMeshToolSubdivisionType::PNTriangles) 
@@ -405,9 +412,9 @@ namespace DisplaceMeshToolLocals{
 			PNTriangles.Progress = ProgressCancel;
 			PNTriangles.TessellationLevel = SubdivisionsCount;
 
-			if (PNTriangles.Validate() == EOperationValidationResult::Ok)
+			if (ensureMsgf(PNTriangles.Validate() == EOperationValidationResult::Ok, TEXT("The tessellator parameters are invalid.")))
 			{
-				PNTriangles.Compute(); 
+				PNTriangles.Compute();
 			}
 		}
 		else 
@@ -1049,10 +1056,7 @@ void UDisplaceMeshTool::Setup()
 	if (SelectiveTessellationProperties->ActiveMaterial.IsNone() == false) 
 	{
 		int32 Index = SelectiveTessellationProperties->MaterialIDList.Find(SelectiveTessellationProperties->ActiveMaterial.ToString());
-		if (ensure(Index != INDEX_NONE)) 
-		{
-			SubParameters.ActiveMaterialID = Index;
-		}
+		SubParameters.ActiveMaterialID = Index;
 	}
 
 	Subdivider = MakeUnique<FSubdivideMeshOpFactory>(OriginalMesh, SubParameters, CommonProperties->SubdivisionType);
@@ -1308,11 +1312,8 @@ void UDisplaceMeshTool::OnPropertyModified(UObject* PropertySet, FProperty* Prop
 		else if (PropName == GET_MEMBER_NAME_CHECKED(USelectiveTessellationProperties, ActiveMaterial)) 
 		{
 			int32 Index = SelectiveTessellationProperties->MaterialIDList.Find(SelectiveTessellationProperties->ActiveMaterial.ToString());
-			if (ensure(Index != INDEX_NONE)) 
-			{
-				SubdividerDownCast->SetActiveMaterialID(Index);
-				bNeedsSubdivided = true;
-			}
+			SubdividerDownCast->SetActiveMaterialID(Index);
+			bNeedsSubdivided = true;
 		}
 		
 		StartComputation();
