Geometry Tools:  Pivot Actor tool  - Adding a simple display names for the case when the tool is adding vs modifying a pivot actor

diff --git a/Engine/Plugins/Experimental/MeshModelingToolsetExp/Source/MeshModelingToolsEditorOnlyExp/Private/AddPivotActorTool.cpp b/Engine/Plugins/Experimental/MeshModelingToolsetExp/Source/MeshModelingToolsEditorOnlyExp/Private/AddPivotActorTool.cpp
--- a/Engine/Plugins/Experimental/MeshModelingToolsetExp/Source/MeshModelingToolsEditorOnlyExp/Private/AddPivotActorTool.cpp
+++ b/Engine/Plugins/Experimental/MeshModelingToolsetExp/Source/MeshModelingToolsEditorOnlyExp/Private/AddPivotActorTool.cpp
@@ -84,6 +84,7 @@ UMultiSelectionMeshEditingTool* UAddPivotActorToolBuilder::CreateNewTool(const F
 	{
 		NewTool->SetPivotRepositionMode(SceneState.SelectedActors[0]);
 	}
+
 	return NewTool;
 }
 
@@ -94,6 +95,8 @@ void UAddPivotActorTool::Setup()
 	FTransform StartTransform = FTransform::Identity;
 	if (ExistingPivotActor.IsValid())
 	{
+		SetToolDisplayName(LOCTEXT("ModifyPivotActorToolName", "Modify Pivot Actor"));
+
 		GetToolManager()->DisplayMessage(LOCTEXT("OnStartToolEdit",
 			"Modifies the position of the pivot actor while keeping children in place. "
 			"Hold Ctrl to snap to items in scene."),
@@ -111,6 +114,8 @@ void UAddPivotActorTool::Setup()
 	}
 	else
 	{
+		SetToolDisplayName(LOCTEXT("AddPivotActorToolName", "Add Pivot Actor"));
+
 		GetToolManager()->DisplayMessage(LOCTEXT("OnStartToolAdd",
 			"Adds an empty actor as the parent of the selected actors. Use gizmo to choose where/how "
 			"the empty actor is placed. Hold Ctrl to snap to items in scene."),
