MeshModelingToolset: Use NoResetToDefault for mesh statistics in Inspect tool

diff --git a/Engine/Plugins/Experimental/MeshModelingToolsetExp/Source/MeshModelingToolsExp/Public/Properties/MeshAnalysisProperties.h b/Engine/Plugins/Experimental/MeshModelingToolsetExp/Source/MeshModelingToolsExp/Public/Properties/MeshAnalysisProperties.h
--- a/Engine/Plugins/Experimental/MeshModelingToolsetExp/Source/MeshModelingToolsExp/Public/Properties/MeshAnalysisProperties.h
+++ b/Engine/Plugins/Experimental/MeshModelingToolsetExp/Source/MeshModelingToolsExp/Public/Properties/MeshAnalysisProperties.h
@@ -20,11 +20,11 @@ class MESHMODELINGTOOLSEXP_API UMeshAnalysisProperties : public UInteractiveTool
 public:
 
 	/** In meters squared */
-	UPROPERTY(VisibleAnywhere, Category = MeshAnalysis)
+	UPROPERTY(VisibleAnywhere, Category = MeshAnalysis, meta = (NoResetToDefault))
 	FString SurfaceArea;
 
 	/** In cubic meters */
-	UPROPERTY(VisibleAnywhere, Category = MeshAnalysis)
+	UPROPERTY(VisibleAnywhere, Category = MeshAnalysis, meta = (NoResetToDefault))
 	FString Volume;
 
 	void Update(const FDynamicMesh3& Mesh, const FTransform& Transform);
diff --git a/Engine/Plugins/Experimental/MeshModelingToolsetExp/Source/MeshModelingToolsExp/Public/Properties/MeshStatisticsProperties.h b/Engine/Plugins/Experimental/MeshModelingToolsetExp/Source/MeshModelingToolsExp/Public/Properties/MeshStatisticsProperties.h
--- a/Engine/Plugins/Experimental/MeshModelingToolsetExp/Source/MeshModelingToolsExp/Public/Properties/MeshStatisticsProperties.h
+++ b/Engine/Plugins/Experimental/MeshModelingToolsetExp/Source/MeshModelingToolsExp/Public/Properties/MeshStatisticsProperties.h
@@ -18,13 +18,13 @@ class MESHMODELINGTOOLSEXP_API UMeshStatisticsProperties : public UInteractiveTo
 	GENERATED_BODY()
 public:
 
-	UPROPERTY(VisibleAnywhere, Category = MeshStatistics)
+	UPROPERTY(VisibleAnywhere, Category = MeshStatistics, meta = (NoResetToDefault))
 	FString Mesh;
 
-	UPROPERTY(VisibleAnywhere, Category = MeshStatistics, AdvancedDisplay)
+	UPROPERTY(VisibleAnywhere, Category = MeshStatistics, AdvancedDisplay, meta = (NoResetToDefault))
 	FString UV;
 
-	UPROPERTY(VisibleAnywhere, Category = MeshStatistics, AdvancedDisplay)
+	UPROPERTY(VisibleAnywhere, Category = MeshStatistics, AdvancedDisplay, meta = (NoResetToDefault))
 	FString Attributes;
 
 	void Update(const FDynamicMesh3& Mesh);
