Geometry Tools:  SubdividePolyTool.cpp - Make sure we don't stomp on any tool-based warnings.

diff --git a/Engine/Plugins/Experimental/MeshModelingToolsetExp/Source/MeshModelingToolsEditorOnlyExp/Private/SubdividePolyTool.cpp b/Engine/Plugins/Experimental/MeshModelingToolsetExp/Source/MeshModelingToolsEditorOnlyExp/Private/SubdividePolyTool.cpp
--- a/Engine/Plugins/Experimental/MeshModelingToolsetExp/Source/MeshModelingToolsEditorOnlyExp/Private/SubdividePolyTool.cpp
+++ b/Engine/Plugins/Experimental/MeshModelingToolsetExp/Source/MeshModelingToolsEditorOnlyExp/Private/SubdividePolyTool.cpp
@@ -74,6 +74,8 @@ USingleSelectionMeshEditingTool* USubdividePolyToolBuilder::CreateNewTool(const
 
 bool USubdividePolyTool::CheckGroupTopology(FText& Message)
 {
+	Message = FText();
+
 	FGroupTopology Topo(OriginalMesh.Get(), true);
 	FSubdividePoly TempSubD(Topo, *OriginalMesh, 1);
 	FSubdividePoly::ETopologyCheckResult CheckResult = TempSubD.ValidateTopology();
@@ -135,20 +137,25 @@ void USubdividePolyTool::CapSubdivisionLevel(ESubdivisionScheme Scheme, int Desi
 	}
 	int MaxLevel = (int)floor(log2(MaxFaces / (NumOriginalFaces+1)) / 2.0);
 
+	FText WarningText;
 	if (DesiredLevel > MaxLevel)
 	{
-		FText WarningText = FText::Format(LOCTEXT("SubdivisionLevelTooHigh", "Subdivision level clamped: desired subdivision level ({0}) exceeds maximum level ({1}) for a mesh with this number of faces."),
+		WarningText = FText::Format(LOCTEXT("SubdivisionLevelTooHigh", "Subdivision level clamped: desired subdivision level ({0}) exceeds maximum level ({1}) for a mesh with this number of faces."),
 										  FText::AsNumber(DesiredLevel),
 										  FText::AsNumber(MaxLevel));
-		GetToolManager()->DisplayMessage(WarningText, EToolMessageLevel::UserWarning);
+		
+
 		Properties->SubdivisionLevel = MaxLevel;
 		Properties->SilentUpdateWatched();		// Don't trigger this function again due to setting SubdivisionLevel above
 	}
-	else
+	
+	if (!PersistentErrorMessage.IsEmpty())
 	{
-		// Clear possible lingering warning message
-		GetToolManager()->DisplayMessage(FText(), EToolMessageLevel::UserWarning);
+		FText Delimiter = FText::FromString("\n");
+		WarningText = WarningText.IsEmpty() ?  PersistentErrorMessage : FText::Join(Delimiter, PersistentErrorMessage, WarningText);
 	}
+	// if WarningText is empty this will clear possible lingering warning message 
+	GetToolManager()->DisplayMessage(WarningText, EToolMessageLevel::UserWarning);
 }
 
 
@@ -162,25 +169,20 @@ void USubdividePolyTool::Setup()
 		return;
 	}
 
+	GetToolManager()->DisplayMessage(LOCTEXT("SubdividePolyToolMessage", "Set the subdivision level and hit Accept to create a new subdivided mesh"), EToolMessageLevel::UserNotification);
+
 	bool bWantVertexNormals = false;
 	OriginalMesh = MakeShared<FDynamicMesh3, ESPMode::ThreadSafe>(bWantVertexNormals, false, false, false);
 	FMeshDescriptionToDynamicMesh Converter;
 	Converter.Convert(UE::ToolTarget::GetMeshDescription(Target), *OriginalMesh);
 
-	FText ErrorMessage;
-	bool bCatmullClarkOK = CheckGroupTopology(ErrorMessage);
+	const bool bCatmullClarkOK = CheckGroupTopology(PersistentErrorMessage);
 
-	if (bCatmullClarkOK)
-	{
-		GetToolManager()->DisplayMessage(LOCTEXT("SubdividePolyToolMessage",
-												 "Set the subdivision level and hit Accept to create a new subdivided mesh"),
-										 EToolMessageLevel::UserNotification);
-	}
-	else
+	if (!bCatmullClarkOK)
 	{
-		GetToolManager()->DisplayMessage(ErrorMessage, EToolMessageLevel::UserWarning);
+		GetToolManager()->DisplayMessage(PersistentErrorMessage, EToolMessageLevel::UserWarning);
 	}
-
+	
 	Properties = NewObject<USubdividePolyToolProperties>(this, TEXT("Subdivide Mesh Tool Settings"));
 	Properties->RestoreProperties(this);
 
diff --git a/Engine/Plugins/Experimental/MeshModelingToolsetExp/Source/MeshModelingToolsEditorOnlyExp/Public/SubdividePolyTool.h b/Engine/Plugins/Experimental/MeshModelingToolsetExp/Source/MeshModelingToolsEditorOnlyExp/Public/SubdividePolyTool.h
--- a/Engine/Plugins/Experimental/MeshModelingToolsetExp/Source/MeshModelingToolsEditorOnlyExp/Public/SubdividePolyTool.h
+++ b/Engine/Plugins/Experimental/MeshModelingToolsetExp/Source/MeshModelingToolsEditorOnlyExp/Public/SubdividePolyTool.h
@@ -110,5 +110,8 @@ protected:
 	bool CheckGroupTopology(FText& Message);
 
 	void CapSubdivisionLevel(ESubdivisionScheme Scheme, int DesiredLevel);
+
+	// Tool message when no explicit parameter-related warnings are being pushed.  
+	FText PersistentErrorMessage;
 };
 
