clarify what 'Discard Attributes' does in remeshing

diff --git a/Engine/Plugins/Experimental/MeshModelingToolsetExp/Source/MeshModelingToolsExp/Public/RemeshMeshTool.h b/Engine/Plugins/Experimental/MeshModelingToolsetExp/Source/MeshModelingToolsExp/Public/RemeshMeshTool.h
--- a/Engine/Plugins/Experimental/MeshModelingToolsetExp/Source/MeshModelingToolsExp/Public/RemeshMeshTool.h
+++ b/Engine/Plugins/Experimental/MeshModelingToolsetExp/Source/MeshModelingToolsExp/Public/RemeshMeshTool.h
@@ -52,7 +52,7 @@ public:
 	UPROPERTY(EditAnywhere, Category = Remeshing)
 	ERemeshSmoothingType SmoothingType;
 
-	/** If true, UVs and Normals are discarded  */
+	/** If true, UVs and existing normals are discarded, allowing the remesher to ignore any UV and normal seams. New per-vertex normals are computed. */
 	UPROPERTY(EditAnywhere, Category = Remeshing)
 	bool bDiscardAttributes;
 
