PatternTool: Add ability to project pattern elements into world along the tool's plane's normal.

The projection is done on per pattern item basis and will not project one pattern item onto another, meaning that projecting something like a vertical line of boxes down the Z axis will result in many overlapping boxes.

Projection offset can be used to fine tune the final distance between the result pattern and the components in the scene onto which the pattern is being projected, usually because the location of the pivot of the source mesh results in the projected pattern either floating above or intersecting the components after the projection.

diff --git a/Engine/Plugins/Experimental/MeshModelingToolsetExp/Source/MeshModelingToolsExp/Private/PatternTool.cpp b/Engine/Plugins/Experimental/MeshModelingToolsetExp/Source/MeshModelingToolsExp/Private/PatternTool.cpp
--- a/Engine/Plugins/Experimental/MeshModelingToolsetExp/Source/MeshModelingToolsExp/Private/PatternTool.cpp
+++ b/Engine/Plugins/Experimental/MeshModelingToolsetExp/Source/MeshModelingToolsExp/Private/PatternTool.cpp
@@ -233,7 +233,6 @@ public:
 			check(false);
 		}
 	}
-
 };
 
 
@@ -288,7 +287,6 @@ public:
 
 	void UpdatePattern_LineFill();
 	void UpdatePattern_RectangleFill();
-
 };
 
 
@@ -508,6 +506,8 @@ void UPatternTool::Setup()
 	Settings->WatchProperty(Settings->SinglePlane, [this](EPatternToolSinglePlane) { OnParametersUpdated(); } );
 	Settings->WatchProperty(Settings->bHideSources, [this](bool bNewValue) { OnSourceVisibilityToggled(!bNewValue); } );
 	Settings->WatchProperty(Settings->Seed, [this](int32 NewSeed) { MarkPatternDirty(); } );
+	Settings->WatchProperty(Settings->bProjectElementsDown, [this](bool bNewValue) { MarkPatternDirty(); } );
+	Settings->WatchProperty(Settings->ProjectionOffset, [this](float NewValue) { MarkPatternDirty(); } );
 
 	LinearSettings = NewObject<UPatternTool_LinearSettings>();
 	AddToolPropertySource(LinearSettings);
@@ -590,6 +590,7 @@ void UPatternTool::OnShutdown(EToolShutdownType ShutdownType)
 		Component->DestroyComponent();
 	}
 	AllComponents.Reset();
+	AllPreviewComponents.Reset();
 	PreviewComponents.Reset();
 	StaticMeshPools.Reset();
 	DynamicMeshPools.Reset();
@@ -956,8 +957,6 @@ void UPatternTool::GetPatternTransforms_Radial(TArray<UE::Geometry::FTransformSR
 	TransformsOut = MoveTemp(Generator.Pattern);
 }
 
-
-
 void UPatternTool::UpdatePattern()
 {
 	TArray<FTransformSRT3d> Pattern;
@@ -1001,14 +1000,14 @@ void UPatternTool::UpdatePattern()
 			}
 			if (Component != nullptr)
 			{
-				FTransform PatternTransform = (FTransform)Pattern[k];
-				FTransform WorldTransform = (FTransform)ElementTransform * PatternTransform * CurrentStartFrameWorld.ToFTransform();
+				FTransform WorldTransform;
+				ComputeWorldTransform(WorldTransform, ElementTransform, Pattern[k]);
+				
 				Component->SetWorldTransform( WorldTransform );
 				Component->SetVisibility(true);
-
+		
 				ElemComponents.Components.Add(Component);
 			}
-
 		}
 	}
 
@@ -1022,7 +1021,23 @@ void UPatternTool::UpdatePattern()
 	CurrentPattern = MoveTemp(Pattern);
 }
 
-
+void UPatternTool::ComputeWorldTransform(FTransform& OutWorldTransform, const FTransform& InElementTransform, const FTransform& InPatternTransform) const
+{
+	OutWorldTransform = InElementTransform * InPatternTransform * CurrentStartFrameWorld.ToFTransform();
+	
+	if (Settings->bProjectElementsDown)
+	{
+		FVector3d ProjectionAxis = -CurrentStartFrameWorld.Z();
+		const FRay WorldRay(OutWorldTransform.GetTranslation(), ProjectionAxis);
+		FHitResult HitResult;
+	
+		if (ToolSceneQueriesUtil::FindNearestVisibleObjectHit(this, HitResult, WorldRay, &AllPreviewComponents, nullptr))
+		{
+			FVector3d ProjectionTranslation = ProjectionAxis * (HitResult.Distance + Settings->ProjectionOffset);
+			OutWorldTransform.SetTranslation(OutWorldTransform.GetTranslation() + ProjectionTranslation);
+		}
+	}
+}
 
 
 UStaticMeshComponent* UPatternTool::GetPreviewStaticMesh(FPatternElement& Element)
@@ -1054,6 +1069,7 @@ UStaticMeshComponent* UPatternTool::GetPreviewStaticMesh(FPatternElement& Elemen
 	StaticMeshComp->RegisterComponent();
 
 	AllComponents.Add(StaticMeshComp);
+	AllPreviewComponents.Add(StaticMeshComp);
 
 	return StaticMeshComp;
 }
@@ -1111,6 +1127,7 @@ UDynamicMeshComponent* UPatternTool::GetPreviewDynamicMesh(FPatternElement& Elem
 	DynamicMeshComp->SetMesh(MoveTemp(ElementMeshCopy));
 	
 	AllComponents.Add(DynamicMeshComp);
+	AllPreviewComponents.Add(DynamicMeshComp);
 
 	return DynamicMeshComp;
 }
@@ -1176,6 +1193,26 @@ void UPatternTool::EmitResults()
 
 	TArray<AActor*> NewActors;		// set of new actors created by operation
 
+	// Used when constructing actors, not appending components to an actor. Otherwise UPrimitiveComponent::SetVisibility()
+	// is used instead. Keeping everything invisible until the end allows us to avoid hitting the created components
+	// when projecting downward, just like we do in the previews via ignored components.
+	constexpr bool bPropagateToChildren = true;
+	auto SetActorComponentsVisibility = [bPropagateToChildren](AActor* InActor, const bool bNewVisibility)
+	{
+		if (const AStaticMeshActor* ActorAsStaticMeshActor = Cast<AStaticMeshActor, AActor>(InActor))
+		{
+			ActorAsStaticMeshActor->GetStaticMeshComponent()->SetVisibility(bNewVisibility, bPropagateToChildren);
+		}
+		
+		for (UActorComponent* Component : InActor->GetComponents())
+		{
+			if (UPrimitiveComponent* ActorComponentAsPrimitiveComponent = Cast<UPrimitiveComponent, UActorComponent>(Component))
+			{
+				ActorComponentAsPrimitiveComponent->SetVisibility(bNewVisibility, bPropagateToChildren);
+			}
+		}
+	};
+
 	// TODO: investigate use of CopyPropertiesForUnrelatedObjects to transfer settings from source to target Components/Actors
 
 	for (int32 ElemIdx = 0; ElemIdx < NumElements; ++ElemIdx)
@@ -1216,8 +1253,9 @@ void UPatternTool::EmitResults()
 				for (int32 k = 0; k < NumPatternItems; ++k)
 				{
 					FTransformSRT3d PatternTransform = CurrentPattern[k];
-					FTransform WorldTransform = (FTransform)ElementTransform * (FTransform)PatternTransform * CurrentStartFrameWorld.ToFTransform();
-
+					FTransform WorldTransform;
+					ComputeWorldTransform(WorldTransform, (FTransform)ElementTransform, (FTransform)PatternTransform);
+					
 					FDynamicMesh3 PatternItemMesh = ElementMesh;
 					// TODO: may need to bake nonuniform scale in here, if we allow scaling in transform
 					FCreateMeshObjectResult Result = EmitDynamicMeshActor(MoveTemp(PatternItemMesh),
@@ -1225,6 +1263,7 @@ void UPatternTool::EmitResults()
 					if (Result.IsOK())
 					{
 						NewActors.Add(Result.NewActor);
+						SetActorComponentsVisibility(Result.NewActor, false);
 					}
 				}
 			}
@@ -1236,20 +1275,29 @@ void UPatternTool::EmitResults()
 				FMeshIndexMappings Mappings;
 				for (int32 k = 0; k < NumPatternItems; ++k)
 				{
+					FTransformSRT3d PatternTransform = CurrentPattern[k];
+					FTransform WorldTransform;
+					ComputeWorldTransform(WorldTransform, (FTransform)ElementTransform, (FTransform)PatternTransform);
+
+					// We have world transforms of the components, but need their local transforms in the space of the
+					// final mesh, whose pivot is determined by CurrentStartFrameWorld. So we apply the inverse of that
+					// while appending them (the inverse FTransform is accurate in this case because
+					// CurrentStartFrameWorld can't be nonuniformly scaled).
 					FTransformSequence3d TransformSeq;
-					TransformSeq.Append(ElementTransform);
-					TransformSeq.Append(CurrentPattern[k]);
+					TransformSeq.Append(WorldTransform);
+					TransformSeq.Append(CurrentStartFrameWorld.ToFTransform().Inverse());
 					Mappings.Reset();
 					Editor.AppendMesh(&ElementMesh, Mappings,
 						[&](int, const FVector3d& Position) { return TransformSeq.TransformPosition(Position); },
 						[&](int, const FVector3d& Normal) { return TransformSeq.TransformNormal(Normal); } );
 				}
-
+				
 				FCreateMeshObjectResult Result = EmitDynamicMeshActor(MoveTemp(CombinedPatternMesh),
-					FString::Printf(TEXT("Pattern_%d"), ElemIdx), CurrentStartFrameWorld.ToFTransform());
+				FString::Printf(TEXT("Pattern_%d"), ElemIdx), CurrentStartFrameWorld.ToFTransform());
 				if (Result.IsOK())
 				{
 					NewActors.Add(Result.NewActor);
+					SetActorComponentsVisibility(Result.NewActor, false);
 				}
 			}
 		}
@@ -1264,8 +1312,9 @@ void UPatternTool::EmitResults()
 				for (int32 k = 0; k < NumPatternItems; ++k)
 				{
 					FTransformSRT3d PatternTransform = CurrentPattern[k];
-					FTransform WorldTransform = (FTransform)ElementTransform * (FTransform)PatternTransform * CurrentStartFrameWorld.ToFTransform();
-
+					FTransform WorldTransform;
+					ComputeWorldTransform(WorldTransform, (FTransform)ElementTransform, (FTransform)PatternTransform);
+					
 					FActorSpawnParameters SpawnInfo;
 					SpawnInfo.Template = SourceActor;
 					AStaticMeshActor* NewActor = GetTargetWorld()->SpawnActor<AStaticMeshActor>(SpawnInfo);
@@ -1273,6 +1322,7 @@ void UPatternTool::EmitResults()
 					{
 						NewActor->SetActorTransform(WorldTransform);
 						NewActors.Add(NewActor);
+						SetActorComponentsVisibility(NewActor, false);
 					}
 				}
 			}
@@ -1285,6 +1335,7 @@ void UPatternTool::EmitResults()
 					NewActors.Add(NewActor);
 
 					UInstancedStaticMeshComponent* ISMComponent = NewObject<UInstancedStaticMeshComponent>(NewActor);
+					ISMComponent->SetVisibility(false, bPropagateToChildren);
 					ISMComponent->SetFlags(RF_Transactional);
 					ISMComponent->bHasPerInstanceHitProxies = true;
 
@@ -1303,7 +1354,11 @@ void UPatternTool::EmitResults()
 					for (int32 k = 0; k < NumPatternItems; ++k)
 					{
 						FTransformSRT3d PatternTransform = CurrentPattern[k];
-						ISMComponent->AddInstance((FTransform)ElementTransform * (FTransform)PatternTransform);
+						FTransform WorldTransform;
+						ComputeWorldTransform(WorldTransform, (FTransform)ElementTransform, (FTransform)PatternTransform);
+
+						constexpr bool bTransformInWorldSpace = true;
+						ISMComponent->AddInstance(WorldTransform, bTransformInWorldSpace);
 					}
 
 					ISMComponent->RegisterComponent();
@@ -1318,14 +1373,18 @@ void UPatternTool::EmitResults()
 				for (int32 k = 0; k < NumPatternItems; ++k)
 				{
 					FTransformSRT3d PatternTransform = CurrentPattern[k];
-					FTransform WorldTransform = (FTransform)PatternTransform * CurrentStartFrameWorld.ToFTransform();
+					FTransform WorldTransform;
+					ComputeWorldTransform(WorldTransform, (FTransform)ElementTransform, (FTransform)PatternTransform);
+					
 					if (k == 0)
 					{
 						FActorSpawnParameters SpawnInfo;
 						SpawnInfo.Template = SourceActor;
 						ParentActor = GetTargetWorld()->SpawnActor<AStaticMeshActor>(SpawnInfo);
+						NewActors.Add(ParentActor);
 						ParentActor->SetActorTransform(CurrentStartFrameWorld.ToFTransform());
 						ParentActor->GetStaticMeshComponent()->SetWorldTransform(WorldTransform);
+						ParentActor->GetStaticMeshComponent()->SetVisibility(false, bPropagateToChildren);
 					}
 					else if (k == 1)
 					{
@@ -1336,6 +1395,7 @@ void UPatternTool::EmitResults()
 						ParentActor->AddInstanceComponent(TemplateComponent);
 						TemplateComponent->RegisterComponent();
 						TemplateComponent->SetWorldTransform( WorldTransform );
+						TemplateComponent->SetVisibility(false, bPropagateToChildren);
 					}
 					else
 					{
@@ -1346,13 +1406,19 @@ void UPatternTool::EmitResults()
 						ParentActor->AddInstanceComponent(NewCloneComponent);
 						NewCloneComponent->RegisterComponent();
 						NewCloneComponent->SetWorldTransform( WorldTransform );
+						NewCloneComponent->SetVisibility(false, bPropagateToChildren);
 					}
 				}			
-
 			}
 		}
 	}
 
+	// Make all components visible now that they have all been created
+	for (auto Actor : NewActors)
+	{
+		SetActorComponentsVisibility(Actor, true);
+	}
+
 	if (NewActors.Num() > 0)
 	{
 		ToolSelectionUtil::SetNewActorSelection(GetToolManager(), NewActors);
diff --git a/Engine/Plugins/Experimental/MeshModelingToolsetExp/Source/MeshModelingToolsExp/Public/PatternTool.h b/Engine/Plugins/Experimental/MeshModelingToolsetExp/Source/MeshModelingToolsExp/Public/PatternTool.h
--- a/Engine/Plugins/Experimental/MeshModelingToolsetExp/Source/MeshModelingToolsExp/Public/PatternTool.h
+++ b/Engine/Plugins/Experimental/MeshModelingToolsetExp/Source/MeshModelingToolsExp/Public/PatternTool.h
@@ -110,13 +110,21 @@ public:
 	UPROPERTY(EditAnywhere, Category = Shape, meta = (DisplayName = "Plane", EditCondition = "Shape != EPatternToolShape::Line", EditConditionHides))
 	EPatternToolSinglePlane SinglePlane = EPatternToolSinglePlane::XYPlane;
 
-	/** Shape of the underlying Pattern */
+	/** Hide the source meshes when enabled */
 	UPROPERTY(EditAnywhere, Category = Shape)
 	bool bHideSources = true;
 
 	/** The seed used to introduce random transform variations when enabled */
 	UPROPERTY(EditAnywhere, Category = Shape)
 	int32 Seed = FMath::Rand();
+
+	/** Whether or not the pattern items should be projected along the negative Z axis of the plane mechanic */
+	UPROPERTY(EditAnywhere, Category = Shape)
+	bool bProjectElementsDown = false;
+
+	/** How much each pattern item should be moved along the negative Z axis of the plane mechanic if Project Elements Down is enabled */
+	UPROPERTY(EditAnywhere, Category = Shape, meta = (EditCondition = "bProjectElementsDown == true", EditConditionHides))
+	float ProjectionOffset = 0.0f;
 };
 
 
@@ -465,6 +473,7 @@ protected:
 	void OnShapeUpdated();
 	void OnParametersUpdated();
 	void UpdatePattern();
+	void ComputeWorldTransform(FTransform& OutWorldTransform, const FTransform& InElementTransform, const FTransform& InPatternTransform) const;
 	void GetPatternTransforms_Linear(TArray<UE::Geometry::FTransformSRT3d>& TransformsOut);
 	void GetPatternTransforms_Grid(TArray<UE::Geometry::FTransformSRT3d>& TransformsOut);
 	void GetPatternTransforms_Radial(TArray<UE::Geometry::FTransformSRT3d>& TransformsOut);
@@ -495,7 +504,13 @@ protected:
 		TArray<UPrimitiveComponent*> Components;
 	};
 	TArray<FComponentSet> PreviewComponents;
+	
+	// This duplicates the data stored by PreviewComponents but it is necessary to have a simple
+	// TArray<UPrimitiveComponent*> of all preview components when raycasting, otherwise the raycasts
+	// will hit the preview components from the previous frame.
+	TArray<const UPrimitiveComponent*> AllPreviewComponents;			// This is passed to FindNearestVisibleObjectHit as ComponentsToIgnore
 
+protected:
 	UPROPERTY()
 	TSet<TObjectPtr<UPrimitiveComponent>> AllComponents;		// to keep components in FComponentSet alive
 
