PatternTool: Hide Jitter UProperties which were prematurely robomerged to Release-5.1 from Main.

diff --git a/Engine/Plugins/Experimental/MeshModelingToolsetExp/Source/MeshModelingToolsExp/Public/PatternTool.h b/Engine/Plugins/Experimental/MeshModelingToolsetExp/Source/MeshModelingToolsExp/Public/PatternTool.h
--- a/Engine/Plugins/Experimental/MeshModelingToolsetExp/Source/MeshModelingToolsExp/Public/PatternTool.h
+++ b/Engine/Plugins/Experimental/MeshModelingToolsetExp/Source/MeshModelingToolsExp/Public/PatternTool.h
@@ -271,7 +271,7 @@ public:
 	bool bInterpolate = false;
 
 	/** If true, Rotation at each Pattern Element is offset by a uniformly chosen random value in the range of [-RotationJitterRange, RotationJitterRange] */
-	UPROPERTY(EditAnywhere, Category = Rotation)
+	UPROPERTY(/*EditAnywhere, Category = Rotation*/)
 	bool bJitter = false;
 
 	/** Rotation at first Pattern Element */
@@ -283,7 +283,7 @@ public:
 	FRotator Rotation = FRotator::ZeroRotator;
 
 	/** Upper bound of the range which is sampled to randomly rotate each Pattern Element if Jitter is true */
-	UPROPERTY(EditAnywhere, Category = Rotation, meta = (ClampMin = 0, EditCondition = "bJitter", EditConditionHides, HideEditConditionToggle))
+	UPROPERTY(/*EditAnywhere, Category = Rotation, meta = (ClampMin = 0, EditCondition = "bJitter", EditConditionHides, HideEditConditionToggle)*/)
 	FRotator RotationJitterRange = FRotator::ZeroRotator;
 };
 
@@ -301,7 +301,7 @@ public:
 	bool bInterpolate = false;
 
 	/** If true, Translation at each Pattern Element is offset by a uniformly chosen random value in the range of [-TranslationJitterRange, TranslationJitterRange] */
-	UPROPERTY(EditAnywhere, Category = Translation)
+	UPROPERTY(/*EditAnywhere, Category = Translation*/)
 	bool bJitter = false;
 
 	/** Translation at first Pattern Element */
@@ -313,7 +313,7 @@ public:
 	FVector Translation = FVector::ZeroVector;
 
 	/** Upper bound of the range which is sampled to randomly translate each Pattern Element if Jitter is true */
-	UPROPERTY(EditAnywhere, Category = Translation, meta = (ClampMin = 0, EditCondition = "bJitter", EditConditionHides, HideEditConditionToggle))
+	UPROPERTY(/*EditAnywhere, Category = Translation, meta = (ClampMin = 0, EditCondition = "bJitter", EditConditionHides, HideEditConditionToggle)*/)
 	FVector TranslationJitterRange = FVector::ZeroVector;
 };
 
@@ -335,7 +335,7 @@ public:
 	bool bInterpolate = false;
 
 	/** If true, Scale at each Pattern Element is offset by a uniformly chosen random value in the range of [-ScaleJitterRange, ScaleJitterRange] */
-	UPROPERTY(EditAnywhere, Category = Scale)
+	UPROPERTY(/*EditAnywhere, Category = Scale*/)
 	bool bJitter = false;
 	
 	/** Uniform Scale at first Pattern Element */
@@ -347,7 +347,7 @@ public:
 	float Scale = 1;
 
 	/** Upper bound of the range which is sampled to randomly scale each Pattern Element if Jitter is true */
-	UPROPERTY(EditAnywhere, Category = Scale, meta = (ClampMin = 0, EditCondition = "bJitter && bUniform", EditConditionHides, HideEditConditionToggle))
+	UPROPERTY(/*EditAnywhere, Category = Scale, meta = (ClampMin = 0, EditCondition = "bJitter && bUniform", EditConditionHides, HideEditConditionToggle)*/)
 	float ScaleJitterRange = 0.0f;
 
 	/** Non-Uniform Scale at first Pattern Element */
@@ -359,7 +359,7 @@ public:
 	FVector ScaleNonUniform = FVector(1,1,1);
 
 	/** Upper bound of the range which is sampled to randomly scale each Pattern Element if Jitter is true (Non-Uniform) */
-	UPROPERTY(EditAnywhere, Category = Scale, meta = (ClampMin = 0, DisplayName = "Scale Jitter Range", EditCondition = "bJitter && bUniform == false", EditConditionHides, HideEditConditionToggle))
+	UPROPERTY(/*EditAnywhere, Category = Scale, meta = (ClampMin = 0, DisplayName = "Scale Jitter Range", EditCondition = "bJitter && bUniform == false", EditConditionHides, HideEditConditionToggle)*/)
 	FVector ScaleJitterRangeNonUniform = FVector(0, 0, 0);
 };
 
