PatternTool: Fix issues with input fields in the pattern tool.

Extent, radius, and step size can no longer be negative.

Start/end angle are UI limited to the range [-360, 360] but can accept any values. Very large values are useful to create spirals and it would feel a little bit arbitrary to just choose a hard range. Angle shift is clamped to [-180, 180] because values outside of that range are redundant.

Count, CountX, and CountY are clamped to a minimum of 1 with a UI max of 25. The behavior here is unchanged but several redundant UIMins were removed because ClampMin was also already set.

diff --git a/Engine/Plugins/Experimental/MeshModelingToolsetExp/Source/MeshModelingToolsExp/Public/PatternTool.h b/Engine/Plugins/Experimental/MeshModelingToolsetExp/Source/MeshModelingToolsExp/Public/PatternTool.h
--- a/Engine/Plugins/Experimental/MeshModelingToolsetExp/Source/MeshModelingToolsExp/Public/PatternTool.h
+++ b/Engine/Plugins/Experimental/MeshModelingToolsetExp/Source/MeshModelingToolsExp/Public/PatternTool.h
@@ -141,15 +141,15 @@ public:
 	EPatternToolAxisSpacingMode SpacingMode = EPatternToolAxisSpacingMode::ByCount;
 
 	/** Number of Pattern Elements to place */
-	UPROPERTY(EditAnywhere, Category = LinearPattern, meta = (UIMin = 1, UIMax = 25, ClampMin = 1, EditCondition = "SpacingMode == EPatternToolAxisSpacingMode::ByCount"))
+	UPROPERTY(EditAnywhere, Category = LinearPattern, meta = (ClampMin = 1, UIMax = 25, EditCondition = "SpacingMode == EPatternToolAxisSpacingMode::ByCount"))
 	int32 Count = 10;
 
 	/** Fixed Increment used to place Pattern Elements */
-	UPROPERTY(EditAnywhere, Category = LinearPattern, meta = (UIMin = 0, EditCondition = "SpacingMode == EPatternToolAxisSpacingMode::StepSize"))
+	UPROPERTY(EditAnywhere, Category = LinearPattern, meta = (ClampMin = 0, EditCondition = "SpacingMode == EPatternToolAxisSpacingMode::StepSize"))
 	double StepSize = 100.0;
 
 	/** Length of Pattern along the Axis */
-	UPROPERTY(EditAnywhere, Category = LinearPattern, meta = (UIMin = 0))
+	UPROPERTY(EditAnywhere, Category = LinearPattern, meta = (ClampMin = 0))
 	double Extent = 1000.0;
 
 	/** If true, Pattern is centered at the Origin, otherwise Pattern starts at the Origin */
@@ -174,15 +174,15 @@ public:
 	EPatternToolAxisSpacingMode SpacingX = EPatternToolAxisSpacingMode::ByCount;
 
 	/** Number of  Pattern Elements to place along the Main axis */
-	UPROPERTY(EditAnywhere, Category = GridPatternX, meta = (UIMin = 1, UIMax = 25, ClampMin = 1, EditCondition = "SpacingX == EPatternToolAxisSpacingMode::ByCount"))
+	UPROPERTY(EditAnywhere, Category = GridPatternX, meta = (ClampMin = 1, UIMax = 25, EditCondition = "SpacingX == EPatternToolAxisSpacingMode::ByCount"))
 	int32 CountX = 10;
 
 	/** Fixed Increment used to place Pattern Elements along the Main axis */
-	UPROPERTY(EditAnywhere, Category = GridPatternX, meta = (UIMin = 0, EditCondition = "SpacingX == EPatternToolAxisSpacingMode::StepSize"))
+	UPROPERTY(EditAnywhere, Category = GridPatternX, meta = (ClampMin = 0, EditCondition = "SpacingX == EPatternToolAxisSpacingMode::StepSize"))
 	double StepSizeX = 100.0;
 
 	/** Length/Extent of Pattern falong the Main Axis */
-	UPROPERTY(EditAnywhere, Category = GridPatternX, meta = (UIMin = 0))
+	UPROPERTY(EditAnywhere, Category = GridPatternX, meta = (ClampMin = 0))
 	double ExtentX = 1000.0;
 
 	/** If true, Pattern is centered at the Origin along the Main axis, otherwise Pattern starts at the Origin */
@@ -194,15 +194,15 @@ public:
 	EPatternToolAxisSpacingMode SpacingY = EPatternToolAxisSpacingMode::ByCount;
 
 	/** Number of  Pattern Elements to place along the Secondary axis */
-	UPROPERTY(EditAnywhere, Category = GridPatternY, meta = (UIMin = 1, UIMax = 25, ClampMin = 1, EditCondition = "SpacingY == EPatternToolAxisSpacingMode::ByCount"))
+	UPROPERTY(EditAnywhere, Category = GridPatternY, meta = (ClampMin = 1, UIMax = 25, EditCondition = "SpacingY == EPatternToolAxisSpacingMode::ByCount"))
 	int32 CountY = 10;
 
 	/** Fixed Increment used to place Pattern Elements along the Secondary axis */
-	UPROPERTY(EditAnywhere, Category = GridPatternY, meta = (UIMin = 0, EditCondition = "SpacingY == EPatternToolAxisSpacingMode::StepSize"))
+	UPROPERTY(EditAnywhere, Category = GridPatternY, meta = (ClampMin = 0, EditCondition = "SpacingY == EPatternToolAxisSpacingMode::StepSize"))
 	double StepSizeY = 100.0;
 
 	/** Length/Extent of Pattern falong the Secondary Axis */
-	UPROPERTY(EditAnywhere, Category = GridPatternY, meta = (UIMin = 0))
+	UPROPERTY(EditAnywhere, Category = GridPatternY, meta = (ClampMin = 0))
 	double ExtentY = 1000.0;
 
 	/** If true, Pattern is centered at the Origin along the Secondary axis, otherwise Pattern starts at the Origin */
@@ -227,27 +227,27 @@ public:
 	EPatternToolAxisSpacingMode SpacingMode = EPatternToolAxisSpacingMode::ByCount;
 
 	/** Number of  Pattern Elements to place */
-	UPROPERTY(EditAnywhere, Category = RadialPattern, meta = (UIMin = 1, UIMax = 25, ClampMin = 1, EditCondition = "SpacingMode == EPatternToolAxisSpacingMode::ByCount"))
+	UPROPERTY(EditAnywhere, Category = RadialPattern, meta = (ClampMin = 1, UIMax = 25, EditCondition = "SpacingMode == EPatternToolAxisSpacingMode::ByCount"))
 	int32 Count = 10;
 
 	/** Fixed Increment (in Degrees) used to position Pattern Elements around the Circle/Arc */
-	UPROPERTY(EditAnywhere, Category = RadialPattern, meta = (Units = "Degrees", UIMin = 0, EditCondition = "SpacingMode == EPatternToolAxisSpacingMode::StepSize"))
+	UPROPERTY(EditAnywhere, Category = RadialPattern, meta = (Units = "Degrees", ClampMin = 0, EditCondition = "SpacingMode == EPatternToolAxisSpacingMode::StepSize"))
 	double StepSize = 100.0;
 
 	/** Radius of the Circle/Arc */
-	UPROPERTY(EditAnywhere, Category = RadialPattern, meta = (UIMin = 0))
+	UPROPERTY(EditAnywhere, Category = RadialPattern, meta = (ClampMin = 0))
 	double Radius = 250;
 
 	/** Start angle of the Circle/Arc */
-	UPROPERTY(EditAnywhere, Category = RadialPattern, meta = (Units = "Degrees", UIMin = -360))
+	UPROPERTY(EditAnywhere, Category = RadialPattern, meta = (Units = "Degrees", UIMin = -360, UIMax = 360))
 	double StartAngle = 0.0;
 
 	/** End angle of the Circle/Arc */
-	UPROPERTY(EditAnywhere, Category = RadialPattern, meta = (Units = "Degrees", UIMax = 360))
+	UPROPERTY(EditAnywhere, Category = RadialPattern, meta = (Units = "Degrees", UIMin = -360, UIMax = 360))
 	double EndAngle = 360.0;
 
 	/** Fixed offset added to Start/End Angles */
-	UPROPERTY(EditAnywhere, Category = RadialPattern, meta = (Units = "Degrees", UIMin = -180, UIMax = 180))
+	UPROPERTY(EditAnywhere, Category = RadialPattern, meta = (Units = "Degrees", ClampMin = -180, ClampMax = 180))
 	double AngleShift = 0.0;
 
 	/** If true, Pattern elements are rotated to align with the Circle tangen */
