MeshModelingTools: Fix CubeGrid issues with subtraction when starting with an empty mesh.

diff --git a/Engine/Plugins/Experimental/MeshModelingToolsetExp/Source/MeshModelingToolsExp/Private/CubeGridTool.cpp b/Engine/Plugins/Experimental/MeshModelingToolsetExp/Source/MeshModelingToolsExp/Private/CubeGridTool.cpp
--- a/Engine/Plugins/Experimental/MeshModelingToolsetExp/Source/MeshModelingToolsExp/Private/CubeGridTool.cpp
+++ b/Engine/Plugins/Experimental/MeshModelingToolsetExp/Source/MeshModelingToolsExp/Private/CubeGridTool.cpp
@@ -597,50 +597,25 @@ void UCubeGridTool::InvalidatePreview(bool bUpdateCornerLineSet)
 {
 	using namespace CubeGridToolLocals;
 
-	// Do the line set first before we get to all the early returns
+	// Do the line set first
 	if (bUpdateCornerLineSet && Mode == EMode::Corner)
 	{
 		UpdateCornerModeLineSet();
 	}
 
-	if (CurrentExtrudeAmount == 0)
+	// If we're guaranteed not to change the result, just reset the preview
+	if (CurrentExtrudeAmount == 0 || (CurrentMesh->TriangleCount() == 0 && CurrentExtrudeAmount < 0))
 	{
-		// Reset the preview
 		Preview->CancelCompute();
+		LastOpChangedTids = nullptr;
 		if (bPreviewMayDiffer)
 		{
 			Preview->PreviewMesh->UpdatePreview(CurrentMesh.Get());
 			bPreviewMayDiffer = false;
-			bWaitingToApplyPreview = false;
 		}
 		return;
 	}
-	else if (CurrentMesh->TriangleCount() == 0 && CurrentExtrudeAmount < 0 
-		&& Mode != EMode::Corner)
-	{
-		// We're subtracting from an empty mesh. Just slide the selection.
-		// This will also reset the preview.
-		SlideSelection(CurrentExtrudeAmount, true);
-		bPreviewMayDiffer = false;
-		bWaitingToApplyPreview = false;
-		return;
-	}
-
-	// If we didn't start with an existing mesh, and we are adding to an empty
-	// starting mesh, set the transform such that it is in the (grid) minimum of
-	// the selection. This frequently (though not always, in the case of
-	// a pyramid) places the pivot in a handy corner for snapping.
-	if (!Target && CurrentMesh->TriangleCount() == 0 && CurrentExtrudeAmount > 0
-		&& ensure(bHaveSelection))
-	{
-		FFrame3d GridFrame = CubeGrid->GetFrame();
-		GridFrame.Origin = GridFrame.FromFramePoint(Selection.Box.Min);
-		CurrentMeshTransform = GridFrame.ToTransform();
 
-		GridGizmoAlignmentMechanic->InitializeDeformedMeshRayCast([this]() { return MeshSpatial.Get(); }, CurrentMeshTransform, nullptr);
-	}
-
-	// Finally invalidate the preview
 	Preview->InvalidateResult();
 	bPreviewMayDiffer = true;
 }
@@ -669,6 +644,19 @@ TUniquePtr<FDynamicMeshOperator> UCubeGridTool::MakeNewOperator()
 	TUniquePtr<FCubeGridBooleanOp> Op = MakeUnique<FCubeGridBooleanOp>();
 	Op->InputMesh = ComputeStartMesh;
 	Op->InputTransform = CurrentMeshTransform;
+
+	// If we didn't start with an existing mesh, and we are adding to an empty
+	// starting mesh, set the transform such that it is in the (grid) minimum of
+	// the selection. This frequently (though not always, in some corner-mode cases)
+	// places the pivot in a handy corner for snapping.
+	if (!Target && ComputeStartMesh->TriangleCount() == 0 && CurrentExtrudeAmount > 0
+		&& ensure(bHaveSelection))
+	{
+		FFrame3d TransformFrame = CubeGrid->GetFrame();
+		TransformFrame.Origin = TransformFrame.FromFramePoint(Selection.Box.Min);
+		Op->InputTransform = TransformFrame.ToTransform();
+	}
+	
 	Op->bKeepInputTransform = true;
 	Op->WorldBox = WorldBox;
 	Op->bSubtract = CurrentExtrudeAmount < 0;
@@ -1055,6 +1043,7 @@ void UCubeGridTool::Setup()
 			UpdateGridTransform(Transform, /*bUpdateDetailPanel */ false);
 		});
 
+	MaterialProperties->SilentUpdateWatched();
 	Settings->SilentUpdateWatched();
 
 	UpdateComputeInputs();
@@ -1203,7 +1192,11 @@ void UCubeGridTool::OnTick(float DeltaTime)
 		PendingAction = ECubeGridToolAction::NoAction;
 	}
 
-	if (bWaitingToApplyPreview && Preview->HaveValidResult())
+	if (bWaitingToApplyPreview && 
+		// Note that the check for bPreviewMayDiffer is needed, because when we do a forced
+		// reset, HaveValidResult() returns false, yet we may still want to do an application
+		// to slide the selection.
+		(!bPreviewMayDiffer || Preview->HaveValidResult()))
 	{
 		ApplyPreview();
 	}
@@ -1212,41 +1205,58 @@ void UCubeGridTool::OnTick(float DeltaTime)
 void UCubeGridTool::ApplyPreview()
 {
 	using namespace CubeGridToolLocals;
+
+	bWaitingToApplyPreview = false;
+	bPreviewMayDiffer = false;
+	bBlockUntilPreviewUpdate = false;
+
 	const FText TransactionText = LOCTEXT("CubeGridToolTransactionName", "Block Tool Change");
 	GetToolManager()->BeginUndoTransaction(TransactionText);
 
-	FDynamicMeshChangeTracker ChangeTracker(CurrentMesh.Get());
-	ChangeTracker.BeginChange();
-	ChangeTracker.SaveTriangles(*LastOpChangedTids, true /*bSaveVertices*/);
-
-	// Update current mesh
-	bool bWasValid = Preview->ProcessCurrentMesh([this](const FDynamicMesh3& Mesh)
+	// See if we're actually making a change to the mesh.
+	if (CurrentExtrudeAmount > 0 || (CurrentExtrudeAmount < 0 && LastOpChangedTids.IsValid() && LastOpChangedTids->Num() > 0))
 	{
-		CurrentMesh->Copy(Mesh);
-	});
-	if (!ensure(bWasValid))
-	{
-		return;
-	}
+		FDynamicMeshChangeTracker ChangeTracker(CurrentMesh.Get());
+		ChangeTracker.BeginChange();
+		ChangeTracker.SaveTriangles(*LastOpChangedTids, true /*bSaveVertices*/);
 
-	if (!bChangesMade)
-	{
-		bChangesMade = true;
-		GetToolManager()->EmitObjectChange(this, MakeUnique<FCubeGridChangesMadeChange>(), TransactionText);
-	}
+		// Update current mesh
+		bool bWasValid = Preview->ProcessCurrentMesh([this](const FDynamicMesh3& Mesh)
+		{
+			CurrentMesh->Copy(Mesh);
+		});
+		if (!ensure(bWasValid))
+		{
+			CurrentExtrudeAmount = 0;
+			return;
+		}
 
-	MeshSpatial->Build();
+		// The transform might have changed if we added to an empty mesh
+		if (!FTransform(CurrentMeshTransform).Equals(Preview->PreviewMesh->GetTransform()))
+		{
+			CurrentMeshTransform = Preview->PreviewMesh->GetTransform();
+			GridGizmoAlignmentMechanic->InitializeDeformedMeshRayCast([this]() { return MeshSpatial.Get(); }, 
+				CurrentMeshTransform, nullptr);
+		}
 
-	CurrentMeshMaterials.Reset();
-	Preview->PreviewMesh->GetMaterials(CurrentMeshMaterials);
+		// Mark our tool as having done something, if we haven't already.
+		if (!bChangesMade)
+		{
+			bChangesMade = true;
+			GetToolManager()->EmitObjectChange(this, MakeUnique<FCubeGridChangesMadeChange>(), TransactionText);
+		}
 
-	UpdateComputeInputs();
+		MeshSpatial->Build();
 
-	bWaitingToApplyPreview = false;
-	bPreviewMayDiffer = false;
-	bBlockUntilPreviewUpdate = false;
+		CurrentMeshMaterials.Reset();
+		Preview->PreviewMesh->GetMaterials(CurrentMeshMaterials);
+
+		UpdateComputeInputs();
+
+		GetToolManager()->EmitObjectChange(this, MakeUnique<FCubeGridToolMeshChange>(ChangeTracker.EndChange()), TransactionText);
+	}
 
-	if (bAdjustSelectionOnPreviewUpdate)
+	if (bAdjustSelectionOnPreviewUpdate && CurrentExtrudeAmount != 0)
 	{
 		// Save UV offset before sliding selection
 		OpMeshHeightUVOffset += GetFrameSpaceExtrudeDist(*CubeGrid, Selection.Box.Min, CurrentExtrudeAmount, Selection.Direction);
@@ -1257,9 +1267,9 @@ void UCubeGridTool::ApplyPreview()
 		SlideSelection(CurrentExtrudeAmount, true);
 	}
 
+	// This actually also happens as a side effect of SlideSelection above, but here for clarity.
 	CurrentExtrudeAmount = 0;
 
-	GetToolManager()->EmitObjectChange(this, MakeUnique<FCubeGridToolMeshChange>(ChangeTracker.EndChange()), TransactionText);
 	GetToolManager()->EndUndoTransaction();
 }
 
