make split tool's "from output" use the current target's output type instead of the first target's type

diff --git a/Engine/Plugins/Experimental/MeshModelingToolsetExp/Source/MeshModelingToolsExp/Private/SplitMeshesTool.cpp b/Engine/Plugins/Experimental/MeshModelingToolsetExp/Source/MeshModelingToolsExp/Private/SplitMeshesTool.cpp
--- a/Engine/Plugins/Experimental/MeshModelingToolsetExp/Source/MeshModelingToolsExp/Private/SplitMeshesTool.cpp
+++ b/Engine/Plugins/Experimental/MeshModelingToolsetExp/Source/MeshModelingToolsExp/Private/SplitMeshesTool.cpp
@@ -112,7 +112,7 @@ void USplitMeshesTool::OnShutdown(EToolShutdownType ShutdownType)
 
 			if (OutputTypeProperties->OutputType == UCreateMeshObjectTypeProperties::AutoIdentifier)
 			{
-				UE::ToolTarget::ConfigureCreateMeshObjectParams(Targets[0], BaseMeshObjectParams);
+				UE::ToolTarget::ConfigureCreateMeshObjectParams(Targets[ti], BaseMeshObjectParams);
 			}
 			else
 			{
