Fix vertex sculpting brush alpha's "RotationAngle" parameter name not fitting in the assigned space, by changing the display name from "Rotation" to "Angle"

diff --git a/Engine/Plugins/Experimental/MeshModelingToolsetExp/Source/MeshModelingToolsExp/Public/MeshVertexSculptTool.h b/Engine/Plugins/Experimental/MeshModelingToolsetExp/Source/MeshModelingToolsExp/Public/MeshVertexSculptTool.h
--- a/Engine/Plugins/Experimental/MeshModelingToolsetExp/Source/MeshModelingToolsExp/Public/MeshVertexSculptTool.h
+++ b/Engine/Plugins/Experimental/MeshModelingToolsetExp/Source/MeshModelingToolsExp/Public/MeshVertexSculptTool.h
@@ -158,7 +158,7 @@ public:
 	TObjectPtr<UTexture2D> Alpha = nullptr;
 
 	/** Alpha is rotated by this angle, inside the brush stamp frame (vertically aligned) */
-	UPROPERTY(EditAnywhere, Category = Alpha, meta = (DisplayName = "Rotation", UIMin = "-180.0", UIMax = "180.0", ClampMin = "-360.0", ClampMax = "360.0"))
+	UPROPERTY(EditAnywhere, Category = Alpha, meta = (DisplayName = "Angle", UIMin = "-180.0", UIMax = "180.0", ClampMin = "-360.0", ClampMax = "360.0"))
 	float RotationAngle = 0.0;
 
 	/** If true, a random angle in +/- RandomRange is added to Rotation angle for each stamp */
