Account for MeshCenter when extracting boundary of mesh for preview in polypath tool

diff --git a/Engine/Plugins/Experimental/MeshModelingToolsetExp/Source/MeshModelingToolsExp/Private/DrawPolyPathTool.cpp b/Engine/Plugins/Experimental/MeshModelingToolsetExp/Source/MeshModelingToolsExp/Private/DrawPolyPathTool.cpp
--- a/Engine/Plugins/Experimental/MeshModelingToolsetExp/Source/MeshModelingToolsExp/Private/DrawPolyPathTool.cpp
+++ b/Engine/Plugins/Experimental/MeshModelingToolsetExp/Source/MeshModelingToolsExp/Private/DrawPolyPathTool.cpp
@@ -851,9 +851,17 @@ FVector3d UDrawPolyPathTool::GeneratePathMesh(FDynamicMesh3& Mesh)
 	if (Loops.Loops.Num() > 0)
 	{
 		Loops.Loops[0].GetVertices<FVector3d>(CurPolyLoop);
+		for (FVector3d& Pt : CurPolyLoop)
+		{
+			Pt += MeshCenter;
+		}
 		if (Loops.Loops.Num() > 1)
 		{
 			Loops.Loops[1].GetVertices<FVector3d>(SecondPolyLoop);
+			for (FVector3d& Pt : SecondPolyLoop)
+			{
+				Pt += MeshCenter;
+			}
 		}
 	}
 
