Fix Merge (self union) tool not working at LWC

diff --git a/Engine/Plugins/Experimental/MeshModelingToolsetExp/Source/MeshModelingToolsExp/Private/SelfUnionMeshesTool.cpp b/Engine/Plugins/Experimental/MeshModelingToolsetExp/Source/MeshModelingToolsExp/Private/SelfUnionMeshesTool.cpp
--- a/Engine/Plugins/Experimental/MeshModelingToolsetExp/Source/MeshModelingToolsExp/Private/SelfUnionMeshesTool.cpp
+++ b/Engine/Plugins/Experimental/MeshModelingToolsetExp/Source/MeshModelingToolsExp/Private/SelfUnionMeshesTool.cpp
@@ -97,6 +97,14 @@ void USelfUnionMeshesTool::ConvertInputsAndSetPreviewMaterials(bool bSetPreviewM
 	CombinedSourceMeshes->Attributes()->EnablePrimaryColors();
 	FDynamicMeshEditor AppendEditor(CombinedSourceMeshes.Get());
 
+	CombinedCenter = FVector3d(0, 0, 0);
+	for (int ComponentIdx = 0; ComponentIdx < Targets.Num(); ComponentIdx++)
+	{
+		CombinedCenter += TransformProxies[ComponentIdx]->GetTransform().GetTranslation();
+	}
+	CombinedCenter /= double(Targets.Num());
+	
+
 	bool bNeedColorAttr = false;
 	for (int ComponentIdx = 0; ComponentIdx < Targets.Num(); ComponentIdx++)
 	{
@@ -112,7 +120,6 @@ void USelfUnionMeshesTool::ConvertInputsAndSetPreviewMaterials(bool bSetPreviewM
 		{
 			MaterialIDs->SetValue(TID, MaterialRemap[ComponentIdx][MaterialIDs->GetValue(TID)]);
 		}
-		// TODO: center the meshes
 		FTransformSRT3d WorldTransform = TransformProxies[ComponentIdx]->GetTransform();
 		if (WorldTransform.GetDeterminant() < 0)
 		{
@@ -120,9 +127,9 @@ void USelfUnionMeshesTool::ConvertInputsAndSetPreviewMaterials(bool bSetPreviewM
 		}
 		FMeshIndexMappings IndexMaps;
 		AppendEditor.AppendMesh(&ComponentMesh, IndexMaps,
-			[WorldTransform](int VID, const FVector3d& Pos)
+			[WorldTransform, this](int VID, const FVector3d& Pos)
 			{
-				return WorldTransform.TransformPosition(Pos);
+				return WorldTransform.TransformPosition(Pos) - CombinedCenter;
 			},
 			[WorldTransform](int VID, const FVector3d& Normal)
 			{
@@ -178,6 +185,7 @@ void USelfUnionMeshesTool::UpdateVisualization()
 	FVector3d A, B;
 
 	DrawnLineSet->Clear();
+	DrawnLineSet->SetWorldTransform(Preview->PreviewMesh->GetTransform());
 	if (Properties->bShowNewBoundaryEdges)
 	{
 		for (int EID : CreatedBoundaryEdges)
@@ -198,7 +206,7 @@ TUniquePtr<FDynamicMeshOperator> USelfUnionMeshesTool::MakeNewOperator()
 	Op->WindingNumberThreshold = Properties->WindingThreshold;
 	Op->bTrimFlaps = Properties->bTrimFlaps;
 
-	Op->SetResultTransform(FTransformSRT3d::Identity()); // TODO Center the combined meshes (when building them) and change this transform accordingly
+	Op->SetResultTransform(FTransformSRT3d(CombinedCenter));
 	Op->CombinedMesh = CombinedSourceMeshes;
 
 	return Op;
diff --git a/Engine/Plugins/Experimental/MeshModelingToolsetExp/Source/MeshModelingToolsExp/Public/SelfUnionMeshesTool.h b/Engine/Plugins/Experimental/MeshModelingToolsetExp/Source/MeshModelingToolsExp/Public/SelfUnionMeshesTool.h
--- a/Engine/Plugins/Experimental/MeshModelingToolsetExp/Source/MeshModelingToolsExp/Public/SelfUnionMeshesTool.h
+++ b/Engine/Plugins/Experimental/MeshModelingToolsetExp/Source/MeshModelingToolsExp/Public/SelfUnionMeshesTool.h
@@ -97,6 +97,8 @@ protected:
 
 	// for visualization of any errors in the currently-previewed merge operation
 	TArray<int> CreatedBoundaryEdges;
+
+	FVector3d CombinedCenter;
 };
 
 
