Modeling tools.  EditPivotTool - when ignoring the snap-drag rotation , don't reset the existing rotation part of the gizmo on as snap dragging it.

diff --git a/Engine/Plugins/Experimental/MeshModelingToolsetExp/Source/MeshModelingToolsExp/Private/EditPivotTool.cpp b/Engine/Plugins/Experimental/MeshModelingToolsetExp/Source/MeshModelingToolsExp/Private/EditPivotTool.cpp
--- a/Engine/Plugins/Experimental/MeshModelingToolsetExp/Source/MeshModelingToolsExp/Private/EditPivotTool.cpp
+++ b/Engine/Plugins/Experimental/MeshModelingToolsetExp/Source/MeshModelingToolsExp/Private/EditPivotTool.cpp
@@ -364,7 +364,7 @@ void UEditPivotTool::OnClickDrag(const FInputDeviceRay& DragPos)
 	FVector TargetNormal = (NormalSign) * Result.Normal;
 
 	FQuaterniond AlignRotation = (bRotate) ?
-		FQuaterniond(FVector3d::UnitZ(), (FVector3d)TargetNormal) : FQuaterniond::Identity();
+		FQuaterniond(FVector3d::UnitZ(), (FVector3d)TargetNormal) : FQuaterniond(StartDragTransform.GetRotation());
 
 	FTransform NewTransform = StartDragTransform;
 	NewTransform.SetRotation((FQuat)AlignRotation);
