Make sculpt tools focus on the brush w/ 'F'

diff --git a/Engine/Plugins/Experimental/MeshModelingToolsetExp/Source/MeshModelingToolsExp/Private/DynamicMeshSculptTool.cpp b/Engine/Plugins/Experimental/MeshModelingToolsetExp/Source/MeshModelingToolsExp/Private/DynamicMeshSculptTool.cpp
--- a/Engine/Plugins/Experimental/MeshModelingToolsetExp/Source/MeshModelingToolsExp/Private/DynamicMeshSculptTool.cpp
+++ b/Engine/Plugins/Experimental/MeshModelingToolsetExp/Source/MeshModelingToolsExp/Private/DynamicMeshSculptTool.cpp
@@ -343,6 +343,19 @@ void UDynamicMeshSculptTool::OnPropertyModified(UObject* PropertySet, FProperty*
 	CalculateBrushRadius();
 }
 
+
+FBox UDynamicMeshSculptTool::GetWorldSpaceFocusBox()
+{
+	if (LastBrushTriangleID == INDEX_NONE)
+	{
+		return Super::GetWorldSpaceFocusBox();
+	}
+	FVector Center = LastBrushPosWorld;
+	double Size = CurrentBrushRadius;
+	return FBox(Center - FVector(Size), Center + FVector(Size));
+}
+
+
 bool UDynamicMeshSculptTool::HitTest(const FRay& Ray, FHitResult& OutHit)
 {
 	FRay3d LocalRay(CurTargetTransform.InverseTransformPosition((FVector3d)Ray.Origin),
diff --git a/Engine/Plugins/Experimental/MeshModelingToolsetExp/Source/MeshModelingToolsExp/Private/Sculpting/MeshSculptToolBase.cpp b/Engine/Plugins/Experimental/MeshModelingToolsetExp/Source/MeshModelingToolsExp/Private/Sculpting/MeshSculptToolBase.cpp
--- a/Engine/Plugins/Experimental/MeshModelingToolsetExp/Source/MeshModelingToolsExp/Private/Sculpting/MeshSculptToolBase.cpp
+++ b/Engine/Plugins/Experimental/MeshModelingToolsetExp/Source/MeshModelingToolsExp/Private/Sculpting/MeshSculptToolBase.cpp
@@ -954,6 +954,18 @@ void UMeshSculptToolBase::DecreaseBrushRadiusSmallStepAction()
 
 
 
+FBox UMeshSculptToolBase::GetWorldSpaceFocusBox()
+{
+	if (LastBrushTriangleID == INDEX_NONE)
+	{
+		return Super::GetWorldSpaceFocusBox();
+	}
+	FVector Center = LastBrushFrameWorld.Origin;
+	double Size = GetCurrentBrushRadius();
+	return FBox(Center - FVector(Size), Center + FVector(Size));
+}
+
+
 
 
 void UMeshSculptToolBase::SetViewPropertiesEnabled(bool bNewValue)
diff --git a/Engine/Plugins/Experimental/MeshModelingToolsetExp/Source/MeshModelingToolsExp/Public/DynamicMeshSculptTool.h b/Engine/Plugins/Experimental/MeshModelingToolsetExp/Source/MeshModelingToolsExp/Public/DynamicMeshSculptTool.h
--- a/Engine/Plugins/Experimental/MeshModelingToolsetExp/Source/MeshModelingToolsExp/Public/DynamicMeshSculptTool.h
+++ b/Engine/Plugins/Experimental/MeshModelingToolsetExp/Source/MeshModelingToolsExp/Public/DynamicMeshSculptTool.h
@@ -280,6 +280,9 @@ public:
 
 	virtual void OnPropertyModified(UObject* PropertySet, FProperty* Property) override;
 
+	// IInteractiveToolCameraFocusAPI override to focus on brush w/ 'F' 
+	virtual FBox GetWorldSpaceFocusBox() override;
+
 public:
 	/** Properties that control brush size/etc*/
 	UPROPERTY()
diff --git a/Engine/Plugins/Experimental/MeshModelingToolsetExp/Source/MeshModelingToolsExp/Public/Sculpting/MeshSculptToolBase.h b/Engine/Plugins/Experimental/MeshModelingToolsetExp/Source/MeshModelingToolsExp/Public/Sculpting/MeshSculptToolBase.h
--- a/Engine/Plugins/Experimental/MeshModelingToolsetExp/Source/MeshModelingToolsExp/Public/Sculpting/MeshSculptToolBase.h
+++ b/Engine/Plugins/Experimental/MeshModelingToolsetExp/Source/MeshModelingToolsExp/Public/Sculpting/MeshSculptToolBase.h
@@ -398,6 +398,11 @@ public:
 	virtual void PreviousBrushModeAction() {}
 
 
+public:
+	// IInteractiveToolCameraFocusAPI override to focus on brush w/ 'F'
+	virtual FBox GetWorldSpaceFocusBox() override;
+
+
 
 	//
 	// Brush/Stroke stuff
