ModelingTools: Fixed BakeTexture/BakeAll not showing error material when output types are restored on tool startup with invalid inputs.

diff --git a/Engine/Plugins/Experimental/MeshModelingToolsetExp/Source/MeshModelingToolsExp/Private/BakeMeshAttributeMapsToolBase.cpp b/Engine/Plugins/Experimental/MeshModelingToolsetExp/Source/MeshModelingToolsExp/Private/BakeMeshAttributeMapsToolBase.cpp
--- a/Engine/Plugins/Experimental/MeshModelingToolsetExp/Source/MeshModelingToolsExp/Private/BakeMeshAttributeMapsToolBase.cpp
+++ b/Engine/Plugins/Experimental/MeshModelingToolsetExp/Source/MeshModelingToolsExp/Private/BakeMeshAttributeMapsToolBase.cpp
@@ -115,6 +115,10 @@ void UBakeMeshAttributeMapsToolBase::OnTick(float DeltaTime)
 			}
 		}
 	}
+	else if (static_cast<bool>(OpState & EBakeOpState::Invalid))
+	{
+		PreviewMesh->SetOverrideRenderMaterial(ErrorPreviewMaterial);
+	} 
 }
 
 
diff --git a/Engine/Plugins/Experimental/MeshModelingToolsetExp/Source/MeshModelingToolsExp/Private/BakeMeshAttributeVertexTool.cpp b/Engine/Plugins/Experimental/MeshModelingToolsetExp/Source/MeshModelingToolsExp/Private/BakeMeshAttributeVertexTool.cpp
--- a/Engine/Plugins/Experimental/MeshModelingToolsetExp/Source/MeshModelingToolsExp/Private/BakeMeshAttributeVertexTool.cpp
+++ b/Engine/Plugins/Experimental/MeshModelingToolsetExp/Source/MeshModelingToolsExp/Private/BakeMeshAttributeVertexTool.cpp
@@ -437,6 +437,10 @@ void UBakeMeshAttributeVertexTool::OnTick(float DeltaTime)
 			}
 		}
 	}
+	else if (static_cast<bool>(OpState & EBakeOpState::Invalid))
+	{
+		PreviewMesh->SetOverrideRenderMaterial(ErrorPreviewMaterial);
+	}
 }
 
 void UBakeMeshAttributeVertexTool::Render(IToolsContextRenderAPI* RenderAPI)
