Fix hole tool and boundary revolve tool seeing polygroup boundary edges instead of mesh boundary edges

Fix boundary revolve tool skipping last vertex on open boundaries

diff --git a/Engine/Plugins/Experimental/MeshModelingToolsetExp/Source/MeshModelingToolsExp/Private/RevolveBoundaryTool.cpp b/Engine/Plugins/Experimental/MeshModelingToolsetExp/Source/MeshModelingToolsExp/Private/RevolveBoundaryTool.cpp
--- a/Engine/Plugins/Experimental/MeshModelingToolsetExp/Source/MeshModelingToolsExp/Private/RevolveBoundaryTool.cpp
+++ b/Engine/Plugins/Experimental/MeshModelingToolsetExp/Source/MeshModelingToolsExp/Private/RevolveBoundaryTool.cpp
@@ -46,16 +46,18 @@ TUniquePtr<FDynamicMeshOperator> URevolveBoundaryOperatorFactory::MakeNewOperato
 			const TArray<int32>& VertexIndices = RevolveBoundaryTool->Topology->GetGroupEdgeVertices(EdgeID);
 			FTransform ToWorld = Cast<IPrimitiveComponentBackedTarget>(RevolveBoundaryTool->Target)->GetWorldTransform();
 
-			// Boundary loop includes the last vertex twice, so stop early.
-			CurveSweepOp->ProfileCurve.Reserve(VertexIndices.Num() - 1);
-			for (int32 i = 0; i < VertexIndices.Num()-1; ++i)
+			// If boundary loop includes the last vertex as first, stop early.
+			// (Note: This is generally true unless the mesh has bowties that confuse the boundary walk.)
+			bool bIsLoop = VertexIndices.Num() && VertexIndices.Last() == VertexIndices[0];
+			CurveSweepOp->ProfileCurve.Reserve(VertexIndices.Num() - (int32)bIsLoop);
+			for (int32 i = 0; i < VertexIndices.Num() - (int32)bIsLoop; ++i)
 			{
 				int32 VertIndex = VertexIndices[i];
 
 				FVector3d NewPos = (FVector3d)ToWorld.TransformPosition((FVector)RevolveBoundaryTool->OriginalMesh->GetVertex(VertIndex));
 				CurveSweepOp->ProfileCurve.Add(NewPos);
 			}
-			CurveSweepOp->bProfileCurveIsClosed = true;
+			CurveSweepOp->bProfileCurveIsClosed = bIsLoop;
 		}
 	}
 
diff --git a/Engine/Plugins/Runtime/GeometryProcessing/Source/DynamicMesh/Private/GroupTopology.cpp b/Engine/Plugins/Runtime/GeometryProcessing/Source/DynamicMesh/Private/GroupTopology.cpp
--- a/Engine/Plugins/Runtime/GeometryProcessing/Source/DynamicMesh/Private/GroupTopology.cpp
+++ b/Engine/Plugins/Runtime/GeometryProcessing/Source/DynamicMesh/Private/GroupTopology.cpp
@@ -171,7 +171,7 @@ bool FGroupTopology::ShouldVertBeCorner(int VertexID, bool* bIsAnExtraCornerOut)
 	{
 		FIndex2i EdgeTris = Mesh->GetEdgeT(Eid);
 		if (EdgeTris.B == IndexConstants::InvalidID
-			|| Mesh->GetTriangleGroup(EdgeTris.A) != Mesh->GetTriangleGroup(EdgeTris.B))
+			|| GetGroupID(EdgeTris.A) != GetGroupID(EdgeTris.B))
 		{
 			// This is a group edge.
 			if (AttachedGroupEdgeEids[0] == IndexConstants::InvalidID)
