Move sculpt plane gizmo to the target object (as much as possible without changing the plane) if it's far from the mesh that is being sculpted

diff --git a/Engine/Plugins/Experimental/MeshModelingToolsetExp/Source/MeshModelingToolsExp/Private/DynamicMeshSculptTool.cpp b/Engine/Plugins/Experimental/MeshModelingToolsetExp/Source/MeshModelingToolsExp/Private/DynamicMeshSculptTool.cpp
--- a/Engine/Plugins/Experimental/MeshModelingToolsetExp/Source/MeshModelingToolsExp/Private/DynamicMeshSculptTool.cpp
+++ b/Engine/Plugins/Experimental/MeshModelingToolsetExp/Source/MeshModelingToolsExp/Private/DynamicMeshSculptTool.cpp
@@ -143,7 +143,8 @@ void UDynamicMeshSculptTool::Setup()
 	PendingTargetUpdate.Wait();
 
 	// initialize brush radius range interval, brush properties
-	double MaxDimension = DynamicMeshComponent->GetMesh()->GetBounds(true).MaxDim();
+	FAxisAlignedBox3d Bounds = DynamicMeshComponent->GetMesh()->GetBounds(true);
+	double MaxDimension = Bounds.MaxDim();
 	BrushRelativeSizeRange = FInterval1d(MaxDimension*0.01, MaxDimension);
 	BrushProperties = NewObject<UDynamicMeshBrushProperties>(this);
 	BrushProperties->BrushSize.InitializeWorldSizeRange(
@@ -194,6 +195,7 @@ void UDynamicMeshSculptTool::Setup()
 	GizmoProperties = NewObject<UFixedPlaneBrushProperties>();
 	GizmoProperties->RestoreProperties(this);
 	AddToolPropertySource(GizmoProperties);
+	GizmoProperties->RecenterGizmoIfFar(CurTargetTransform.TransformPosition(Bounds.Center()), Bounds.MaxDim());
 
 	if (this->bEnableRemeshing)
 	{
diff --git a/Engine/Plugins/Experimental/MeshModelingToolsetExp/Source/MeshModelingToolsExp/Private/MeshVertexSculptTool.cpp b/Engine/Plugins/Experimental/MeshModelingToolsetExp/Source/MeshModelingToolsExp/Private/MeshVertexSculptTool.cpp
--- a/Engine/Plugins/Experimental/MeshModelingToolsetExp/Source/MeshModelingToolsExp/Private/MeshVertexSculptTool.cpp
+++ b/Engine/Plugins/Experimental/MeshModelingToolsetExp/Source/MeshModelingToolsExp/Private/MeshVertexSculptTool.cpp
@@ -287,6 +287,8 @@ void UMeshVertexSculptTool::Setup()
 
 	AddToolPropertySource(UMeshSculptToolBase::GizmoProperties);
 	SetToolPropertySourceEnabled(UMeshSculptToolBase::GizmoProperties, false);
+	// Move the gizmo toward the center of the mesh, without changing the plane it represents
+	UMeshSculptToolBase::GizmoProperties->RecenterGizmoIfFar(GetSculptMeshComponent()->GetComponentTransform().TransformPosition(Bounds.Center()), Bounds.MaxDim());
 
 	AddToolPropertySource(UMeshSculptToolBase::ViewProperties);
 
diff --git a/Engine/Plugins/Experimental/MeshModelingToolsetExp/Source/MeshModelingToolsExp/Public/DynamicMeshSculptTool.h b/Engine/Plugins/Experimental/MeshModelingToolsetExp/Source/MeshModelingToolsExp/Public/DynamicMeshSculptTool.h
--- a/Engine/Plugins/Experimental/MeshModelingToolsetExp/Source/MeshModelingToolsExp/Public/DynamicMeshSculptTool.h
+++ b/Engine/Plugins/Experimental/MeshModelingToolsetExp/Source/MeshModelingToolsExp/Public/DynamicMeshSculptTool.h
@@ -234,6 +234,17 @@ public:
 
 	UPROPERTY(EditAnywhere, Category = TargetPlane, AdvancedDisplay, meta = (HideEditConditionToggle, EditCondition = "bPropertySetEnabled == true"))
 	FQuat Rotation = FQuat::Identity;
+
+	// Recenter the gizmo around the target position (without changing work plane), if it is "too far" (> 10 meters + max bounds dim) from that position currently
+	void RecenterGizmoIfFar(FVector CenterPosition, double BoundsMaxDim, double TooFarDistance = 1000)
+	{
+		double DistanceTolSq = (BoundsMaxDim + TooFarDistance) * (BoundsMaxDim + TooFarDistance);
+		if (FVector::DistSquared(CenterPosition, Position) > DistanceTolSq)
+		{
+			FVector Normal = Rotation.GetAxisZ();
+			Position = CenterPosition - (CenterPosition - Position).ProjectOnToNormal(Normal);
+		}
+	}
 };
 
 
diff --git a/Engine/Plugins/Experimental/MeshModelingToolsetExp/Source/MeshModelingToolsExp/Public/Sculpting/MeshSculptToolBase.h b/Engine/Plugins/Experimental/MeshModelingToolsetExp/Source/MeshModelingToolsExp/Public/Sculpting/MeshSculptToolBase.h
--- a/Engine/Plugins/Experimental/MeshModelingToolsetExp/Source/MeshModelingToolsExp/Public/Sculpting/MeshSculptToolBase.h
+++ b/Engine/Plugins/Experimental/MeshModelingToolsetExp/Source/MeshModelingToolsExp/Public/Sculpting/MeshSculptToolBase.h
@@ -185,7 +185,7 @@ public:
 	UPROPERTY( meta = (TransientToolProperty) )
 	bool bPropertySetEnabled = true;
 
-	/** Toggle whether Work Plane Positioing Gizmo is visible */
+	/** Toggle whether Work Plane Positioning Gizmo is visible */
 	UPROPERTY(EditAnywhere, Category = TargetPlane, meta = (HideEditConditionToggle, EditCondition = "bPropertySetEnabled == true"))
 	bool bShowGizmo = true;
 
@@ -194,6 +194,17 @@ public:
 
 	UPROPERTY(EditAnywhere, Category = TargetPlane, AdvancedDisplay, meta = (HideEditConditionToggle, EditCondition = "bPropertySetEnabled == true"))
 	FQuat Rotation = FQuat::Identity;
+
+	// Recenter the gizmo around the target position (without changing work plane), if it is "too far" (> 10 meters + max bounds dim) from that position currently
+	void RecenterGizmoIfFar(FVector CenterPosition, double BoundsMaxDim, double TooFarDistance = 1000)
+	{
+		double DistanceTolSq = (BoundsMaxDim + TooFarDistance) * (BoundsMaxDim + TooFarDistance);
+		if (FVector::DistSquared(CenterPosition, Position) > DistanceTolSq)
+		{
+			FVector Normal = Rotation.GetAxisZ();
+			Position = CenterPosition - (CenterPosition - Position).ProjectOnToNormal(Normal);
+		}
+	}
 };
 
 
