add comments/tooltips for modeling tools' collision property sets

diff --git a/Engine/Plugins/Experimental/MeshModelingToolsetExp/Source/MeshModelingToolsExp/Public/Physics/CollisionPropertySets.h b/Engine/Plugins/Experimental/MeshModelingToolsetExp/Source/MeshModelingToolsExp/Public/Physics/CollisionPropertySets.h
--- a/Engine/Plugins/Experimental/MeshModelingToolsetExp/Source/MeshModelingToolsExp/Public/Physics/CollisionPropertySets.h
+++ b/Engine/Plugins/Experimental/MeshModelingToolsetExp/Source/MeshModelingToolsExp/Public/Physics/CollisionPropertySets.h
@@ -37,6 +37,7 @@ struct MESHMODELINGTOOLSEXP_API FPhysicsSphereData
 	UPROPERTY(VisibleAnywhere, Category = Sphere)
 	FTransform Transform;
 
+	/** Shape Element storing standard collision and physics properties for a shape */
 	UPROPERTY(VisibleAnywhere, Category = Sphere)
 	FKShapeElem Element;
 };
@@ -52,6 +53,7 @@ struct MESHMODELINGTOOLSEXP_API FPhysicsBoxData
 	UPROPERTY(VisibleAnywhere, Category = Sphere)
 	FTransform Transform;
 
+	/** Shape Element storing standard collision and physics properties for a shape */
 	UPROPERTY(VisibleAnywhere, Category = Sphere)
 	FKShapeElem Element;
 };
@@ -70,6 +72,7 @@ struct MESHMODELINGTOOLSEXP_API FPhysicsCapsuleData
 	UPROPERTY(VisibleAnywhere, Category = Sphere)
 	FTransform Transform;
 
+	/** Shape Element storing standard collision and physics properties for a shape */
 	UPROPERTY(VisibleAnywhere, Category = Sphere)
 	FKShapeElem Element;
 };
@@ -85,6 +88,7 @@ struct MESHMODELINGTOOLSEXP_API FPhysicsConvexData
 	UPROPERTY(VisibleAnywhere, Category = Convex)
 	int32 NumFaces = 0;
 
+	/** Shape Element storing standard collision and physics properties for a shape */
 	UPROPERTY(VisibleAnywhere, Category = Sphere)
 	FKShapeElem Element;
 };
@@ -94,6 +98,7 @@ struct MESHMODELINGTOOLSEXP_API FPhysicsLevelSetData
 {
 	GENERATED_BODY()
 
+	/** Shape Element storing standard collision and physics properties for a shape */
 	UPROPERTY(VisibleAnywhere, Category = LevelSet)
 	FKShapeElem Element;
 };
@@ -103,9 +108,12 @@ class MESHMODELINGTOOLSEXP_API UPhysicsObjectToolPropertySet : public UInteracti
 {
 	GENERATED_BODY()
 public:
+
+	/** Source Object Name */
 	UPROPERTY(VisibleAnywhere, Category = PhysicsData)
 	FString ObjectName;
 
+	/** Collision Flags controlling how simple and complex collision shapes are used */
 	UPROPERTY(VisibleAnywhere, Category = PhysicsData)
 	ECollisionGeometryMode CollisionType;
 
@@ -136,15 +144,19 @@ class MESHMODELINGTOOLSEXP_API UCollisionGeometryVisualizationProperties : publi
 	GENERATED_BODY()
 public:
 
+	/** Thickness of lines used to visualize collision shapes */
 	UPROPERTY(EditAnywhere, Category = "Collision Visualization")
 	float LineThickness = 3.0f;
 
+	/** Show occluded parts of the collision geometry, rendered with dashed lines */
 	UPROPERTY(EditAnywhere, Category = "Collision Visualization")
 	bool bShowHidden = false;
 
+	/** Render each collision geometry with a randomly-chosen color */
 	UPROPERTY(EditAnywhere, Category = "Collision Visualization")
 	bool bRandomColors = true;
 
+	/** The color to use for all collision visualizations, if random colors are not used */
 	UPROPERTY(EditAnywhere, Category = "Collision Visualization", meta = (EditCondition = "!bRandomColors"))
 	FColor Color = FColor::Red;
 };
