ModelingTools: Disable raytracing and shadows on brush indicators

diff --git a/Engine/Plugins/Experimental/MeshModelingToolsetExp/Source/MeshModelingToolsExp/Private/DynamicMeshSculptTool.cpp b/Engine/Plugins/Experimental/MeshModelingToolsetExp/Source/MeshModelingToolsExp/Private/DynamicMeshSculptTool.cpp
--- a/Engine/Plugins/Experimental/MeshModelingToolsetExp/Source/MeshModelingToolsExp/Private/DynamicMeshSculptTool.cpp
+++ b/Engine/Plugins/Experimental/MeshModelingToolsetExp/Source/MeshModelingToolsExp/Private/DynamicMeshSculptTool.cpp
@@ -2210,6 +2210,10 @@ UPreviewMesh* UDynamicMeshSculptTool::MakeDefaultSphereMesh(UObject* Parent, UWo
 		SphereMesh->SetMaterial(BrushIndicatorMaterial);
 	}
 
+	// make sure raytracing is disabled on the brush indicator
+	Cast<UDynamicMeshComponent>(SphereMesh->GetRootComponent())->SetEnableRaytracing(false);
+	SphereMesh->SetShadowsEnabled(false);
+
 	return SphereMesh;
 }
 
diff --git a/Engine/Plugins/Experimental/MeshModelingToolsetExp/Source/MeshModelingToolsExp/Private/MeshVertexSculptTool.cpp b/Engine/Plugins/Experimental/MeshModelingToolsetExp/Source/MeshModelingToolsExp/Private/MeshVertexSculptTool.cpp
--- a/Engine/Plugins/Experimental/MeshModelingToolsetExp/Source/MeshModelingToolsExp/Private/MeshVertexSculptTool.cpp
+++ b/Engine/Plugins/Experimental/MeshModelingToolsetExp/Source/MeshModelingToolsExp/Private/MeshVertexSculptTool.cpp
@@ -375,6 +375,10 @@ UPreviewMesh* UMeshVertexSculptTool::MakeBrushIndicatorMesh(UObject* Parent, UWo
 		PlaneMesh->SetMaterial(BrushIndicatorMaterial);
 	}
 
+	// make sure raytracing is disabled on the brush indicator
+	Cast<UDynamicMeshComponent>(PlaneMesh->GetRootComponent())->SetEnableRaytracing(false);
+	PlaneMesh->SetShadowsEnabled(false);
+
 	return PlaneMesh;
 }
 
diff --git a/Engine/Plugins/Experimental/MeshModelingToolsetExp/Source/MeshModelingToolsExp/Private/Sculpting/MeshSculptToolBase.cpp b/Engine/Plugins/Experimental/MeshModelingToolsetExp/Source/MeshModelingToolsExp/Private/Sculpting/MeshSculptToolBase.cpp
--- a/Engine/Plugins/Experimental/MeshModelingToolsetExp/Source/MeshModelingToolsExp/Private/Sculpting/MeshSculptToolBase.cpp
+++ b/Engine/Plugins/Experimental/MeshModelingToolsetExp/Source/MeshModelingToolsExp/Private/Sculpting/MeshSculptToolBase.cpp
@@ -1208,6 +1208,10 @@ UPreviewMesh* UMeshSculptToolBase::MakeBrushIndicatorMesh(UObject* Parent, UWorl
 		SphereMesh->SetMaterial(BrushIndicatorMaterial);
 	}
 
+	// make sure raytracing is disabled on the brush indicator
+	Cast<UDynamicMeshComponent>(SphereMesh->GetRootComponent())->SetEnableRaytracing(false);
+	SphereMesh->SetShadowsEnabled(false);
+
 	return SphereMesh;
 }
 
