Reflect current 'pivot actor' transform in the tool details panel, so user can see/edit the transform numerically as well

diff --git a/Engine/Plugins/Experimental/MeshModelingToolsetExp/Source/MeshModelingToolsEditorOnlyExp/Private/AddPivotActorTool.cpp b/Engine/Plugins/Experimental/MeshModelingToolsetExp/Source/MeshModelingToolsEditorOnlyExp/Private/AddPivotActorTool.cpp
--- a/Engine/Plugins/Experimental/MeshModelingToolsetExp/Source/MeshModelingToolsEditorOnlyExp/Private/AddPivotActorTool.cpp
+++ b/Engine/Plugins/Experimental/MeshModelingToolsetExp/Source/MeshModelingToolsEditorOnlyExp/Private/AddPivotActorTool.cpp
@@ -140,6 +140,12 @@ void UAddPivotActorTool::Setup()
 		}
 	}
 
+
+	TransformProperties = NewObject<UPivotActorTransformProperties>();
+	AddToolPropertySource(TransformProperties);
+	TransformProperties->Position = StartTransform.GetTranslation();
+	TransformProperties->Rotation = StartTransform.GetRotation();
+
 	// Set up the gizmo.
 	TransformProxy = NewObject<UTransformProxy>(this);
 	TransformProxy->SetTransform(StartTransform);
@@ -147,12 +153,48 @@ void UAddPivotActorTool::Setup()
 		GetToolManager()->GetPairedGizmoManager(),
 		ETransformGizmoSubElements::StandardTranslateRotate, this);
 	TransformGizmo->SetActiveTarget(TransformProxy, GetToolManager());
+	TransformProxy->OnTransformChanged.AddUObject(this, &UAddPivotActorTool::GizmoTransformChanged);
+
+	GizmoPositionWatcher.Initialize(
+		[this]() { return TransformProperties->Position; },
+		[this](FVector NewPosition)
+		{ 
+			UpdateGizmoFromProperties();
+		}, TransformProperties->Position);
+	GizmoRotationWatcher.Initialize(
+		[this]() { return TransformProperties->Rotation; },
+		[this](FQuat NewRotation)
+		{
+			UpdateGizmoFromProperties();
+		}, TransformProperties->Rotation);
 
 	DragAlignmentMechanic = NewObject<UDragAlignmentMechanic>(this);
 	DragAlignmentMechanic->Setup(this);
 	DragAlignmentMechanic->AddToGizmo(TransformGizmo);
 }
 
+void UAddPivotActorTool::OnTick(float DeltaTime)
+{
+	GizmoPositionWatcher.CheckAndUpdate();
+	GizmoRotationWatcher.CheckAndUpdate();
+}
+
+void UAddPivotActorTool::UpdateGizmoFromProperties()
+{
+	if (TransformGizmo)
+	{
+		TransformGizmo->SetNewGizmoTransform(FTransform(TransformProperties->Rotation, TransformProperties->Position));
+	}
+}
+
+void UAddPivotActorTool::GizmoTransformChanged(UTransformProxy* Proxy, FTransform Transform)
+{
+	TransformProperties->Position = Transform.GetTranslation();
+	TransformProperties->Rotation = Transform.GetRotation();
+	GizmoPositionWatcher.SilentUpdate();
+	GizmoRotationWatcher.SilentUpdate();
+}
+
 void UAddPivotActorTool::OnShutdown(EToolShutdownType ShutdownType)
 {
 	using namespace AddPivotActorToolLocals;
diff --git a/Engine/Plugins/Experimental/MeshModelingToolsetExp/Source/MeshModelingToolsEditorOnlyExp/Public/AddPivotActorTool.h b/Engine/Plugins/Experimental/MeshModelingToolsetExp/Source/MeshModelingToolsEditorOnlyExp/Public/AddPivotActorTool.h
--- a/Engine/Plugins/Experimental/MeshModelingToolsetExp/Source/MeshModelingToolsEditorOnlyExp/Public/AddPivotActorTool.h
+++ b/Engine/Plugins/Experimental/MeshModelingToolsetExp/Source/MeshModelingToolsEditorOnlyExp/Public/AddPivotActorTool.h
@@ -7,6 +7,7 @@
 #include "InteractiveToolBuilder.h"
 #include "BaseTools/MultiSelectionMeshEditingTool.h"
 #include "GameFramework/Actor.h"
+#include "Changes/ValueWatcher.h"
 
 #include "AddPivotActorTool.generated.h"
 
@@ -15,6 +16,21 @@ class UCombinedTransformGizmo;
 class UTransformProxy;
 
 
+UCLASS()
+class MESHMODELINGTOOLSEDITORONLYEXP_API UPivotActorTransformProperties : public UInteractiveToolPropertySet
+{
+	GENERATED_BODY()
+
+public:
+
+	UPROPERTY(EditAnywhere, Category = PivotLocation)
+	FVector Position = FVector::ZeroVector;
+
+	UPROPERTY(EditAnywhere, Category = PivotLocation)
+	FQuat Rotation = FQuat::Identity;
+};
+
+
 UCLASS()
 class MESHMODELINGTOOLSEDITORONLYEXP_API UAddPivotActorToolBuilder : public UMultiSelectionMeshEditingToolBuilder
 {
@@ -47,6 +63,9 @@ public:
 
 	virtual bool HasAccept() const override { return true; }
 	virtual bool HasCancel() const override { return true; }
+
+	virtual void OnTick(float DeltaTime) override;
+
 	// Uses the base class CanAccept
 
 protected:
@@ -60,7 +79,16 @@ protected:
 	UPROPERTY()
 	TObjectPtr<UTransformProxy> TransformProxy = nullptr;
 
+	UPROPERTY()
+	TObjectPtr<UPivotActorTransformProperties> TransformProperties;
+
 	TWeakObjectPtr<AActor> ExistingPivotActor = nullptr;
 
+	TValueWatcher<FVector> GizmoPositionWatcher;
+	TValueWatcher<FQuat> GizmoRotationWatcher;
+
+	void UpdateGizmoFromProperties();
+	void GizmoTransformChanged(UTransformProxy* Proxy, FTransform Transform);
+
 	FTransform ExistingPivotOriginalTransform;
 };
