if axis is 100,000 units (1000 meters) away on boundary revolution tool start, snap it to the local mesh pivot, to work around confusion when the tool is started at LWC.

diff --git a/Engine/Plugins/Experimental/MeshModelingToolsetExp/Source/MeshModelingToolsExp/Private/RevolveBoundaryTool.cpp b/Engine/Plugins/Experimental/MeshModelingToolsetExp/Source/MeshModelingToolsExp/Private/RevolveBoundaryTool.cpp
--- a/Engine/Plugins/Experimental/MeshModelingToolsetExp/Source/MeshModelingToolsExp/Private/RevolveBoundaryTool.cpp
+++ b/Engine/Plugins/Experimental/MeshModelingToolsetExp/Source/MeshModelingToolsExp/Private/RevolveBoundaryTool.cpp
@@ -13,6 +13,7 @@
 #include "GroupTopology.h"
 #include "ToolBuilderUtil.h"
 #include "Selection/ToolSelectionUtil.h"
+#include "ModelingToolTargetUtil.h"
 #include "ToolSceneQueriesUtil.h"
 #include "ToolSetupUtil.h"
 
@@ -102,6 +103,16 @@ void URevolveBoundaryTool::Setup()
 	AddToolPropertySource(MaterialProperties);
 	MaterialProperties->RestoreProperties(this);
 
+
+	FTransform LocalToWorld = UE::ToolTarget::GetLocalToWorldTransform(Target);
+	// Assume an axis that is > 1000 meters away is typically not desired, and can be snapped closer
+	// This works around bad behavior if the tool is started at LWC
+	constexpr double AxisVeryFarThreshold = 100 * 1000;
+	if (FVector::DistSquared(Settings->AxisOrigin, LocalToWorld.GetTranslation()) > AxisVeryFarThreshold * AxisVeryFarThreshold)
+	{
+		Settings->AxisOrigin = LocalToWorld.GetTranslation();
+	}
+
 	UpdateRevolutionAxis();
 
 	// The plane mechanic is used for the revolution axis
