MeshModelingTools: Disallow PolyCut from accepting failed cuts by default, but make it user configurable. Note failed cuts in non-boolean mode. Decrease coplanar cases a bit more in boolean mode.

diff --git a/Engine/Plugins/Experimental/MeshModelingToolsetExp/Source/MeshModelingToolsEditorOnlyExp/Private/PolygonOnMeshTool.cpp b/Engine/Plugins/Experimental/MeshModelingToolsetExp/Source/MeshModelingToolsEditorOnlyExp/Private/PolygonOnMeshTool.cpp
--- a/Engine/Plugins/Experimental/MeshModelingToolsetExp/Source/MeshModelingToolsEditorOnlyExp/Private/PolygonOnMeshTool.cpp
+++ b/Engine/Plugins/Experimental/MeshModelingToolsetExp/Source/MeshModelingToolsEditorOnlyExp/Private/PolygonOnMeshTool.cpp
@@ -112,7 +112,16 @@ void UPolygonOnMeshTool::Setup()
 		{
 			GetToolManager()->PostInvalidation();
 			UpdateVisualization();
-			UpdateAcceptWarnings(UpdatedPreview->HaveEmptyResult() ? EAcceptWarning::EmptyForbidden : EAcceptWarning::NoWarning);
+			if (!bOperationSucceeded)
+			{
+				GetToolManager()->DisplayMessage(LOCTEXT("FailNotification", "Unable to complete cut."),
+					EToolMessageLevel::UserWarning);
+			}
+			else
+			{
+				// This clears the warning if needed
+				UpdateAcceptWarnings(UpdatedPreview->HaveEmptyResult() ? EAcceptWarning::EmptyForbidden : EAcceptWarning::NoWarning);
+			}
 		}
 	);
 
@@ -170,6 +179,13 @@ void UPolygonOnMeshTool::UpdateVisualization()
 			TargetMesh->GetEdgeV(EID, A, B);
 			DrawnLineSet->AddLine((FVector)A, (FVector)B, PartialPathEdgeColor, PartialPathEdgeThickness, PartialPathEdgeDepthBias);
 		}
+		// In the case where we don't allow failed results, it can be disorienting to show the broken mesh (especially since sometimes
+		// most of it may be cut away). But user can change this behavior.
+		if (!BasicProperties->bCanAcceptFailedResult && !BasicProperties->bShowIntermediateResultOnFailure)
+		{
+			// Reset the preview.
+			Preview->PreviewMesh->UpdatePreview(OriginalDynamicMesh.Get());
+		}
 	}
 }
 
@@ -405,7 +421,8 @@ void UPolygonOnMeshTool::CompleteDrawPolygon()
 
 bool UPolygonOnMeshTool::CanAccept() const
 {
-	return Super::CanAccept() && Preview != nullptr && Preview->HaveValidNonEmptyResult();
+	return Super::CanAccept() && Preview != nullptr && Preview->HaveValidNonEmptyResult() 
+		&& (bOperationSucceeded || BasicProperties->bCanAcceptFailedResult);
 }
 
 
diff --git a/Engine/Plugins/Experimental/MeshModelingToolsetExp/Source/MeshModelingToolsEditorOnlyExp/Public/PolygonOnMeshTool.h b/Engine/Plugins/Experimental/MeshModelingToolsetExp/Source/MeshModelingToolsEditorOnlyExp/Public/PolygonOnMeshTool.h
--- a/Engine/Plugins/Experimental/MeshModelingToolsetExp/Source/MeshModelingToolsEditorOnlyExp/Public/PolygonOnMeshTool.h
+++ b/Engine/Plugins/Experimental/MeshModelingToolsetExp/Source/MeshModelingToolsEditorOnlyExp/Public/PolygonOnMeshTool.h
@@ -103,6 +103,20 @@ public:
 	/** Number of sides in Circle or RoundRect Corner */
 	UPROPERTY(EditAnywhere, Category = Shape, meta = (UIMin = "3", UIMax = "20", ClampMin = "3", ClampMax = "10000", EditCondition = "Shape == EPolygonType::Circle || Shape == EPolygonType::RoundRect"))
 	int32 Subdivisions = 12;
+
+	/**
+	 * Whether the tool will allow accepting a result if the operation fails, for instance due to inability to insert the
+	 * polygon when not cutting with boolean, or due to unrepaired cracks in the result.
+	 */
+	UPROPERTY(EditAnywhere, Category = Operation, AdvancedDisplay)
+	bool bCanAcceptFailedResult = false;
+
+	/** 
+	 * If an operation fails and we do not allow accepting the result, whether to show the intermediate failed result, or to
+	 * show the original mesh.
+	 */
+	UPROPERTY(EditAnywhere, Category = Operation, AdvancedDisplay, meta = (EditCondition = "!bCanAcceptFailedResult"))
+	bool bShowIntermediateResultOnFailure = false;
 };
 
 
diff --git a/Engine/Plugins/Runtime/MeshModelingToolset/Source/ModelingOperatorsEditorOnly/Private/CuttingOps/EmbedPolygonsOp.cpp b/Engine/Plugins/Runtime/MeshModelingToolset/Source/ModelingOperatorsEditorOnly/Private/CuttingOps/EmbedPolygonsOp.cpp
--- a/Engine/Plugins/Runtime/MeshModelingToolset/Source/ModelingOperatorsEditorOnly/Private/CuttingOps/EmbedPolygonsOp.cpp
+++ b/Engine/Plugins/Runtime/MeshModelingToolset/Source/ModelingOperatorsEditorOnly/Private/CuttingOps/EmbedPolygonsOp.cpp
@@ -120,7 +120,6 @@ void FEmbedPolygonsOp::BooleanPath(FProgressCancel* Progress)
 	}
 
 	FAxisAlignedBox3d Bounds = OriginalMesh->GetBounds();
-	Bounds.Expand(.01); // expand a little beyond bounds to avoid creating coplanar cases
 	double MeshDiameter = Bounds.MaxDim();
 
 	FFrame3d Frame = PolygonFrame;
@@ -132,7 +131,8 @@ void FEmbedPolygonsOp::BooleanPath(FProgressCancel* Progress)
 		FVector3d Corner = Bounds.GetCorner(CornerIdx);
 		Range.Contain(Dir.Dot(Corner - Frame.Origin));
 	}
-
+	// expand a little beyond bounds to avoid creating coplanar cases
+	Range.Expand(1.0);
 
 	Frame.Origin = Frame.Origin + Range.Min * Dir;
 	FGeneralizedCylinderGenerator ExtrudePolyGen;
@@ -465,7 +465,12 @@ void FEmbedPolygonsOp::CalculateResult(FProgressCancel* Progress)
 		FDynamicMeshEditor MeshEditor(ResultMesh.Get());
 		FDynamicMeshEditResult ResultOut;
 		bool bStitched = MeshEditor.StitchSparselyCorrespondedVertexLoops(AllPathVertIDs[0], AllPathVertCorrespond[0], AllPathVertIDs[1], AllPathVertCorrespond[1], ResultOut);
-		if (bStitched && ResultMesh->HasAttributes())
+		if (!bStitched)
+		{
+			// Don't set bOperationSucceeded to true
+			return;
+		}
+		if (ResultMesh->HasAttributes())
 		{
 			MeshEditor.SetTubeNormals(ResultOut.NewTriangles, AllPathVertIDs[0], AllPathVertCorrespond[0], AllPathVertIDs[1], AllPathVertCorrespond[1]);
 			TArray<float> UValues; UValues.SetNumUninitialized(AllPathVertCorrespond[1].Num() + 1);
