PatternTool: First UI improvement pass.

Interpolate check boxes are now inline with their respective properties and don't occupy their own rows.

Uniform is still on its own row but is quite different in its behavior: instead of having 2 sets of properties for uniform vs. non-uniform scale where bUniform dictated which set was used, bUniform now toggles proportional scaling of StartScale and EndScale.

diff --git a/Engine/Plugins/Experimental/MeshModelingToolsetExp/Source/MeshModelingToolsExp/Private/PatternTool.cpp b/Engine/Plugins/Experimental/MeshModelingToolsetExp/Source/MeshModelingToolsExp/Private/PatternTool.cpp
--- a/Engine/Plugins/Experimental/MeshModelingToolsetExp/Source/MeshModelingToolsExp/Private/PatternTool.cpp
+++ b/Engine/Plugins/Experimental/MeshModelingToolsetExp/Source/MeshModelingToolsExp/Private/PatternTool.cpp
@@ -536,6 +536,40 @@ void UPatternTool::Setup()
 	AddToolPropertySource(ScaleSettings);
 	ScaleSettings->RestoreProperties(this);
 
+	auto OnUniformChanged = [this](bool bNewValue)
+	{
+		if (bNewValue)
+		{
+			constexpr float Tolerance = 1E-08;
+			const FVector DefaultDirection = FVector::OneVector.GetUnsafeNormal();
+			
+			StartScaleDirection = ScaleSettings->StartScale.GetSafeNormal(Tolerance, DefaultDirection);
+			EndScaleDirection = ScaleSettings->EndScale.GetSafeNormal(Tolerance, DefaultDirection);
+		}
+	};
+	ScaleSettings->WatchProperty(ScaleSettings->bUniform, OnUniformChanged);
+	
+	StartScaleWatcherIdx = ScaleSettings->WatchProperty(ScaleSettings->StartScale, [this](const FVector& NewStartScale)
+	{
+		if (ScaleSettings->bUniform)
+		{
+			ScaleSettings->StartScale = StartScaleDirection * NewStartScale.Size();
+			ScaleSettings->SilentUpdateWatcherAtIndex(StartScaleWatcherIdx);
+		}
+	});
+	
+	EndScaleWatcherIdx = ScaleSettings->WatchProperty(ScaleSettings->EndScale, [this](const FVector& NewEndScale)
+	{
+		if (ScaleSettings->bUniform)
+		{
+			ScaleSettings->EndScale = EndScaleDirection * NewEndScale.Size();
+			ScaleSettings->SilentUpdateWatcherAtIndex(EndScaleWatcherIdx);
+		}
+	});
+
+	// Initialize StartScaleDirection and EndScaleDirection
+	OnUniformChanged(true);
+	
 	OutputSettings = NewObject<UPatternTool_OutputSettings>();
 	AddToolPropertySource(OutputSettings);
 	OutputSettings->RestoreProperties(this);
@@ -793,11 +827,11 @@ static void InitializeGenerator(FPatternGenerator& Generator, UPatternTool* Tool
 	if (Generator.bInterpolateRotation)
 	{
 		Generator.StartRotation = FQuaterniond(Tool->RotationSettings->StartRotation);
-		Generator.EndRotation = FQuaterniond(Tool->RotationSettings->Rotation);
+		Generator.EndRotation = FQuaterniond(Tool->RotationSettings->EndRotation);
 	}
 	else
 	{
-		Generator.StartRotation = Generator.EndRotation = FQuaterniond(Tool->RotationSettings->Rotation);
+		Generator.StartRotation = Generator.EndRotation = FQuaterniond(Tool->RotationSettings->StartRotation);
 	}
 
 	Generator.bInterpolateTranslation = Tool->TranslationSettings->bInterpolate;
@@ -806,11 +840,11 @@ static void InitializeGenerator(FPatternGenerator& Generator, UPatternTool* Tool
 	if (Generator.bInterpolateTranslation)
 	{
 		Generator.StartTranslation = Tool->TranslationSettings->StartTranslation;
-		Generator.EndTranslation = Tool->TranslationSettings->Translation;
+		Generator.EndTranslation = Tool->TranslationSettings->EndTranslation;
 	}
 	else
 	{
-		Generator.StartTranslation = Generator.EndTranslation = Tool->TranslationSettings->Translation;
+		Generator.StartTranslation = Generator.EndTranslation = Tool->TranslationSettings->StartTranslation;
 	}
 
 	Generator.bInterpolateScale = Tool->ScaleSettings->bInterpolate;
@@ -818,12 +852,12 @@ static void InitializeGenerator(FPatternGenerator& Generator, UPatternTool* Tool
 	Generator.ScaleJitterRange = (Tool->ScaleSettings->bUniform) ? (FVector3d::One()*Tool->ScaleSettings->ScaleJitterRange) : Tool->ScaleSettings->ScaleJitterRangeNonUniform;
 	if (Generator.bInterpolateScale)
 	{
-		Generator.StartScale = (Tool->ScaleSettings->bUniform) ? (FVector3d::One()*Tool->ScaleSettings->StartScale) : Tool->ScaleSettings->StartScaleNonUniform;
-		Generator.EndScale = (Tool->ScaleSettings->bUniform) ? (FVector3d::One()*Tool->ScaleSettings->Scale) : Tool->ScaleSettings->ScaleNonUniform;
+		Generator.StartScale = Tool->ScaleSettings->StartScale;
+		Generator.EndScale = Tool->ScaleSettings->EndScale;
 	}
 	else
 	{
-		Generator.StartScale = Generator.EndScale = (Tool->ScaleSettings->bUniform) ? (FVector3d::One()*Tool->ScaleSettings->Scale) : Tool->ScaleSettings->ScaleNonUniform;
+		Generator.StartScale = Generator.EndScale = Tool->ScaleSettings->StartScale;
 	}
 }
 
diff --git a/Engine/Plugins/Experimental/MeshModelingToolsetExp/Source/MeshModelingToolsExp/Public/PatternTool.h b/Engine/Plugins/Experimental/MeshModelingToolsetExp/Source/MeshModelingToolsExp/Public/PatternTool.h
--- a/Engine/Plugins/Experimental/MeshModelingToolsetExp/Source/MeshModelingToolsExp/Public/PatternTool.h
+++ b/Engine/Plugins/Experimental/MeshModelingToolsetExp/Source/MeshModelingToolsExp/Public/PatternTool.h
@@ -267,7 +267,7 @@ class MESHMODELINGTOOLSEXP_API UPatternTool_RotationSettings : public UInteracti
 	GENERATED_BODY()
 public:
 	/** If true, Rotation is linearly interpolated between StartRotation and Rotation values */
-	UPROPERTY(EditAnywhere, Category = Rotation)
+	UPROPERTY(EditAnywhere, Category = Rotation, meta = (InlineEditConditionToggle))
 	bool bInterpolate = false;
 
 	/** If true, Rotation at each Pattern Element is offset by a uniformly chosen random value in the range of [-RotationJitterRange, RotationJitterRange] */
@@ -275,12 +275,12 @@ public:
 	bool bJitter = false;
 
 	/** Rotation at first Pattern Element */
-	UPROPERTY(EditAnywhere, Category = Rotation, meta = (EditCondition = "bInterpolate", EditConditionHides, HideEditConditionToggle))
+	UPROPERTY(EditAnywhere, Category = Rotation)
 	FRotator StartRotation = FRotator::ZeroRotator;
 
 	/** Rotation applied to all Pattern Elements, or at Last Pattern Element for Interpolated rotations */
-	UPROPERTY(EditAnywhere, Category = Rotation)
-	FRotator Rotation = FRotator::ZeroRotator;
+	UPROPERTY(EditAnywhere, Category = Rotation, meta = (EditCondition = "bInterpolate"))
+	FRotator EndRotation = FRotator::ZeroRotator;
 
 	/** Upper bound of the range which is sampled to randomly rotate each Pattern Element if Jitter is true */
 	UPROPERTY(EditAnywhere, Category = Rotation, meta = (ClampMin = 0, EditCondition = "bJitter", EditConditionHides, HideEditConditionToggle))
@@ -297,7 +297,7 @@ class MESHMODELINGTOOLSEXP_API UPatternTool_TranslationSettings : public UIntera
 	GENERATED_BODY()
 public:
 	/** If true, Translation is linearly interpolated between StartTranslation and Translation values */
-	UPROPERTY(EditAnywhere, Category = Translation)
+	UPROPERTY(EditAnywhere, Category = Translation, meta = (InlineEditConditionToggle))
 	bool bInterpolate = false;
 
 	/** If true, Translation at each Pattern Element is offset by a uniformly chosen random value in the range of [-TranslationJitterRange, TranslationJitterRange] */
@@ -305,12 +305,12 @@ public:
 	bool bJitter = false;
 
 	/** Translation at first Pattern Element */
-	UPROPERTY(EditAnywhere, Category = Translation, meta = (EditCondition = "bInterpolate", EditConditionHides, HideEditConditionToggle))
+	UPROPERTY(EditAnywhere, Category = Translation)
 	FVector StartTranslation = FVector::ZeroVector;
 
 	/** Translation applied to all Pattern Elements, or at Last Pattern Element for Interpolated translations */
-	UPROPERTY(EditAnywhere, Category = Translation)
-	FVector Translation = FVector::ZeroVector;
+	UPROPERTY(EditAnywhere, Category = Translation, meta = (EditCondition = "bInterpolate"))
+	FVector EndTranslation = FVector::ZeroVector;
 
 	/** Upper bound of the range which is sampled to randomly translate each Pattern Element if Jitter is true */
 	UPROPERTY(EditAnywhere, Category = Translation, meta = (ClampMin = 0, EditCondition = "bJitter", EditConditionHides, HideEditConditionToggle))
@@ -329,9 +329,9 @@ public:
 	/** If true, Scaling is limited to Uniform Scaling */
 	UPROPERTY(EditAnywhere, Category = Scale)
 	bool bUniform = true;
-
+	
 	/** If true, Scale is linearly interpolated between StartScale and Scale values */
-	UPROPERTY(EditAnywhere, Category = Scale)
+	UPROPERTY(EditAnywhere, Category = Scale, meta = (InlineEditConditionToggle))
 	bool bInterpolate = false;
 
 	/** If true, Scale at each Pattern Element is offset by a uniformly chosen random value in the range of [-ScaleJitterRange, ScaleJitterRange] */
@@ -339,25 +339,17 @@ public:
 	bool bJitter = false;
 	
 	/** Uniform Scale at first Pattern Element */
-	UPROPERTY(EditAnywhere, Category = Scale, meta = (EditCondition = "bInterpolate && bUniform", EditConditionHides, HideEditConditionToggle))
-	float StartScale = 1;
-
+	UPROPERTY(EditAnywhere, Category = Scale)
+	FVector StartScale = FVector::OneVector;
+	
 	/** Uniform Scale applied to all Pattern Elements, or at Last Pattern Element for Interpolated scales */
-	UPROPERTY(EditAnywhere, Category = Scale, meta = (EditCondition = "bUniform", EditConditionHides, HideEditConditionToggle))
-	float Scale = 1;
-
+	UPROPERTY(EditAnywhere, Category = Scale, meta = (EditCondition = "bInterpolate"))
+	FVector EndScale = FVector::OneVector;
+	
 	/** Upper bound of the range which is sampled to randomly scale each Pattern Element if Jitter is true */
 	UPROPERTY(EditAnywhere, Category = Scale, meta = (ClampMin = 0, EditCondition = "bJitter && bUniform", EditConditionHides, HideEditConditionToggle))
 	float ScaleJitterRange = 0.0f;
 
-	/** Non-Uniform Scale at first Pattern Element */
-	UPROPERTY(EditAnywhere, Category = Scale, meta = (DisplayName = "Start Scale", EditCondition = "bInterpolate && bUniform == false", EditConditionHides, HideEditConditionToggle))
-	FVector StartScaleNonUniform = FVector(1,1,1);
-
-	/** Non-Uniform Scale applied to all Pattern Elements, or at Last Pattern Element for Interpolated scales */
-	UPROPERTY(EditAnywhere, Category = Scale, meta = (DisplayName = "Scale", EditCondition = "bUniform == false", EditConditionHides, HideEditConditionToggle))
-	FVector ScaleNonUniform = FVector(1,1,1);
-
 	/** Upper bound of the range which is sampled to randomly scale each Pattern Element if Jitter is true (Non-Uniform) */
 	UPROPERTY(EditAnywhere, Category = Scale, meta = (ClampMin = 0, DisplayName = "Scale Jitter Range", EditCondition = "bJitter && bUniform == false", EditConditionHides, HideEditConditionToggle))
 	FVector ScaleJitterRangeNonUniform = FVector(0, 0, 0);
@@ -439,8 +431,12 @@ public:
 
 	UPROPERTY()
 	TObjectPtr<UPatternTool_ScaleSettings> ScaleSettings;
-
 	
+	FVector StartScaleDirection;
+	FVector EndScaleDirection;
+	
+	int32 StartScaleWatcherIdx;
+	int32 EndScaleWatcherIdx;
 
 	UPROPERTY()
 	TObjectPtr<UPatternTool_OutputSettings> OutputSettings;
