PatternTool: Prevent tool from being activated with volume or skeletal meshes.

Overload CanBuildTool() to check that there exists some targetable components (at least 1) which are castable to either a static or a dynamic mesh component.

Additionally, overload InitializeNewTool() so that only targetable components which can be cast to static or dynamic mesh components are built into targets.

diff --git a/Engine/Plugins/Experimental/MeshModelingToolsetExp/Source/MeshModelingToolsExp/Private/PatternTool.cpp b/Engine/Plugins/Experimental/MeshModelingToolsetExp/Source/MeshModelingToolsExp/Private/PatternTool.cpp
--- a/Engine/Plugins/Experimental/MeshModelingToolsetExp/Source/MeshModelingToolsExp/Private/PatternTool.cpp
+++ b/Engine/Plugins/Experimental/MeshModelingToolsetExp/Source/MeshModelingToolsExp/Private/PatternTool.cpp
@@ -44,12 +44,22 @@ using namespace UE::Geometry;
  * ToolBuilder
  */
 
-const FToolTargetTypeRequirements& UPatternToolBuilder::GetTargetRequirements() const
+bool UPatternToolBuilder::CanBuildTool(const FToolBuilderState& SceneState) const
 {
-	static FToolTargetTypeRequirements TypeRequirements(
-		UPrimitiveComponentBackedTarget::StaticClass()
-		);
-	return TypeRequirements;
+	int32 ValidTargets = 0;
+
+	// This tool currently only has 5 output modes, each of them either supporting either static or dynamic meshes.
+	// Other targets which are UPrimitiveComponent backed are not currently supported.
+	// todo: Add support for skeletal meshes and volumes
+	SceneState.TargetManager->EnumerateSelectedAndTargetableComponents(SceneState, GetTargetRequirements(), [&ValidTargets](UActorComponent* Component)
+	{
+		if (Cast<UStaticMeshComponent>(Component) || Cast<UDynamicMeshComponent>(Component))
+		{
+			ValidTargets++;
+		}
+	});
+	
+	return ValidTargets > 0;
 }
 
 UMultiSelectionMeshEditingTool* UPatternToolBuilder::CreateNewTool(const FToolBuilderState& SceneState) const
@@ -57,6 +67,28 @@ UMultiSelectionMeshEditingTool* UPatternToolBuilder::CreateNewTool(const FToolBu
 	return NewObject<UPatternTool>(SceneState.ToolManager);
 }
 
+void UPatternToolBuilder::InitializeNewTool(UMultiSelectionMeshEditingTool* NewTool, const FToolBuilderState& SceneState) const
+{
+	TArray<TObjectPtr<UToolTarget>> Targets;
+	SceneState.TargetManager->EnumerateSelectedAndTargetableComponents(SceneState, GetTargetRequirements(), [this, &Targets, &SceneState](UActorComponent* Component)
+	{
+		if (Cast<UStaticMeshComponent>(Component) || Cast<UDynamicMeshComponent>(Component))
+		{
+			Targets.Add(SceneState.TargetManager->BuildTarget(Component, GetTargetRequirements()));
+		}
+	});
+
+	NewTool->SetTargets(Targets);
+	NewTool->SetWorld(SceneState.World);
+}
+
+const FToolTargetTypeRequirements& UPatternToolBuilder::GetTargetRequirements() const
+{
+	static FToolTargetTypeRequirements TypeRequirements(
+		UPrimitiveComponentBackedTarget::StaticClass()
+		);
+	return TypeRequirements;
+}
 
 // these pattern generators should be promoted to GeometryProcessing, however they need some
 // work to clean up the API and make them more correct (ie handling of step size seems a bit
diff --git a/Engine/Plugins/Experimental/MeshModelingToolsetExp/Source/MeshModelingToolsExp/Public/PatternTool.h b/Engine/Plugins/Experimental/MeshModelingToolsetExp/Source/MeshModelingToolsExp/Public/PatternTool.h
--- a/Engine/Plugins/Experimental/MeshModelingToolsetExp/Source/MeshModelingToolsExp/Public/PatternTool.h
+++ b/Engine/Plugins/Experimental/MeshModelingToolsetExp/Source/MeshModelingToolsExp/Public/PatternTool.h
@@ -35,8 +35,12 @@ class MESHMODELINGTOOLSEXP_API UPatternToolBuilder : public UMultiSelectionMeshE
 	GENERATED_BODY()
 
 public:
+	virtual bool CanBuildTool(const FToolBuilderState& SceneState) const override;
+	
 	virtual UMultiSelectionMeshEditingTool* CreateNewTool(const FToolBuilderState& SceneState) const override;
 
+	virtual void InitializeNewTool(UMultiSelectionMeshEditingTool* NewTool, const FToolBuilderState& SceneState) const override;
+	
 protected:
 	virtual const FToolTargetTypeRequirements& GetTargetRequirements() const override;
 };
@@ -478,6 +482,8 @@ protected:
 	struct FPatternElement
 	{
 		int32 TargetIndex = 0;
+
+		// todo: This is no longer necessary now that InitializeNewTool filters out invalid targets.
 		bool bValid = true;
 
 		UPrimitiveComponent* SourceComponent = nullptr;
