work around crash in physics DDC code by splitting BakeRS and Pivot tool transactions if the object has simple collision w/ convex hulls: First make a transaction with the simple collision update, then make a transaction with the static mesh change.

Only do this for the convex hull case, because the other cases do not seem to crash and splitting the transaction is not ideal.

diff --git a/Engine/Plugins/Experimental/MeshModelingToolsetExp/Source/MeshModelingToolsExp/Private/BakeTransformTool.cpp b/Engine/Plugins/Experimental/MeshModelingToolsetExp/Source/MeshModelingToolsExp/Private/BakeTransformTool.cpp
--- a/Engine/Plugins/Experimental/MeshModelingToolsetExp/Source/MeshModelingToolsExp/Private/BakeTransformTool.cpp
+++ b/Engine/Plugins/Experimental/MeshModelingToolsetExp/Source/MeshModelingToolsExp/Private/BakeTransformTool.cpp
@@ -17,6 +17,7 @@
 #include "DynamicMeshToMeshDescription.h"
 
 #include "Physics/ComponentCollisionUtil.h"
+#include "PhysicsEngine/BodySetup.h"
 
 #include "TargetInterfaces/MeshDescriptionCommitter.h"
 #include "TargetInterfaces/MeshDescriptionProvider.h"
@@ -117,20 +118,48 @@ void UBakeTransformTool::UpdateAssets()
 	// This is to avoid potential stability issues related to creation/load of
 	// mesh descriptions inside a transaction.
 	// TODO: this may not be necessary anymore. Also may not be the most efficient
-	TArray<FMeshDescription> SourceMeshes;
+	// Note: for the crash workaround below, this also now pre-computes the source mesh bounds
+	TArray<FBox> SourceMeshBounds;
 	for (int32 ComponentIdx = 0; ComponentIdx < Targets.Num(); ComponentIdx++)
 	{
-		UE::ToolTarget::GetMeshDescription(Targets[ComponentIdx]);
+		const FMeshDescription* MeshDescription = UE::ToolTarget::GetMeshDescription(Targets[ComponentIdx]);
+		if (MapToFirstOccurrences[ComponentIdx] < ComponentIdx)
+		{
+			SourceMeshBounds.Add(SourceMeshBounds[MapToFirstOccurrences[ComponentIdx]]);
+		}
+		else
+		{
+			SourceMeshBounds.Add(MeshDescription->ComputeBoundingBox());
+		}
 	}
 
-	GetToolManager()->BeginUndoTransaction(LOCTEXT("BakeTransformToolTransactionName", "Bake Transforms"));
+	constexpr bool bWorkaroundForCrashIfConvexAndMeshModifiedInSameTransaction = true;
+	bool bNeedSeparateTransactionForSimpleCollision = false;
+	if constexpr (bWorkaroundForCrashIfConvexAndMeshModifiedInSameTransaction)
+	{
+		for (int32 ComponentIdx = 0; ComponentIdx < Targets.Num(); ComponentIdx++)
+		{
+			UToolTarget* Target = Targets[ComponentIdx];
+			UPrimitiveComponent* Component = UE::ToolTarget::GetTargetComponent(Target);
+			if (UE::Geometry::ComponentTypeSupportsCollision(Component))
+			{
+				if (UBodySetup* BodySetup = UE::Geometry::GetBodySetup(Component))
+				{
+					if (BodySetup->AggGeom.ConvexElems.Num() > 0)
+					{
+						bNeedSeparateTransactionForSimpleCollision = true;
+						break;
+					}
+				}
+			}
+		}
+	}
 
-	TArray<FTransformSRT3d> BakedTransforms;
+	// Compute all the transforms that we should bake to the assets or apply to the components
+	TArray<FTransformSRT3d> BakedTransforms, ComponentTransforms;
 	for (int32 ComponentIdx = 0; ComponentIdx < Targets.Num(); ComponentIdx++)
 	{
 		UToolTarget* Target = Targets[ComponentIdx];
-		UPrimitiveComponent* Component = UE::ToolTarget::GetTargetComponent(Target);
-		Component->Modify();
 
 		FTransformSRT3d ComponentToWorld = UE::ToolTarget::GetLocalToWorldTransform(Target);
 		FTransformSRT3d ToBakePart = FTransformSRT3d::Identity();
@@ -199,20 +228,63 @@ void UBakeTransformTool::UpdateAssets()
 				check(false); // must explicitly handle all cases
 			}
 
-			// apply edit
-			FMeshDescription SourceMesh(UE::ToolTarget::GetMeshDescriptionCopy(Targets[ComponentIdx]));
-			FMeshDescriptionEditableTriangleMeshAdapter EditableMeshDescAdapter(&SourceMesh);
-
 			// do this part within the commit because we have the MeshDescription already computed
 			if (BasicProperties->bRecenterPivot)
 			{
-				FBox BBox = SourceMesh.ComputeBoundingBox();
+				FBox BBox = SourceMeshBounds[ComponentIdx];
 				FVector3d Center(BBox.GetCenter());
 				FFrame3d LocalFrame(Center);
 				ToBakePart.SetTranslation(ToBakePart.GetTranslation() - Center);
 				NewWorldPart.SetTranslation(NewWorldPart.GetTranslation() + NewWorldPart.TransformVector(Center));
 			}
 
+			BakedTransforms.Add(ToBakePart);
+		}
+
+		ComponentTransforms.Add(NewWorldPart);
+	}
+
+	auto UpdateSimpleCollision = [](UPrimitiveComponent* Component, const FTransformSRT3d& ToBakePart)
+	{
+		// try to transform simple collision
+		if (UE::Geometry::ComponentTypeSupportsCollision(Component))
+		{
+			UE::Geometry::TransformSimpleCollision(Component, ToBakePart);
+		}
+	};
+
+	// If necessary, make a first transaction to bake simple collision updates separately (works around crash in undo/redo of transactions that update both hulls and static meshes together)
+	if (bNeedSeparateTransactionForSimpleCollision)
+	{
+		GetToolManager()->BeginUndoTransaction(LOCTEXT("BakeTransformToolSimpleCollisionTransactionName", "Bake Transforms Part 1 (Simple Collision)"));
+		for (int32 ComponentIdx = 0; ComponentIdx < Targets.Num(); ComponentIdx++)
+		{
+			UToolTarget* Target = Targets[ComponentIdx];
+			UPrimitiveComponent* Component = UE::ToolTarget::GetTargetComponent(Target);
+			Component->Modify();
+			UpdateSimpleCollision(Component, BakedTransforms[ComponentIdx]);
+		}
+		GetToolManager()->EndUndoTransaction();
+	}
+
+	GetToolManager()->BeginUndoTransaction(bNeedSeparateTransactionForSimpleCollision ?
+		LOCTEXT("BakeTransformToolGeometryTransactionName", "Bake Transforms Part 2 (Visible Geometry)") :
+		LOCTEXT("BakeTransformToolTransactionName", "Bake Transforms"));
+
+	for (int32 ComponentIdx = 0; ComponentIdx < Targets.Num(); ComponentIdx++)
+	{
+		UToolTarget* Target = Targets[ComponentIdx];
+		UPrimitiveComponent* Component = UE::ToolTarget::GetTargetComponent(Target);
+		Component->Modify();
+
+		FTransformSRT3d ToBakePart = BakedTransforms[ComponentIdx];
+
+		if (MapToFirstOccurrences[ComponentIdx] == ComponentIdx)
+		{
+			// apply edit
+			FMeshDescription SourceMesh(UE::ToolTarget::GetMeshDescriptionCopy(Targets[ComponentIdx]));
+			FMeshDescriptionEditableTriangleMeshAdapter EditableMeshDescAdapter(&SourceMesh);
+
 			MeshAdapterTransforms::ApplyTransform(EditableMeshDescAdapter, ToBakePart);
 
 			FVector3d BakeScaleVec = ToBakePart.GetScale();
@@ -224,16 +296,15 @@ void UBakeTransformTool::UpdateAssets()
 			// todo: support vertex-only update
 			UE::ToolTarget::CommitMeshDescriptionUpdate(Target, &SourceMesh);
 
-			// try to transform simple collision
-			if (UE::Geometry::ComponentTypeSupportsCollision(Component))
+			if (!bNeedSeparateTransactionForSimpleCollision)
 			{
-				UE::Geometry::TransformSimpleCollision(Component, ToBakePart);
+				UpdateSimpleCollision(Component, ToBakePart);
 			}
 
 			BakedTransforms.Add(ToBakePart);
 		}
 
-		Component->SetWorldTransform((FTransform)NewWorldPart);
+		Component->SetWorldTransform((FTransform)ComponentTransforms[ComponentIdx]);
 
 		AActor* TargetActor = UE::ToolTarget::GetTargetActor(Target);
 		if (TargetActor)
diff --git a/Engine/Plugins/Experimental/MeshModelingToolsetExp/Source/MeshModelingToolsExp/Private/EditPivotTool.cpp b/Engine/Plugins/Experimental/MeshModelingToolsetExp/Source/MeshModelingToolsExp/Private/EditPivotTool.cpp
--- a/Engine/Plugins/Experimental/MeshModelingToolsetExp/Source/MeshModelingToolsExp/Private/EditPivotTool.cpp
+++ b/Engine/Plugins/Experimental/MeshModelingToolsetExp/Source/MeshModelingToolsExp/Private/EditPivotTool.cpp
@@ -13,6 +13,7 @@
 #include "BaseBehaviors/ClickDragBehavior.h"
 #include "ToolSceneQueriesUtil.h"
 #include "Physics/ComponentCollisionUtil.h"
+#include "PhysicsEngine/BodySetup.h"
 
 #include "BaseGizmos/GizmoComponents.h"
 #include "BaseGizmos/TransformGizmoUtil.h"
@@ -400,8 +401,42 @@
 
 void UEditPivotTool::UpdateAssets(const FFrame3d& NewPivotWorldFrame)
 {
-	GetToolManager()->BeginUndoTransaction(LOCTEXT("EditPivotToolTransactionName", "Edit Pivot"));
+	// Make sure mesh descriptions are deserialized before we open transaction.
+	// This is to avoid potential stability issues related to creation/load of
+	// mesh descriptions inside a transaction.
+	// TODO: this may not be necessary anymore. Also may not be the most efficient
+	for (int32 ComponentIdx = 0; ComponentIdx < Targets.Num(); ComponentIdx++)
+	{
+		const FMeshDescription* MeshDescription = UE::ToolTarget::GetMeshDescription(Targets[ComponentIdx]);
+	}
 
+	// If the targets have any convex hulls, we need to update those first in a separate transaction to work around a crash
+	// in the handling of convex DDC data
+	constexpr bool bWorkaroundForCrashIfConvexAndMeshModifiedInSameTransaction = true;
+	bool bNeedSeparateTransactionForSimpleCollision = false;
+	if constexpr (bWorkaroundForCrashIfConvexAndMeshModifiedInSameTransaction)
+	{
+		for (int32 ComponentIdx = 0; ComponentIdx < Targets.Num(); ComponentIdx++)
+		{
+			UToolTarget* Target = Targets[ComponentIdx];
+			UPrimitiveComponent* Component = UE::ToolTarget::GetTargetComponent(Target);
+			if (UE::Geometry::ComponentTypeSupportsCollision(Component))
+			{
+				if (UBodySetup* BodySetup = UE::Geometry::GetBodySetup(Component))
+				{
+					if (BodySetup->AggGeom.ConvexElems.Num() > 0)
+					{
+						bNeedSeparateTransactionForSimpleCollision = true;
+						break;
+					}
+				}
+			}
+		}
+	}
+
+	// First pre-compute all the transforms we will need to apply
+	TArray<FTransformSequence3d> BakeTransforms;
+	TArray<FTransform> ComponentTransforms;
 	FTransform NewWorldTransform = NewPivotWorldFrame.ToFTransform();
 	FTransform NewWorldInverse = NewPivotWorldFrame.ToInverseFTransform();
 	TArray<FTransform> OriginalTransforms;
@@ -412,33 +447,14 @@
 	for (int32 ComponentIdx = 0; ComponentIdx < Targets.Num(); ComponentIdx++)
 	{
 		UToolTarget* Target = Targets[ComponentIdx];
-		UPrimitiveComponent* Component = UE::ToolTarget::GetTargetComponent(Target);
-		Component->Modify();
+		const UPrimitiveComponent* Component = UE::ToolTarget::GetTargetComponent(Target);
 
-		UInstancedStaticMeshComponent* InstancedComponent = Cast<UInstancedStaticMeshComponent>(Component);
+		const UInstancedStaticMeshComponent* InstancedComponent = Cast<const UInstancedStaticMeshComponent>(Component);
 		if (InstancedComponent != nullptr)
 		{
-			// For ISMC, we will not bake in a mesh transform, instead we will update the instance transforms relative to the new pivot
-			// TODO: this could be optional, and another alternative would be to bake the pivot to the mesh and then update
-			//   all the instance transforms so they stay in the same position?
-
-			// save world transforms
-			int32 NumInstances = InstancedComponent->GetInstanceCount();
-			TArray<FTransform> WorldTransforms;
-			WorldTransforms.SetNum(NumInstances);
-			for (int32 k = 0; k < NumInstances; ++k)
-			{
-				InstancedComponent->GetInstanceTransform(k, WorldTransforms[k], true);
-			}
-
-			// update position to new pivot
-			InstancedComponent->SetWorldTransform(NewWorldTransform);
-
-			// restore world transforms, which will compute new local transforms such that the instances do not move in the world
-			for (int32 k = 0; k < NumInstances; ++k)
-			{
-				InstancedComponent->UpdateInstanceTransform(k, WorldTransforms[k], true, true, false);
-			}
+			// no bake for instancedcomponents, nothing needs to be precomputed here
+			BakeTransforms.Emplace();
+			ComponentTransforms.Add(FTransform::Identity);
 		}
 		else if (MapToFirstOccurrences[ComponentIdx] == ComponentIdx)
 		{
@@ -491,39 +507,116 @@
 				// else do nothing -- scale is not invertible and must be baked
 			}
 
-			
+			FTransformSequence3d& Sequence = BakeTransforms.Emplace_GetRef();
+			Sequence.Append(ToBake);
+			if (bNeedSeparateScale)
+			{
+				Sequence.Append(SeparateBakeScale);
+			}
+			ComponentTransforms.Add(ScaledNewWorldTransform);
+		}
+		else
+		{
+			// try to invert baked transform -- note that this may not be a correct inverse if there is rotation + non-uniform scale,
+			// but it's the best we can do given we can not bake a separate custom scale when this is not the first occurrence
+			int32 FirstOccurrence = MapToFirstOccurrences[ComponentIdx];
+			FTransform ApproxBaked = OriginalTransforms[FirstOccurrence] * NewWorldInverse;
+			FTransform ComponentTransform = ApproxBaked.Inverse() * OriginalTransforms[ComponentIdx];
+			ComponentTransforms.Add(ComponentTransform);
+			BakeTransforms.Add(BakeTransforms[FirstOccurrence]);
+		}
+	}
 
+	// Pre-computed transforms must be 1:1 with Targets
+	check(ComponentTransforms.Num() == Targets.Num());
+	check(BakeTransforms.Num() == Targets.Num());
+
+	auto UpdateSimpleCollision = [](UPrimitiveComponent* Component, const FTransformSequence3d& ToBakeSeq)
+	{
+		// try to transform simple collision
+		if (UE::Geometry::ComponentTypeSupportsCollision(Component))
+		{
+			for (const FTransformSRT3d& ToBake : ToBakeSeq.GetTransforms())
+			{
+				UE::Geometry::TransformSimpleCollision(Component, ToBake);
+			}
+		}
+	};
+
+	// If necessary, make a first transaction to bake simple collision updates separately (works around crash in undo/redo of transactions that update both hulls and static meshes together)
+	if (bNeedSeparateTransactionForSimpleCollision)
+	{
+		GetToolManager()->BeginUndoTransaction(LOCTEXT("EditPivotToolSimpleCollisionTransactionName", "Edit Pivot Part 1 (Simple Collision)"));
+		for (int32 ComponentIdx = 0; ComponentIdx < Targets.Num(); ComponentIdx++)
+		{
+			UToolTarget* Target = Targets[ComponentIdx];
+			UPrimitiveComponent* Component = UE::ToolTarget::GetTargetComponent(Target);
+			Component->Modify();
+
+			UpdateSimpleCollision(Component, BakeTransforms[ComponentIdx]);
+		}
+		GetToolManager()->EndUndoTransaction();
+	}
+
+	GetToolManager()->BeginUndoTransaction(bNeedSeparateTransactionForSimpleCollision ?
+		LOCTEXT("EditPivotToolGeometryTransactionName", "Edit Pivot Part 2 (Visible Geometry)") :
+		LOCTEXT("EditPivotToolTransactionName", "Edit Pivot"));
+
+	for (int32 ComponentIdx = 0; ComponentIdx < Targets.Num(); ComponentIdx++)
+	{
+		UToolTarget* Target = Targets[ComponentIdx];
+		UPrimitiveComponent* Component = UE::ToolTarget::GetTargetComponent(Target);
+		Component->Modify();
+
+		UInstancedStaticMeshComponent* InstancedComponent = Cast<UInstancedStaticMeshComponent>(Component);
+		if (InstancedComponent != nullptr)
+		{
+			// For ISMC, we will not bake in a mesh transform, instead we will update the instance transforms relative to the new pivot
+			// TODO: this could be optional, and another alternative would be to bake the pivot to the mesh and then update
+			//   all the instance transforms so they stay in the same position?
+
+			// save world transforms
+			int32 NumInstances = InstancedComponent->GetInstanceCount();
+			TArray<FTransform> WorldTransforms;
+			WorldTransforms.SetNum(NumInstances);
+			for (int32 k = 0; k < NumInstances; ++k)
+			{
+				InstancedComponent->GetInstanceTransform(k, WorldTransforms[k], true);
+			}
+
+			// update position to new pivot
+			InstancedComponent->SetWorldTransform(NewWorldTransform);
+
+			// restore world transforms, which will compute new local transforms such that the instances do not move in the world
+			for (int32 k = 0; k < NumInstances; ++k)
+			{
+				InstancedComponent->UpdateInstanceTransform(k, WorldTransforms[k], true, true, false);
+			}
+		}
+		else if (MapToFirstOccurrences[ComponentIdx] == ComponentIdx)
+		{
 			FMeshDescription SourceMesh(UE::ToolTarget::GetMeshDescriptionCopy(Target));
 			FMeshDescriptionEditableTriangleMeshAdapter EditableMeshDescAdapter(&SourceMesh);
 			
-			MeshAdapterTransforms::ApplyTransform(EditableMeshDescAdapter, ToBake);
-			
-			if (bNeedSeparateScale)
+			for (const FTransformSRT3d& ToBake : BakeTransforms[ComponentIdx].GetTransforms())
 			{
-				MeshAdapterTransforms::ApplyTransform(EditableMeshDescAdapter, SeparateBakeScale);
+				MeshAdapterTransforms::ApplyTransform(EditableMeshDescAdapter, ToBake);
 			}
 			
 			// todo: support vertex-only update
 			UE::ToolTarget::CommitMeshDescriptionUpdate(Target, &SourceMesh);
 
 			// transform simple collision geometry
-			if (UE::Geometry::ComponentTypeSupportsCollision(Component))
+			if (!bNeedSeparateTransactionForSimpleCollision)
 			{
-				UE::Geometry::TransformSimpleCollision(Component, ToBake);
-				if (bNeedSeparateScale)
-				{
-					UE::Geometry::TransformSimpleCollision(Component, SeparateBakeScale);
-				}
+				UpdateSimpleCollision(Component, BakeTransforms[ComponentIdx]);
 			}
 
-			Component->SetWorldTransform(ScaledNewWorldTransform);
+			Component->SetWorldTransform(ComponentTransforms[ComponentIdx]);
 		}
 		else
 		{
-			// try to invert baked transform -- note that this may not be a correct inverse if there is rotation + non-uniform scale,
-			// but it's the best we can do given we can not bake a separate custom scale when this is not the first occurrence
-			FTransform Baked = OriginalTransforms[MapToFirstOccurrences[ComponentIdx]] * NewWorldInverse;
-			Component->SetWorldTransform(Baked.Inverse() * OriginalTransforms[ComponentIdx]);
+			Component->SetWorldTransform(ComponentTransforms[ComponentIdx]);
 		}
 
 		AActor* TargetActor = UE::ToolTarget::GetTargetActor(Target);
diff --git a/Engine/Plugins/Runtime/MeshModelingToolset/Source/ModelingComponents/Private/Physics/ComponentCollisionUtil.cpp b/Engine/Plugins/Runtime/MeshModelingToolset/Source/ModelingComponents/Private/Physics/ComponentCollisionUtil.cpp
--- a/Engine/Plugins/Runtime/MeshModelingToolset/Source/ModelingComponents/Private/Physics/ComponentCollisionUtil.cpp
+++ b/Engine/Plugins/Runtime/MeshModelingToolset/Source/ModelingComponents/Private/Physics/ComponentCollisionUtil.cpp
@@ -16,6 +16,8 @@
 #include "Physics/PhysicsDataCollection.h"
 
 #include "Components/StaticMeshComponent.h"
+#include "Components/BrushComponent.h"
+#include "Components/DynamicMeshComponent.h"
 #include "Engine/StaticMesh.h"
 #include "Engine/Classes/PhysicsEngine/BodySetup.h"
 #include "PhysicsEngine/AggregateGeom.h"
@@ -96,6 +98,50 @@ void UE::Geometry::UpdateSimpleCollision(
 }
 
 
+const UBodySetup* UE::Geometry::GetBodySetup(const UPrimitiveComponent* Component)
+{
+	if (const UStaticMeshComponent* StaticMeshComponent = Cast<UStaticMeshComponent>(Component))
+	{
+		if (UStaticMesh* StaticMesh = StaticMeshComponent->GetStaticMesh())
+		{
+			return StaticMesh->GetBodySetup();
+		}
+	}
+	else if (const UBrushComponent* BrushComponent = Cast<UBrushComponent>(Component))
+	{
+		return BrushComponent->BrushBodySetup;
+	}
+	else if (const UDynamicMeshComponent* DynamicMeshComponent = Cast<UDynamicMeshComponent>(Component))
+	{
+		return DynamicMeshComponent->GetBodySetup();
+	}
+
+	return nullptr;
+}
+
+UBodySetup* UE::Geometry::GetBodySetup(UPrimitiveComponent* Component)
+{
+	if (UStaticMeshComponent* StaticMeshComponent = Cast<UStaticMeshComponent>(Component))
+	{
+		if (UStaticMesh* StaticMesh = StaticMeshComponent->GetStaticMesh())
+		{
+			return StaticMesh->GetBodySetup();
+		}
+	}
+	else if (UBrushComponent* BrushComponent = Cast<UBrushComponent>(Component))
+	{
+		return BrushComponent->BrushBodySetup;
+	}
+	else if (UDynamicMeshComponent* DynamicMeshComponent = Cast<UDynamicMeshComponent>(Component))
+	{
+		return DynamicMeshComponent->GetBodySetup();
+	}
+
+	return nullptr;
+}
+
+
+
 bool UE::Geometry::SetSimpleCollision(
 	UPrimitiveComponent* Component,
 	const FSimpleShapeSet3d* ShapeSet,
diff --git a/Engine/Plugins/Runtime/MeshModelingToolset/Source/ModelingComponents/Public/Physics/ComponentCollisionUtil.h b/Engine/Plugins/Runtime/MeshModelingToolset/Source/ModelingComponents/Public/Physics/ComponentCollisionUtil.h
--- a/Engine/Plugins/Runtime/MeshModelingToolset/Source/ModelingComponents/Public/Physics/ComponentCollisionUtil.h
+++ b/Engine/Plugins/Runtime/MeshModelingToolset/Source/ModelingComponents/Public/Physics/ComponentCollisionUtil.h
@@ -88,6 +88,15 @@ MODELINGCOMPONENTS_API void UpdateSimpleCollision(
 	UStaticMesh* StaticMesh = nullptr,
 	FComponentCollisionSettings CollisionSettings = FComponentCollisionSettings());
 
+/**
+ * @return BodySetup on the given SourceComponent, or nullptr if no BodySetup was found
+ */
+MODELINGCOMPONENTS_API const UBodySetup* GetBodySetup(const UPrimitiveComponent* SourceComponent);
+/**
+ * @return BodySetup on the given SourceComponent, or nullptr if no BodySetup was found
+ */
+MODELINGCOMPONENTS_API UBodySetup* GetBodySetup(UPrimitiveComponent* SourceComponent);
+
 /**
  * Extract the simple collision geometry from AggGeom as meshes (ie spheres and capsules are tessellated) and
  * accumulate into MeshOut. TransformSequence is applied to the meshes as they are accumulated.
