MeshModelingToolset: Resolve typos in Offset tool

diff --git a/Engine/Plugins/Experimental/MeshModelingToolsetExp/Source/MeshModelingToolsExp/Private/OffsetMeshTool.cpp b/Engine/Plugins/Experimental/MeshModelingToolsetExp/Source/MeshModelingToolsExp/Private/OffsetMeshTool.cpp
--- a/Engine/Plugins/Experimental/MeshModelingToolsetExp/Source/MeshModelingToolsExp/Private/OffsetMeshTool.cpp
+++ b/Engine/Plugins/Experimental/MeshModelingToolsetExp/Source/MeshModelingToolsExp/Private/OffsetMeshTool.cpp
@@ -56,7 +56,7 @@ void UOffsetMeshTool::OnShutdown(EToolShutdownType ShutdownType)
 
 FText UOffsetMeshTool::GetToolMessageString() const
 {
-	return LOCTEXT("StartOffsetToolMessage", "Offset the mesh vertex positions using various Offseting methods.");
+	return LOCTEXT("StartOffsetToolMessage", "Offset the mesh vertex positions using various Offsetting methods.");
 }
 
 FText UOffsetMeshTool::GetAcceptTransactionName() const
diff --git a/Engine/Plugins/Experimental/MeshModelingToolsetExp/Source/MeshModelingToolsExp/Public/OffsetMeshTool.h b/Engine/Plugins/Experimental/MeshModelingToolsetExp/Source/MeshModelingToolsExp/Public/OffsetMeshTool.h
--- a/Engine/Plugins/Experimental/MeshModelingToolsetExp/Source/MeshModelingToolsExp/Public/OffsetMeshTool.h
+++ b/Engine/Plugins/Experimental/MeshModelingToolsetExp/Source/MeshModelingToolsExp/Public/OffsetMeshTool.h
@@ -11,10 +11,10 @@
 UENUM()
 enum class EOffsetMeshToolOffsetType : uint8
 {
-	/** Iterative Offseting with N iterations */
+	/** Iterative Offsetting with N iterations */
 	Iterative UMETA(DisplayName = "Iterative"),
 
-	/** Implicit Offseting, produces smoother output and does a better job at preserving UVs, but can be very slow on large meshes */
+	/** Implicit Offsetting, produces smoother output and does a better job at preserving UVs, but can be very slow on large meshes */
 	Implicit UMETA(DisplayName = "Implicit")
 };
 
@@ -27,7 +27,7 @@ class MESHMODELINGTOOLSEXP_API UOffsetMeshToolProperties : public UInteractiveTo
 	GENERATED_BODY()
 
 public:
-	/** Type of Offseting to apply */
+	/** Type of Offsetting to apply */
 	UPROPERTY(EditAnywhere, Category = Offset)
 	EOffsetMeshToolOffsetType OffsetType = EOffsetMeshToolOffsetType::Iterative;
 
@@ -54,14 +54,14 @@ public:
 };
 
 
-/** Properties for Iterative Offseting */
+/** Properties for Iterative Offsetting */
 UCLASS()
 class MESHMODELINGTOOLSEXP_API UIterativeOffsetProperties : public UInteractiveToolPropertySet
 {
 	GENERATED_BODY()
 public:
 
-	/** Number of Offseting iterations */
+	/** Number of Offsetting iterations */
 	UPROPERTY(EditAnywhere, Category = IterativeOffsetOptions, meta = (UIMin = "1", UIMax = "100", ClampMin = "0", ClampMax = "1000"))
 	int32 Steps = 10;
 
@@ -79,7 +79,7 @@ public:
 };
 
 
-/** Properties for Implicit Offseting */
+/** Properties for Implicit Offsetting */
 UCLASS()
 class MESHMODELINGTOOLSEXP_API UImplicitOffsetProperties : public UInteractiveToolPropertySet
 {
@@ -87,11 +87,11 @@ class MESHMODELINGTOOLSEXP_API UImplicitOffsetProperties : public UInteractiveTo
 
 public:
 	/** How tightly we should constrain the constrained implicit offset to the explicit offset */
-	UPROPERTY(EditAnywhere, Category = ImplicitOffsetingOptions, meta = (UIMin = "0.0", UIMax = "1.0", ClampMin = "0.0", ClampMax = "100.0"))
+	UPROPERTY(EditAnywhere, Category = ImplicitOffsettingOptions, meta = (UIMin = "0.0", UIMax = "1.0", ClampMin = "0.0", ClampMax = "100.0"))
 	float Smoothness = 0.2f;
 
 	/** If this is false, triangles will be reshaped to be more regular, which will distort UVs */
-	UPROPERTY(EditAnywhere, Category = ImplicitOffsetingOptions)
+	UPROPERTY(EditAnywhere, Category = ImplicitOffsettingOptions)
 	bool bPreserveUVs = true;
 };
 
@@ -101,7 +101,7 @@ public:
 
 
 /**
- * Mesh Offseting Tool
+ * Mesh Offsetting Tool
  */
 UCLASS()
 class MESHMODELINGTOOLSEXP_API UOffsetMeshTool : public UBaseMeshProcessingTool
diff --git a/Engine/Plugins/Runtime/MeshModelingToolset/Source/MeshModelingTools/Public/ToolActivities/PolyEditExtrudeActivity.h b/Engine/Plugins/Runtime/MeshModelingToolset/Source/MeshModelingTools/Public/ToolActivities/PolyEditExtrudeActivity.h
--- a/Engine/Plugins/Runtime/MeshModelingToolset/Source/MeshModelingTools/Public/ToolActivities/PolyEditExtrudeActivity.h
+++ b/Engine/Plugins/Runtime/MeshModelingToolset/Source/MeshModelingTools/Public/ToolActivities/PolyEditExtrudeActivity.h
@@ -177,7 +177,7 @@ public:
 		meta = (ClampMin = "1", EditConditionHides, EditCondition = "DirectionMode == EPolyEditOffsetModeOptions::SelectedTriangleNormalsEven"))
 	double MaxDistanceScaleFactor = 4.0;
 
-	/** Controls whether offseting an entire open-border patch should create a solid or an open shell */
+	/** Controls whether offsetting an entire open-border patch should create a solid or an open shell */
 	UPROPERTY(EditAnywhere, Category = Offset)
 	bool bShellsToSolids = true;
 
@@ -187,7 +187,7 @@ public:
 	EPolyEditExtrudeDirection MeasureDirection = EPolyEditExtrudeDirection::SelectionNormal;
 
 	/**
-	 * When offseting regions that touch the mesh border, assign the side groups (groups on the
+	 * When offsetting regions that touch the mesh border, assign the side groups (groups on the
 	 * stitched side of the extrude) in a way that considers edge colinearity. For instance, when
 	 * true, extruding a flat rectangle will give four different groups on its sides rather than
 	 * one connected group.
@@ -220,7 +220,7 @@ public:
 		meta = (ClampMin = "1", EditConditionHides, EditCondition = "DirectionMode == EPolyEditPushPullModeOptions::SelectedTriangleNormalsEven"))
 	double MaxDistanceScaleFactor = 4.0;
 
-	/** Controls whether offseting an entire open-border patch should create a solid or an open shell */
+	/** Controls whether offsetting an entire open-border patch should create a solid or an open shell */
 	UPROPERTY(EditAnywhere, Category = ExtrusionOptions)
 	bool bShellsToSolids = true;
 
@@ -230,7 +230,7 @@ public:
 	EPolyEditExtrudeDirection MeasureDirection = EPolyEditExtrudeDirection::SelectionNormal;
 
 	/**
-	 * When offseting regions that touch the mesh border, assign the side groups (groups on the
+	 * When offsetting regions that touch the mesh border, assign the side groups (groups on the
 	 * stitched side of the extrude) in a way that considers edge colinearity. For instance, when
 	 * true, extruding a flat rectangle will give four different groups on its sides rather than
 	 * one connected group.
