Fix stall when starting the render capture baking (BakeRC) tool

This change makes the BakeRC tool start instantly on all platforms and also makes the tool consistent with the BakeAll tool in that on first startup the same no-results-computed-yet material is applied to the the target mesh.

#rnx

diff --git a/Engine/Plugins/Experimental/MeshModelingToolsetExp/Source/MeshModelingToolsExp/Public/BakeRenderCaptureTool.h b/Engine/Plugins/Experimental/MeshModelingToolsetExp/Source/MeshModelingToolsExp/Public/BakeRenderCaptureTool.h
--- a/Engine/Plugins/Experimental/MeshModelingToolsetExp/Source/MeshModelingToolsExp/Public/BakeRenderCaptureTool.h
+++ b/Engine/Plugins/Experimental/MeshModelingToolsetExp/Source/MeshModelingToolsExp/Public/BakeRenderCaptureTool.h
@@ -88,31 +88,31 @@ public:
 
 	// Whether to generate a texture for the Base Color property
 	UPROPERTY(Category = RenderCaptureOptions, EditAnywhere, meta = (NoResetToDefault))
-	bool bBaseColorMap = true;
+	bool bBaseColorMap = false;
 
 	// Whether to generate a texture for the World Normal property
 	UPROPERTY(Category = RenderCaptureOptions, EditAnywhere, meta = (NoResetToDefault))
-	bool bNormalMap = true;
+	bool bNormalMap = false;
 
 	// Whether to generate a packed texture with Metallic, Roughness and Specular properties
 	UPROPERTY(Category = RenderCaptureOptions, EditAnywhere, meta = (NoResetToDefault), DisplayName="Packed MRS Map")
-	bool bPackedMRSMap = true;
+	bool bPackedMRSMap = false;
 
 	// Whether to generate a texture for the Metallic property
 	UPROPERTY(Category = RenderCaptureOptions, EditAnywhere, meta = (NoResetToDefault, EditCondition="bPackedMRSMap == false"))
-	bool bMetallicMap = true;
+	bool bMetallicMap = false;
 
 	// Whether to generate a texture for the Roughness property
 	UPROPERTY(Category = RenderCaptureOptions, EditAnywhere, meta = (NoResetToDefault, EditCondition="bPackedMRSMap == false"))
-	bool bRoughnessMap = true;
+	bool bRoughnessMap = false;
 
 	// Whether to generate a texture for the Specular property
 	UPROPERTY(Category = RenderCaptureOptions, EditAnywhere, meta = (NoResetToDefault, EditCondition="bPackedMRSMap == false"))
-	bool bSpecularMap = true;
+	bool bSpecularMap = false;
 	
 	// Whether to generate a texture for the Emissive property
 	UPROPERTY(Category = RenderCaptureOptions, EditAnywhere, meta = (NoResetToDefault))
-	bool bEmissiveMap = true;
+	bool bEmissiveMap = false;
 
 	// Whether to use anti-aliasing in the render captures, this may introduce artefacts if pixels at different scene depths get blended
 	UPROPERTY(Category = RenderCaptureOptions, EditAnywhere, AdvancedDisplay)
