PatternTool: Improved UI for Jitter (enabled inline, proportional jitter changes in scale settings with bUniform)

diff --git a/Engine/Plugins/Experimental/MeshModelingToolsetExp/Source/MeshModelingToolsExp/Private/PatternTool.cpp b/Engine/Plugins/Experimental/MeshModelingToolsetExp/Source/MeshModelingToolsExp/Private/PatternTool.cpp
--- a/Engine/Plugins/Experimental/MeshModelingToolsetExp/Source/MeshModelingToolsExp/Private/PatternTool.cpp
+++ b/Engine/Plugins/Experimental/MeshModelingToolsetExp/Source/MeshModelingToolsExp/Private/PatternTool.cpp
@@ -577,6 +577,7 @@ void UPatternTool::Setup()
 			
 			StartScaleDirection = ScaleSettings->StartScale.GetSafeNormal(Tolerance, DefaultDirection);
 			EndScaleDirection = ScaleSettings->EndScale.GetSafeNormal(Tolerance, DefaultDirection);
+			JitterScaleDirection = ScaleSettings->Jitter.GetSafeNormal(Tolerance, DefaultDirection);
 		}
 	};
 	ScaleSettings->WatchProperty(ScaleSettings->bUniform, OnUniformChanged);
@@ -599,6 +600,15 @@ void UPatternTool::Setup()
 		}
 	});
 
+	JitterScaleWatcherIdx = ScaleSettings->WatchProperty(ScaleSettings->Jitter, [this](const FVector& NewJitter)
+	{
+		if (ScaleSettings->bUniform)
+		{
+			ScaleSettings->Jitter = JitterScaleDirection * NewJitter.Size();
+			ScaleSettings->SilentUpdateWatcherAtIndex(JitterScaleWatcherIdx);
+		}
+	});
+
 	// Initialize StartScaleDirection and EndScaleDirection
 	OnUniformChanged(true);
 	
@@ -855,7 +865,7 @@ static void InitializeGenerator(FPatternGenerator& Generator, UPatternTool* Tool
 	
 	Generator.bInterpolateRotation = Tool->RotationSettings->bInterpolate;
 	Generator.bJitterRotation = Tool->RotationSettings->bJitter;
-	Generator.RotationJitterRange = Tool->RotationSettings->RotationJitterRange;
+	Generator.RotationJitterRange = Tool->RotationSettings->Jitter;
 	if (Generator.bInterpolateRotation)
 	{
 		Generator.StartRotation = FQuaterniond(Tool->RotationSettings->StartRotation);
@@ -868,7 +878,7 @@ static void InitializeGenerator(FPatternGenerator& Generator, UPatternTool* Tool
 
 	Generator.bInterpolateTranslation = Tool->TranslationSettings->bInterpolate;
 	Generator.bJitterTranslation = Tool->TranslationSettings->bJitter;
-	Generator.TranslationJitterRange = Tool->TranslationSettings->TranslationJitterRange;
+	Generator.TranslationJitterRange = Tool->TranslationSettings->Jitter;
 	if (Generator.bInterpolateTranslation)
 	{
 		Generator.StartTranslation = Tool->TranslationSettings->StartTranslation;
@@ -881,7 +891,7 @@ static void InitializeGenerator(FPatternGenerator& Generator, UPatternTool* Tool
 
 	Generator.bInterpolateScale = Tool->ScaleSettings->bInterpolate;
 	Generator.bJitterScale = Tool->ScaleSettings->bJitter;
-	Generator.ScaleJitterRange = (Tool->ScaleSettings->bUniform) ? (FVector3d::One()*Tool->ScaleSettings->ScaleJitterRange) : Tool->ScaleSettings->ScaleJitterRangeNonUniform;
+	Generator.ScaleJitterRange = Tool->ScaleSettings->Jitter;
 	if (Generator.bInterpolateScale)
 	{
 		Generator.StartScale = Tool->ScaleSettings->StartScale;
diff --git a/Engine/Plugins/Experimental/MeshModelingToolsetExp/Source/MeshModelingToolsExp/Public/PatternTool.h b/Engine/Plugins/Experimental/MeshModelingToolsetExp/Source/MeshModelingToolsExp/Public/PatternTool.h
--- a/Engine/Plugins/Experimental/MeshModelingToolsetExp/Source/MeshModelingToolsExp/Public/PatternTool.h
+++ b/Engine/Plugins/Experimental/MeshModelingToolsetExp/Source/MeshModelingToolsExp/Public/PatternTool.h
@@ -275,7 +275,7 @@ public:
 	bool bInterpolate = false;
 
 	/** If true, Rotation at each Pattern Element is offset by a uniformly chosen random value in the range of [-RotationJitterRange, RotationJitterRange] */
-	UPROPERTY(EditAnywhere, Category = Rotation)
+	UPROPERTY(EditAnywhere, Category = Rotation, meta = (InlineEditConditionToggle))
 	bool bJitter = false;
 
 	/** Rotation at first Pattern Element */
@@ -287,8 +287,8 @@ public:
 	FRotator EndRotation = FRotator::ZeroRotator;
 
 	/** Upper bound of the range which is sampled to randomly rotate each Pattern Element if Jitter is true */
-	UPROPERTY(EditAnywhere, Category = Rotation, meta = (ClampMin = 0, EditCondition = "bJitter", EditConditionHides, HideEditConditionToggle))
-	FRotator RotationJitterRange = FRotator::ZeroRotator;
+	UPROPERTY(EditAnywhere, Category = Rotation, meta = (ClampMin = 0, EditCondition = "bJitter"))
+	FRotator Jitter = FRotator::ZeroRotator;
 };
 
 
@@ -305,7 +305,7 @@ public:
 	bool bInterpolate = false;
 
 	/** If true, Translation at each Pattern Element is offset by a uniformly chosen random value in the range of [-TranslationJitterRange, TranslationJitterRange] */
-	UPROPERTY(EditAnywhere, Category = Translation)
+	UPROPERTY(EditAnywhere, Category = Translation, meta = (InlineEditConditionToggle))
 	bool bJitter = false;
 
 	/** Translation at first Pattern Element */
@@ -317,8 +317,8 @@ public:
 	FVector EndTranslation = FVector::ZeroVector;
 
 	/** Upper bound of the range which is sampled to randomly translate each Pattern Element if Jitter is true */
-	UPROPERTY(EditAnywhere, Category = Translation, meta = (ClampMin = 0, EditCondition = "bJitter", EditConditionHides, HideEditConditionToggle))
-	FVector TranslationJitterRange = FVector::ZeroVector;
+	UPROPERTY(EditAnywhere, Category = Translation, meta = (ClampMin = 0, EditCondition = "bJitter"))
+	FVector Jitter = FVector::ZeroVector;
 };
 
 
@@ -339,7 +339,7 @@ public:
 	bool bInterpolate = false;
 
 	/** If true, Scale at each Pattern Element is offset by a uniformly chosen random value in the range of [-ScaleJitterRange, ScaleJitterRange] */
-	UPROPERTY(EditAnywhere, Category = Scale)
+	UPROPERTY(EditAnywhere, Category = Scale, meta = (InlineEditConditionToggle))
 	bool bJitter = false;
 	
 	/** Uniform Scale at first Pattern Element */
@@ -349,14 +349,10 @@ public:
 	/** Uniform Scale applied to all Pattern Elements, or at Last Pattern Element for Interpolated scales */
 	UPROPERTY(EditAnywhere, Category = Scale, meta = (EditCondition = "bInterpolate"))
 	FVector EndScale = FVector::OneVector;
-	
-	/** Upper bound of the range which is sampled to randomly scale each Pattern Element if Jitter is true */
-	UPROPERTY(EditAnywhere, Category = Scale, meta = (ClampMin = 0, EditCondition = "bJitter && bUniform", EditConditionHides, HideEditConditionToggle))
-	float ScaleJitterRange = 0.0f;
 
-	/** Upper bound of the range which is sampled to randomly scale each Pattern Element if Jitter is true (Non-Uniform) */
-	UPROPERTY(EditAnywhere, Category = Scale, meta = (ClampMin = 0, DisplayName = "Scale Jitter Range", EditCondition = "bJitter && bUniform == false", EditConditionHides, HideEditConditionToggle))
-	FVector ScaleJitterRangeNonUniform = FVector(0, 0, 0);
+	/** Upper bound of the range which is sampled to randomly scale each Pattern Element if Jitter is true */
+	UPROPERTY(EditAnywhere, Category = Scale, meta = (ClampMin = 0, EditCondition = "bJitter"))
+	FVector Jitter = FVector::ZeroVector;
 };
 
 
@@ -438,9 +434,11 @@ public:
 	
 	FVector StartScaleDirection;
 	FVector EndScaleDirection;
+	FVector JitterScaleDirection;
 	
 	int32 StartScaleWatcherIdx;
 	int32 EndScaleWatcherIdx;
+	int32 JitterScaleWatcherIdx;
 
 	UPROPERTY()
 	TObjectPtr<UPatternTool_OutputSettings> OutputSettings;
