PatternTool: Fix some UI elements.

Seed should have no revert to default option and some tooltips are wrong after updating UI in previous CL.

diff --git a/Engine/Plugins/Experimental/MeshModelingToolsetExp/Source/MeshModelingToolsExp/Private/PatternTool.cpp b/Engine/Plugins/Experimental/MeshModelingToolsetExp/Source/MeshModelingToolsExp/Private/PatternTool.cpp
--- a/Engine/Plugins/Experimental/MeshModelingToolsetExp/Source/MeshModelingToolsExp/Private/PatternTool.cpp
+++ b/Engine/Plugins/Experimental/MeshModelingToolsetExp/Source/MeshModelingToolsExp/Private/PatternTool.cpp
@@ -609,7 +609,7 @@ void UPatternTool::Setup()
 		}
 	});
 
-	// Initialize StartScaleDirection and EndScaleDirection
+	// Initialize StartScaleDirection, EndScaleDirection, and JitterScaleDirection
 	OnUniformChanged(true);
 	
 	OutputSettings = NewObject<UPatternTool_OutputSettings>();
diff --git a/Engine/Plugins/Experimental/MeshModelingToolsetExp/Source/MeshModelingToolsExp/Public/PatternTool.h b/Engine/Plugins/Experimental/MeshModelingToolsetExp/Source/MeshModelingToolsExp/Public/PatternTool.h
--- a/Engine/Plugins/Experimental/MeshModelingToolsetExp/Source/MeshModelingToolsExp/Public/PatternTool.h
+++ b/Engine/Plugins/Experimental/MeshModelingToolsetExp/Source/MeshModelingToolsExp/Public/PatternTool.h
@@ -119,7 +119,7 @@ public:
 	bool bHideSources = true;
 
 	/** The seed used to introduce random transform variations when enabled */
-	UPROPERTY(EditAnywhere, Category = Shape)
+	UPROPERTY(EditAnywhere, Category = Shape, meta = (NoResetToDefault))
 	int32 Seed = FMath::Rand();
 
 	/** Whether or not the pattern items should be projected along the negative Z axis of the plane mechanic */
@@ -278,11 +278,11 @@ public:
 	UPROPERTY(EditAnywhere, Category = Rotation, meta = (InlineEditConditionToggle))
 	bool bJitter = false;
 
-	/** Rotation at first Pattern Element */
+	/** Rotation applied to all Pattern Elements, or to first Pattern Element for Interpolated rotation */
 	UPROPERTY(EditAnywhere, Category = Rotation)
 	FRotator StartRotation = FRotator::ZeroRotator;
 
-	/** Rotation applied to all Pattern Elements, or at Last Pattern Element for Interpolated rotations */
+	/** Rotation applied to last Pattern Elements for Interpolated rotation */
 	UPROPERTY(EditAnywhere, Category = Rotation, meta = (EditCondition = "bInterpolate"))
 	FRotator EndRotation = FRotator::ZeroRotator;
 
@@ -308,11 +308,11 @@ public:
 	UPROPERTY(EditAnywhere, Category = Translation, meta = (InlineEditConditionToggle))
 	bool bJitter = false;
 
-	/** Translation at first Pattern Element */
+	/** Translation applied to all Pattern Elements, or to first Pattern Element for Interpolated translation */
 	UPROPERTY(EditAnywhere, Category = Translation)
 	FVector StartTranslation = FVector::ZeroVector;
 
-	/** Translation applied to all Pattern Elements, or at Last Pattern Element for Interpolated translations */
+	/** Translation applied to last Pattern Element for Interpolated translation */
 	UPROPERTY(EditAnywhere, Category = Translation, meta = (EditCondition = "bInterpolate"))
 	FVector EndTranslation = FVector::ZeroVector;
 
@@ -330,7 +330,7 @@ class MESHMODELINGTOOLSEXP_API UPatternTool_ScaleSettings : public UInteractiveT
 {
 	GENERATED_BODY()
 public:
-	/** If true, Scaling is limited to Uniform Scaling */
+	/** If true, changes to Start Scale, End Scale, and Jitter are proportional along all the axes */
 	UPROPERTY(EditAnywhere, Category = Scale)
 	bool bUniform = true;
 	
@@ -342,11 +342,11 @@ public:
 	UPROPERTY(EditAnywhere, Category = Scale, meta = (InlineEditConditionToggle))
 	bool bJitter = false;
 	
-	/** Uniform Scale at first Pattern Element */
+	/** Scale applied to all Pattern Elements, or to first Pattern Element for Interpolated scale */
 	UPROPERTY(EditAnywhere, Category = Scale)
 	FVector StartScale = FVector::OneVector;
 	
-	/** Uniform Scale applied to all Pattern Elements, or at Last Pattern Element for Interpolated scales */
+	/** Scale applied to last Pattern Element for Interpolated scale */
 	UPROPERTY(EditAnywhere, Category = Scale, meta = (EditCondition = "bInterpolate"))
 	FVector EndScale = FVector::OneVector;
 
@@ -359,7 +359,7 @@ public:
 
 
 /**
- * Ouptput Settings for the Pattern Tool
+ * Output Settings for the Pattern Tool
  */
 UCLASS()
 class MESHMODELINGTOOLSEXP_API UPatternTool_OutputSettings : public UInteractiveToolPropertySet
@@ -378,7 +378,7 @@ public:
 	UPROPERTY(EditAnywhere, Category = Output, meta = (EditCondition = "bHaveStaticMeshes == true && bSeparateActors == false && bConvertToDynamic == false", HideEditConditionToggle))
 	bool bCreateISMCs = false;
 
-	// internal, used to control state of Instance settings
+	/** internal, used to control state of Instance settings */
 	UPROPERTY(meta = (TransientToolProperty))
 	bool bHaveStaticMeshes = false;
 };
