PatternTool: Fix several issues related to the behavior of Proportional scaling.

bUniform was renamed to bProportional

Proportional scaling behaves almost the same but when a value is entered numerically, previously the approach to scaling would change the just entered value by finding a vector of the new vectors magnitude in the old vectors direction while the new approach scales the other two components up, resulting in more predictable behavior.

The minimum value for a component for the fields StartScale, EndScale, or Jitter in the scale category only is now 0.001 to be consistent with the clamping behavior in ApplyElementTransforms() and to prevent issues with dividing by extremely small (or 0) values when computing ScaleFactor in ApplyProportionalScale.

When proportional scaling is enabled and the reset to default button is pressed for a vector property, it now properly resets to the default. Previously, the resulting vector would be the same length as the default vector but it would maintain the direction of the previous value.

diff --git a/Engine/Plugins/Experimental/MeshModelingToolsetExp/Source/MeshModelingToolsExp/Private/PatternTool.cpp b/Engine/Plugins/Experimental/MeshModelingToolsetExp/Source/MeshModelingToolsExp/Private/PatternTool.cpp
--- a/Engine/Plugins/Experimental/MeshModelingToolsetExp/Source/MeshModelingToolsExp/Private/PatternTool.cpp
+++ b/Engine/Plugins/Experimental/MeshModelingToolsetExp/Source/MeshModelingToolsExp/Private/PatternTool.cpp
@@ -127,7 +127,7 @@ public:
 	FVector3d ScaleJitterRange = FVector3d::Zero();
 	bool bInterpolateScale = false;							// if false, only StartScale is used
 	bool bJitterScale = false;
-
+	
 	FAxisAlignedBox3d Dimensions = FAxisAlignedBox3d(FVector3d::Zero(), 10.0);
 
 public:
@@ -197,10 +197,9 @@ public:
 			FVector3d ScaleJitter;
 			FVector3d TransformScale = Transform.GetScale();
 
-			constexpr double MIN_SCALE = 0.001f;
-			ScaleJitter.X = FMath::Max(MIN_SCALE, TransformScale.X + FMath::Lerp(-ScaleJitterRange.X, ScaleJitterRange.X, ScaleRandomStream.GetFraction()));
-			ScaleJitter.Y = FMath::Max(MIN_SCALE, TransformScale.Y + FMath::Lerp(-ScaleJitterRange.Y, ScaleJitterRange.Y, ScaleRandomStream.GetFraction()));
-			ScaleJitter.Z = FMath::Max(MIN_SCALE, TransformScale.Z + FMath::Lerp(-ScaleJitterRange.Z, ScaleJitterRange.Z, ScaleRandomStream.GetFraction()));
+			ScaleJitter.X = FMath::Max(UPatternTool_ScaleSettings::MinScale, TransformScale.X + FMath::Lerp(-ScaleJitterRange.X, ScaleJitterRange.X, ScaleRandomStream.GetFraction()));
+			ScaleJitter.Y = FMath::Max(UPatternTool_ScaleSettings::MinScale, TransformScale.Y + FMath::Lerp(-ScaleJitterRange.Y, ScaleJitterRange.Y, ScaleRandomStream.GetFraction()));
+			ScaleJitter.Z = FMath::Max(UPatternTool_ScaleSettings::MinScale, TransformScale.Z + FMath::Lerp(-ScaleJitterRange.Z, ScaleJitterRange.Z, ScaleRandomStream.GetFraction()));
 				
 			Transform.SetScale( ScaleJitter );
 		}
@@ -568,49 +567,71 @@ void UPatternTool::Setup()
 	AddToolPropertySource(ScaleSettings);
 	ScaleSettings->RestoreProperties(this);
 
-	auto OnUniformChanged = [this](bool bNewValue)
+	auto OnProportionalChanged = [this](bool bNewValue)
 	{
 		if (bNewValue)
 		{
-			constexpr float Tolerance = 1E-08;
-			const FVector DefaultDirection = FVector::OneVector.GetUnsafeNormal();
-			
-			StartScaleDirection = ScaleSettings->StartScale.GetSafeNormal(Tolerance, DefaultDirection);
-			EndScaleDirection = ScaleSettings->EndScale.GetSafeNormal(Tolerance, DefaultDirection);
-			JitterScaleDirection = ScaleSettings->Jitter.GetSafeNormal(Tolerance, DefaultDirection);
+			CachedStartScale = ScaleSettings->StartScale;
+			CachedEndScale = ScaleSettings->EndScale;
+			CachedJitterScale = ScaleSettings->Jitter;
 		}
 	};
-	ScaleSettings->WatchProperty(ScaleSettings->bUniform, OnUniformChanged);
+	ScaleSettings->WatchProperty(ScaleSettings->bProportional, OnProportionalChanged);
+	OnProportionalChanged(true);	// Initialize StartScaleDirection and EndScaleDirection
+
+	auto ApplyProportionalScale = [this](FVector& NewVector, FVector& CachedVector)
+	{
+		// Determines which component of the vector is being changed by looking at which component is
+		// most different from the previous values
+		FVector Difference = NewVector - CachedVector;
+		int32 DifferenceMaxElementIndex = MaxAbsElementIndex(Difference);
+		
+		// This approach to proportional scaling is desirable because when a user manually enters data
+		// numerically, we scale the other two components such that the entered value is unchanged unless
+		// doing so would result in component values less than MinScale, in which case the the resulting
+		// vector will be in the correct direction but lengthened to a degree to ensure all components are
+		// greater than or equal to MinScale.
+		double ScaleFactor = FMath::Max(NewVector[DifferenceMaxElementIndex] / CachedVector[DifferenceMaxElementIndex], UPatternTool_ScaleSettings::MinScale / CachedVector[MinElementIndex(CachedVector)]);
+		
+		NewVector = CachedVector * ScaleFactor;
+		CachedVector = NewVector;
+	};
 	
-	StartScaleWatcherIdx = ScaleSettings->WatchProperty(ScaleSettings->StartScale, [this](const FVector& NewStartScale)
+	StartScaleWatcherIdx = ScaleSettings->WatchProperty(ScaleSettings->StartScale, [this, &ApplyProportionalScale](const FVector& NewStartScale)
 	{
-		if (ScaleSettings->bUniform)
+		if (ScaleSettings->bProportional)
 		{
-			ScaleSettings->StartScale = StartScaleDirection * NewStartScale.Size();
+			ApplyProportionalScale(ScaleSettings->StartScale, CachedStartScale);
 			ScaleSettings->SilentUpdateWatcherAtIndex(StartScaleWatcherIdx);
+
+			// This is needed in addition to the call in OnPropertyModified due to the fact that these watchers are
+			// called after OnPropertyModified and in some cases the scale values used were inconsistent with the scale
+			// values being displayed in the details panel.
+			OnParametersUpdated();
 		}
 	});
 	
-	EndScaleWatcherIdx = ScaleSettings->WatchProperty(ScaleSettings->EndScale, [this](const FVector& NewEndScale)
+	EndScaleWatcherIdx = ScaleSettings->WatchProperty(ScaleSettings->EndScale, [this, &ApplyProportionalScale](const FVector& NewEndScale)
 	{
-		if (ScaleSettings->bUniform)
+		if (ScaleSettings->bProportional)
 		{
-			ScaleSettings->EndScale = EndScaleDirection * NewEndScale.Size();
+			ApplyProportionalScale(ScaleSettings->EndScale, CachedEndScale);
 			ScaleSettings->SilentUpdateWatcherAtIndex(EndScaleWatcherIdx);
+			
+			OnParametersUpdated();
 		}
 	});
 
-	JitterScaleWatcherIdx = ScaleSettings->WatchProperty(ScaleSettings->Jitter, [this](const FVector& NewJitter)
+	JitterScaleWatcherIdx = ScaleSettings->WatchProperty(ScaleSettings->Jitter, [this, &ApplyProportionalScale](const FVector& NewJitterScale)
 	{
-		if (ScaleSettings->bUniform)
+		if (ScaleSettings->bProportional)
 		{
-			ScaleSettings->Jitter = JitterScaleDirection * NewJitter.Size();
+			ApplyProportionalScale(ScaleSettings->Jitter, CachedJitterScale);
 			ScaleSettings->SilentUpdateWatcherAtIndex(JitterScaleWatcherIdx);
+			
+			OnParametersUpdated();
 		}
 	});
-
-	// Initialize StartScaleDirection, EndScaleDirection, and JitterScaleDirection
-	OnUniformChanged(true);
 	
 	OutputSettings = NewObject<UPatternTool_OutputSettings>();
 	AddToolPropertySource(OutputSettings);
@@ -786,15 +807,31 @@ void UPatternTool::InitializeElements()
 
 void UPatternTool::OnPropertyModified(UObject* PropertySet, FProperty* Property)
 {
-	if (PropertySet == RadialSettings 
-		|| PropertySet == LinearSettings 
-		|| PropertySet == GridSettings
-		|| PropertySet == RotationSettings
-		|| PropertySet == TranslationSettings
-		|| PropertySet == ScaleSettings )
-	{
-		OnParametersUpdated();
+	if (PropertySet == Settings || PropertySet == OutputSettings)
+	{
+		return;
 	}
+
+	if (PropertySet == ScaleSettings && ScaleSettings->bProportional)
+	{
+		// We are silencing all watchers if Property corresponds to an FVector UProperty because this indicates that
+		// the "Reset to Default" button was pressed. If individual components are modified instead, then Property will
+		// not be castable to an FStructProperty
+		
+		const FStructProperty* StructProperty = CastField<FStructProperty>(Property);
+		if (StructProperty != nullptr && StructProperty->Struct->GetName() == FString("Vector"))
+		{
+			CachedStartScale = ScaleSettings->StartScale;
+			CachedEndScale = ScaleSettings->EndScale;
+			CachedJitterScale = ScaleSettings->Jitter;
+			
+			ScaleSettings->SilentUpdateWatcherAtIndex(StartScaleWatcherIdx);
+			ScaleSettings->SilentUpdateWatcherAtIndex(EndScaleWatcherIdx);
+			ScaleSettings->SilentUpdateWatcherAtIndex(JitterScaleWatcherIdx);
+		}
+	}
+
+	OnParametersUpdated();
 }
 
 
diff --git a/Engine/Plugins/Experimental/MeshModelingToolsetExp/Source/MeshModelingToolsExp/Public/PatternTool.h b/Engine/Plugins/Experimental/MeshModelingToolsetExp/Source/MeshModelingToolsExp/Public/PatternTool.h
--- a/Engine/Plugins/Experimental/MeshModelingToolsetExp/Source/MeshModelingToolsExp/Public/PatternTool.h
+++ b/Engine/Plugins/Experimental/MeshModelingToolsetExp/Source/MeshModelingToolsExp/Public/PatternTool.h
@@ -321,7 +321,6 @@ public:
 	FVector Jitter = FVector::ZeroVector;
 };
 
-
 /**
  * Settings for Per Element Scale in the Pattern Tool
  */
@@ -330,9 +329,12 @@ class MESHMODELINGTOOLSEXP_API UPatternTool_ScaleSettings : public UInteractiveT
 {
 	GENERATED_BODY()
 public:
+	/** Initial value for Jitter and referenced in FPatternGenerator when applying scale jitter. Can't be used for ClampMin. */
+	static constexpr double MinScale = 0.001;
+	
 	/** If true, changes to Start Scale, End Scale, and Jitter are proportional along all the axes */
 	UPROPERTY(EditAnywhere, Category = Scale)
-	bool bUniform = true;
+	bool bProportional = true;
 	
 	/** If true, Scale is linearly interpolated between StartScale and Scale values */
 	UPROPERTY(EditAnywhere, Category = Scale, meta = (InlineEditConditionToggle))
@@ -343,16 +345,16 @@ public:
 	bool bJitter = false;
 	
 	/** Scale applied to all Pattern Elements, or to first Pattern Element for Interpolated scale */
-	UPROPERTY(EditAnywhere, Category = Scale)
+	UPROPERTY(EditAnywhere, Category = Scale, meta = (ClampMin = 0.001))
 	FVector StartScale = FVector::OneVector;
 	
 	/** Scale applied to last Pattern Element for Interpolated scale */
-	UPROPERTY(EditAnywhere, Category = Scale, meta = (EditCondition = "bInterpolate"))
+	UPROPERTY(EditAnywhere, Category = Scale, meta = (ClampMin = 0.001, EditCondition = "bInterpolate"))
 	FVector EndScale = FVector::OneVector;
 
 	/** Upper bound of the range which is sampled to randomly scale each Pattern Element if Jitter is true */
-	UPROPERTY(EditAnywhere, Category = Scale, meta = (ClampMin = 0, EditCondition = "bJitter"))
-	FVector Jitter = FVector::ZeroVector;
+	UPROPERTY(EditAnywhere, Category = Scale, meta = (ClampMin = 0.001, EditCondition = "bJitter"))
+	FVector Jitter = FVector(MinScale);
 };
 
 
@@ -432,9 +434,9 @@ public:
 	UPROPERTY()
 	TObjectPtr<UPatternTool_ScaleSettings> ScaleSettings;
 	
-	FVector StartScaleDirection;
-	FVector EndScaleDirection;
-	FVector JitterScaleDirection;
+	FVector CachedStartScale;
+	FVector CachedEndScale;
+	FVector CachedJitterScale;
 	
 	int32 StartScaleWatcherIdx;
 	int32 EndScaleWatcherIdx;
