GeometryTools: Convert Meshes -  fix bug in materials when converting multiple meshes and also hide the transfer materials option when converting to a volume.

diff --git a/Engine/Plugins/Experimental/MeshModelingToolsetExp/Source/MeshModelingToolsExp/Private/ConvertMeshesTool.cpp b/Engine/Plugins/Experimental/MeshModelingToolsetExp/Source/MeshModelingToolsExp/Private/ConvertMeshesTool.cpp
--- a/Engine/Plugins/Experimental/MeshModelingToolsetExp/Source/MeshModelingToolsExp/Private/ConvertMeshesTool.cpp
+++ b/Engine/Plugins/Experimental/MeshModelingToolsetExp/Source/MeshModelingToolsExp/Private/ConvertMeshesTool.cpp
@@ -51,11 +51,17 @@ void UConvertMeshesTool::Setup()
 	OutputTypeProperties = NewObject<UCreateMeshObjectTypeProperties>(this);
 	OutputTypeProperties->RestoreProperties(this, TEXT("ConvertMeshesTool"));
 	OutputTypeProperties->InitializeDefault();
-	OutputTypeProperties->WatchProperty(OutputTypeProperties->OutputType, [this](FString) { OutputTypeProperties->UpdatePropertyVisibility(); });
-	AddToolPropertySource(OutputTypeProperties);
+	
 
 	BasicProperties = NewObject<UConvertMeshesToolProperties>(this);
 	BasicProperties->RestoreProperties(this);
+	BasicProperties->bShowTransferMaterials = !OutputTypeProperties->bShowVolumeList;
+	
+	OutputTypeProperties->WatchProperty(OutputTypeProperties->OutputType, [this](FString) { 
+																							OutputTypeProperties->UpdatePropertyVisibility(); 
+																							BasicProperties->bShowTransferMaterials = !OutputTypeProperties->bShowVolumeList;
+																							});
+	AddToolPropertySource(OutputTypeProperties);
 	AddToolPropertySource(BasicProperties);
 
 	SetToolDisplayName(LOCTEXT("ToolName", "Convert"));
@@ -94,7 +100,7 @@ void UConvertMeshesTool::OnShutdown(EToolShutdownType ShutdownType)
 			FTransform SourceTransform = (FTransform)UE::ToolTarget::GetLocalToWorldTransform(Targets[k]);
 			FDynamicMesh3 SourceMesh = UE::ToolTarget::GetDynamicMeshCopy(Targets[k], true);
 			FString AssetName = TargetActor->GetActorNameOrLabel();
-			FComponentMaterialSet Materials = UE::ToolTarget::GetMaterialSet(Targets[0]);
+			FComponentMaterialSet Materials = UE::ToolTarget::GetMaterialSet(Targets[k]);
 			const FComponentMaterialSet* TransferMaterials = (BasicProperties->bTransferMaterials) ? &Materials : nullptr;
 
 			FCreateMeshObjectParams NewMeshObjectParams;
diff --git a/Engine/Plugins/Experimental/MeshModelingToolsetExp/Source/MeshModelingToolsExp/Public/ConvertMeshesTool.h b/Engine/Plugins/Experimental/MeshModelingToolsetExp/Source/MeshModelingToolsExp/Public/ConvertMeshesTool.h
--- a/Engine/Plugins/Experimental/MeshModelingToolsetExp/Source/MeshModelingToolsExp/Public/ConvertMeshesTool.h
+++ b/Engine/Plugins/Experimental/MeshModelingToolsetExp/Source/MeshModelingToolsExp/Public/ConvertMeshesTool.h
@@ -31,8 +31,12 @@ class MESHMODELINGTOOLSEXP_API UConvertMeshesToolProperties : public UInteractiv
 {
 	GENERATED_BODY()
 public:
-	UPROPERTY(EditAnywhere, Category = Options)
+	UPROPERTY(EditAnywhere, Category = Options, meta = (EditCondition = "bShowTransferMaterials == true", EditConditionHides, HideEditConditionToggle))
 	bool bTransferMaterials = true;
+
+	// control whether the transfer materials option is displayed
+	UPROPERTY(meta = (TransientToolProperty))
+	bool bShowTransferMaterials = true;
 };
 
 
