PatternTool: Add support for relative transforms when using multiple source elements and add functionality to automatically adjust bounding boxes to accommodate relative transformations and Start Scale/Start Rotation

Previously, the pattern tool supported using multiple source elements under the assumption that they all had the same pivot, meaning that regardless of the relative positions of all of the source elements, they are all placed at the same position in the resulting pattern.

Jittering happens on a per pattern basis, meaning that the relative positions of all source elements are retained after jittering. Scaling and rotating are performed about the pivot of the first pattern element.

If Automatically Adjust is enabled, then the bounding box used for determining spacing in packed mode will be recomputed by baking the relative positions, start scale, and start rotations into the individual elements' bounding boxes and then computing a "CombinedPatternBounds" which contains all of the pattern elements and is used for determining spacing in packed mode.

Use Relative Transforms and Automatically Adjust (In Bounding Box dropdown when in Packed mode) allow all new functionality to be toggled.

diff --git a/Engine/Plugins/Experimental/MeshModelingToolsetExp/Source/MeshModelingToolsExp/Private/PatternTool.cpp b/Engine/Plugins/Experimental/MeshModelingToolsetExp/Source/MeshModelingToolsExp/Private/PatternTool.cpp
--- a/Engine/Plugins/Experimental/MeshModelingToolsetExp/Source/MeshModelingToolsExp/Private/PatternTool.cpp
+++ b/Engine/Plugins/Experimental/MeshModelingToolsetExp/Source/MeshModelingToolsExp/Private/PatternTool.cpp
@@ -532,27 +532,39 @@ void UPatternTool::Setup()
 	Settings = NewObject<UPatternToolSettings>();
 	AddToolPropertySource(Settings);
 	Settings->RestoreProperties(this);
-	Settings->WatchProperty(Settings->Shape, [this](EPatternToolShape) { OnShapeUpdated(); } );
-	Settings->WatchProperty(Settings->SingleAxis, [this](EPatternToolSingleAxis) { OnParametersUpdated(); } );
-	Settings->WatchProperty(Settings->SinglePlane, [this](EPatternToolSinglePlane) { OnParametersUpdated(); } );
-	Settings->WatchProperty(Settings->bHideSources, [this](bool bNewValue) { OnSourceVisibilityToggled(!bNewValue); } );
 	Settings->WatchProperty(Settings->Seed, [this](int32 NewSeed) { MarkPatternDirty(); } );
 	Settings->WatchProperty(Settings->bProjectElementsDown, [this](bool bNewValue) { MarkPatternDirty(); } );
 	Settings->WatchProperty(Settings->ProjectionOffset, [this](float NewValue) { MarkPatternDirty(); } );
+	Settings->WatchProperty(Settings->bHideSources, [this](bool bNewValue) { OnSourceVisibilityToggled(!bNewValue); } );
+	Settings->WatchProperty(Settings->bUseRelativeTransforms, [this](bool bNewValue) { MarkPatternDirty(); } );
+	Settings->WatchProperty(Settings->Shape, [this](EPatternToolShape) { OnShapeUpdated(); } );
+	Settings->WatchProperty(Settings->SingleAxis, [this](EPatternToolSingleAxis) { OnParametersUpdated(); } );
+	Settings->WatchProperty(Settings->SinglePlane, [this](EPatternToolSinglePlane) { OnParametersUpdated(); } );
 
+	BoundingBoxSettings = NewObject<UPatternTool_BoundingBoxSettings>();
+	AddToolPropertySource(BoundingBoxSettings);
+	BoundingBoxSettings->RestoreProperties(this);
+	BoundingBoxSettings->WatchProperty(BoundingBoxSettings->bIgnoreTransforms, [this](bool bNewValue) { MarkPatternDirty(); } );
+	BoundingBoxSettings->WatchProperty(BoundingBoxSettings->Adjustment, [this](float NewScale) { MarkPatternDirty(); });
+	SetToolPropertySourceEnabled(BoundingBoxSettings, false);
+	
 	LinearSettings = NewObject<UPatternTool_LinearSettings>();
 	AddToolPropertySource(LinearSettings);
 	LinearSettings->RestoreProperties(this);
+	LinearSettings->WatchProperty(LinearSettings->SpacingMode, [this](EPatternToolAxisSpacingMode NewSpacingMode) { OnSpacingModeUpdated(); });
 	SetToolPropertySourceEnabled(LinearSettings, false);
 
 	GridSettings = NewObject<UPatternTool_GridSettings>();
 	AddToolPropertySource(GridSettings);
 	GridSettings->RestoreProperties(this);
+	GridSettings->WatchProperty(GridSettings->SpacingX, [this](EPatternToolAxisSpacingMode NewSpacingMode) { OnSpacingModeUpdated(); });
+	GridSettings->WatchProperty(GridSettings->SpacingY, [this](EPatternToolAxisSpacingMode NewSpacingMode) { OnSpacingModeUpdated(); });
 	SetToolPropertySourceEnabled(GridSettings, false);
 
 	RadialSettings = NewObject<UPatternTool_RadialSettings>();
 	AddToolPropertySource(RadialSettings);
 	RadialSettings->RestoreProperties(this);
+	RadialSettings->WatchProperty(RadialSettings->SpacingMode, [this](EPatternToolAxisSpacingMode NewSpacingMode) { OnSpacingModeUpdated(); });
 	SetToolPropertySourceEnabled(RadialSettings, false);
 
 	RotationSettings = NewObject<UPatternTool_RotationSettings>();
@@ -637,6 +649,8 @@ void UPatternTool::Setup()
 	AddToolPropertySource(OutputSettings);
 	OutputSettings->RestoreProperties(this);
 
+	BoundingBoxVisualizer.LineThickness = 2.0f;
+	
 	InitializeElements();
 	for (const FPatternElement& Element : Elements)
 	{
@@ -695,6 +709,7 @@ void UPatternTool::OnShutdown(EToolShutdownType ShutdownType)
 	PreviewGeometry->Disconnect();
 
 	Settings->SaveProperties(this);
+	BoundingBoxSettings->SaveProperties(this);
 	LinearSettings->SaveProperties(this);
 	GridSettings->SaveProperties(this);
 	RadialSettings->SaveProperties(this);
@@ -723,6 +738,14 @@ void UPatternTool::Render(IToolsContextRenderAPI* RenderAPI)
 	DragAlignmentMechanic->Render(RenderAPI);
 	PlaneMechanic->Render(RenderAPI);
 
+	// We should only render here if bVisualize is true and also visible in the details panel.
+	// We don't want to leave the user with no obvious way to turn it off which might happen if they
+	// enable bVisualize and then change the spacing mode to something other than packed
+	if (BoundingBoxSettings->IsPropertySetEnabled() && BoundingBoxSettings->bVisualize)
+	{
+		RenderBoundingBoxes(RenderAPI);
+	}
+	
 	if (bPatternNeedsUpdating)
 	{
 		// throttle update rate as it is somewhat expensive
@@ -736,7 +759,46 @@ void UPatternTool::Render(IToolsContextRenderAPI* RenderAPI)
 	}
 }
 
+void UPatternTool::RenderBoundingBoxes(IToolsContextRenderAPI* RenderAPI)
+{
+	BoundingBoxVisualizer.BeginFrame(RenderAPI);
+
+	// Render the individual elements' PatternBounds along the current pattern with green lines
+	BoundingBoxVisualizer.LineColor = FLinearColor::Green;
+	for (int32 ElemIdx = 0; ElemIdx < Elements.Num(); ++ElemIdx)
+	{
+		FPatternElement& Element = Elements[ElemIdx];
+		
+		if (Settings->bUseRelativeTransforms)
+		{
+			const FTransformSRT3d RelativePositionTransform(FQuaterniond(RotationSettings->StartRotation), FVector3d::ZeroVector, ScaleSettings->StartScale);
+			BoundingBoxVisualizer.PushTransform(FTransform(RelativePositionTransform.TransformVector(Element.RelativePosition)));
+		}
+	
+		for (int32 k = 0; k < CurrentPattern.Num(); ++k)
+		{
+			BoundingBoxVisualizer.PushTransform(FTransform(CurrentPattern[k].GetTranslation()) * CurrentStartFrameWorld.ToFTransform());
+			BoundingBoxVisualizer.DrawWireBox(FBox(Element.PatternBounds));
+			BoundingBoxVisualizer.PopTransform();
+		}
+
+		if (Settings->bUseRelativeTransforms)
+		{
+			BoundingBoxVisualizer.PopTransform();
+		}
+	}
 
+	// Render the CombinedPatternBounds along the current pattern with red lines
+	BoundingBoxVisualizer.LineColor = FLinearColor::Red;
+	for (int32 k = 0; k < CurrentPattern.Num(); ++k)
+	{
+		BoundingBoxVisualizer.PushTransform(FTransform(CurrentPattern[k].GetTranslation()) * CurrentStartFrameWorld.ToFTransform());
+		BoundingBoxVisualizer.DrawWireBox(FBox(CombinedPatternBounds));
+		BoundingBoxVisualizer.PopTransform();
+	}
+	
+	BoundingBoxVisualizer.EndFrame();
+}
 
 void UPatternTool::OnSourceVisibilityToggled(bool bVisible)
 {
@@ -772,6 +834,7 @@ void UPatternTool::InitializeElements()
 		Element.SourceComponent = UE::ToolTarget::GetTargetComponent(Targets[TargetIdx]);
 		Element.SourceMaterials = UE::ToolTarget::GetMaterialSet(Targets[TargetIdx], false).Materials;
 		Element.SourceTransform = UE::ToolTarget::GetLocalToWorldTransform(Targets[TargetIdx]);
+		Element.RelativePosition = Element.SourceTransform.GetTranslation() - Elements[0].SourceTransform.GetTranslation();
 		Element.BaseRotateScale = Element.SourceTransform;
 		Element.BaseRotateScale.SetTranslation(FVector3d::Zero());		// clear translation from base transform
 		Element.SourceTransform.SetRotation(FQuaterniond::Identity());	// clear rotate/scale from source transform, so only location is used
@@ -782,24 +845,26 @@ void UPatternTool::InitializeElements()
 		if (UStaticMeshComponent* StaticMeshComp = Cast<UStaticMeshComponent>(Element.SourceComponent))
 		{
 			Element.SourceStaticMesh = StaticMeshComp->GetStaticMesh();
-			Element.LocalBounds = (FAxisAlignedBox3d)Element.SourceStaticMesh->GetBounds().GetBox();
-			Element.PatternBounds = Element.LocalBounds;
-		}
+			Element.LocalBounds = Element.SourceStaticMesh->GetBounds().GetBox();
+		}	
 		else if (UDynamicMeshComponent* DynamicMeshComp = Cast<UDynamicMeshComponent>(Element.SourceComponent))
 		{
 			Element.SourceDynamicMesh = DynamicMeshComp->GetDynamicMesh();
 			Element.SourceDynamicMesh->ProcessMesh([&](const FDynamicMesh3& Mesh) {
 				Element.LocalBounds = Mesh.GetBounds(true);
 			});
-			Element.PatternBounds = Element.LocalBounds;
 		}
 		else
 		{
 			Element.bValid = false;
 		}
+		
+		Element.PatternBounds = Element.LocalBounds;
 	}
 
 	PreviewComponents.SetNum(NumElements);
+
+	ComputeCombinedPatternBounds();
 }
 
 
@@ -855,6 +920,36 @@ void UPatternTool::OnShapeUpdated()
 	OnParametersUpdated();
 }
 
+void UPatternTool::OnSpacingModeUpdated()
+{
+	bool bBoundingBoxSettings = false;
+	
+	if (Settings->Shape == EPatternToolShape::Line)
+	{
+		if (LinearSettings->SpacingMode == EPatternToolAxisSpacingMode::Packed)
+		{
+			bBoundingBoxSettings = true;
+		}
+	}
+	else if (Settings->Shape == EPatternToolShape::Grid)
+	{
+		if (GridSettings->SpacingX == EPatternToolAxisSpacingMode::Packed || GridSettings->SpacingY == EPatternToolAxisSpacingMode::Packed)
+		{
+			bBoundingBoxSettings = true;
+		}
+	}
+	else if (Settings->Shape == EPatternToolShape::Circle)
+	{
+		if (RadialSettings->SpacingMode == EPatternToolAxisSpacingMode::Packed)
+		{
+			bBoundingBoxSettings = true;
+		}
+	}
+
+	SetToolPropertySourceEnabled(BoundingBoxSettings, bBoundingBoxSettings);
+
+	OnParametersUpdated();
+}
 
 void UPatternTool::OnParametersUpdated()
 {
@@ -946,8 +1041,10 @@ void UPatternTool::GetPatternTransforms_Linear(TArray<UE::Geometry::FTransformSR
 {
 	FLinearPatternGenerator Generator;
 	InitializeGenerator(Generator, this);
-	Generator.Dimensions = this->Elements[0].PatternBounds;
-
+	
+	ComputeCombinedPatternBounds();
+	Generator.Dimensions = CombinedPatternBounds;
+	
 	double ExtentX = LinearSettings->Extent;
 
 	Generator.StartFrame = FFrame3d();
@@ -978,8 +1075,10 @@ void UPatternTool::GetPatternTransforms_Grid(TArray<UE::Geometry::FTransformSRT3
 {
 	FLinearPatternGenerator Generator;
 	InitializeGenerator(Generator, this);
-	Generator.Dimensions = this->Elements[0].PatternBounds;
-
+	
+	ComputeCombinedPatternBounds();
+	Generator.Dimensions = CombinedPatternBounds;
+	
 	double ExtentX = GridSettings->ExtentX;
 	double ExtentY = GridSettings->ExtentY;
 
@@ -1029,8 +1128,10 @@ void UPatternTool::GetPatternTransforms_Radial(TArray<UE::Geometry::FTransformSR
 {
 	FRadialPatternGenerator Generator;
 	InitializeGenerator(Generator, this);
-	Generator.Dimensions = this->Elements[0].PatternBounds;
 
+	ComputeCombinedPatternBounds();
+	Generator.Dimensions = CombinedPatternBounds;
+	
 	// Orient the CenterFrame based on the plane setting of the tool & determine which axis is used if bOriented is true
 	const FVector3d Origin(0, 0, 0);
 	switch (Settings->SinglePlane)
@@ -1100,6 +1201,11 @@ void UPatternTool::UpdatePattern()
 		FTransformSRT3d ElementTransform = Element.BaseRotateScale;
 		FComponentSet& ElemComponents = PreviewComponents[ElemIdx];
 
+		if (Settings->bUseRelativeTransforms)
+		{
+			ElementTransform.SetTranslation(ElementTransform.GetTranslation() + Element.RelativePosition);
+		}
+		
 		for (int32 k = 0; k < NumPatternItems; ++k)
 		{
 			UPrimitiveComponent* Component = nullptr;
@@ -1152,6 +1258,49 @@ void UPatternTool::ComputeWorldTransform(FTransform& OutWorldTransform, const FT
 	}
 }
 
+void UPatternTool::ComputePatternBounds(int32 ElemIdx)
+{
+	FPatternElement& Element = Elements[ElemIdx];
+	FTransformSRT3d Transform = Element.BaseRotateScale;
+	
+	if (!BoundingBoxSettings->bIgnoreTransforms)
+	{
+		Transform = (FTransform) FTransformSRT3d(FQuaterniond(RotationSettings->StartRotation), FVector3d::ZeroVector, ScaleSettings->StartScale) * Transform;
+	}
+
+	Element.PatternBounds = FAxisAlignedBox3d(Element.LocalBounds, Transform);
+}
+
+void UPatternTool::ComputeCombinedPatternBounds()
+{
+	// CombinedPatternBounds is the bounding box that contains every pattern element and is computed by setting it
+	// to the first element's pattern bounds and then expanding the box to contain each following element.
+
+	ComputePatternBounds(0);
+	CombinedPatternBounds = Elements[0].PatternBounds;
+
+	// If StartScale or StartRotation are not zero vectors, then Element.RelativePosition must be transformed to be accurate if bUseRelativeTransforms is true
+	const FTransformSRT3d RelativePositionTransform(FQuaterniond(RotationSettings->StartRotation), FVector3d::ZeroVector, ScaleSettings->StartScale);
+	
+	for (int32 ElemIdx = 1; ElemIdx < Elements.Num(); ++ElemIdx)
+	{
+		const FPatternElement& Element = Elements[ElemIdx];
+		ComputePatternBounds(ElemIdx);
+		
+		if (Settings->bUseRelativeTransforms)
+		{
+			CombinedPatternBounds.Contain(Element.PatternBounds.Min + RelativePositionTransform.TransformVector(Element.RelativePosition));
+			CombinedPatternBounds.Contain(Element.PatternBounds.Max + RelativePositionTransform.TransformVector(Element.RelativePosition));
+		}
+		else
+		{
+			CombinedPatternBounds.Contain(Element.PatternBounds);
+		}
+	}
+
+	// The user can manually adjust the box if desired to fine tune packed behavior
+	CombinedPatternBounds.Expand(BoundingBoxSettings->Adjustment);
+}
 
 UStaticMeshComponent* UPatternTool::GetPreviewStaticMesh(FPatternElement& Element)
 {
@@ -1337,6 +1486,11 @@ void UPatternTool::EmitResults()
 		}
 		FTransformSRT3d ElementTransform = Element.BaseRotateScale;
 
+		if (Settings->bUseRelativeTransforms)
+		{
+			ElementTransform.SetTranslation(ElementTransform.GetTranslation() + Element.RelativePosition);
+		}
+		
 		if (Element.SourceDynamicMesh != nullptr || bConvertToDynamic )
 		{
 			FDynamicMesh3 ElementMesh = UE::ToolTarget::GetDynamicMeshCopy(Targets[ElemIdx], true);
diff --git a/Engine/Plugins/Experimental/MeshModelingToolsetExp/Source/MeshModelingToolsExp/Public/PatternTool.h b/Engine/Plugins/Experimental/MeshModelingToolsetExp/Source/MeshModelingToolsExp/Public/PatternTool.h
--- a/Engine/Plugins/Experimental/MeshModelingToolsetExp/Source/MeshModelingToolsExp/Public/PatternTool.h
+++ b/Engine/Plugins/Experimental/MeshModelingToolsetExp/Source/MeshModelingToolsExp/Public/PatternTool.h
@@ -9,6 +9,7 @@
 #include "Changes/TransformChange.h"
 #include "FrameTypes.h"
 #include "BoxTypes.h"
+#include "ToolDataVisualizer.h"
 #include "TransformTypes.h"
 #include "PatternTool.generated.h"
 
@@ -102,6 +103,26 @@ class MESHMODELINGTOOLSEXP_API UPatternToolSettings : public UInteractiveToolPro
 {
 	GENERATED_BODY()
 public:
+	/** The seed used to introduce random transform variations when enabled */
+	UPROPERTY(EditAnywhere, Category = General, meta = (NoResetToDefault))
+	int32 Seed = FMath::Rand();
+
+	/** Whether or not the pattern items should be projected along the negative Z axis of the plane mechanic */
+	UPROPERTY(EditAnywhere, Category = General)
+	bool bProjectElementsDown = false;
+
+	/** How much each pattern item should be moved along the negative Z axis of the plane mechanic if Project Elements Down is enabled */
+	UPROPERTY(EditAnywhere, Category = General, meta = (EditCondition = "bProjectElementsDown == true", EditConditionHides))
+	float ProjectionOffset = 0.0f;
+
+	/** Hide the source meshes when enabled */
+	UPROPERTY(EditAnywhere, Category = General)
+	bool bHideSources = true;
+
+	/** If false, all pattern elements will be positioned at the origin of the first pattern element */
+	UPROPERTY(EditAnywhere, Category = General)
+	bool bUseRelativeTransforms = true;
+	
 	/** Shape of the underlying Pattern */
 	UPROPERTY(EditAnywhere, Category = Shape)
 	EPatternToolShape Shape = EPatternToolShape::Line;
@@ -113,25 +134,29 @@ public:
 	/** Plane used for the Pattern geometry */
 	UPROPERTY(EditAnywhere, Category = Shape, meta = (DisplayName = "Plane", EditCondition = "Shape != EPatternToolShape::Line", EditConditionHides))
 	EPatternToolSinglePlane SinglePlane = EPatternToolSinglePlane::XYPlane;
-
-	/** Hide the source meshes when enabled */
-	UPROPERTY(EditAnywhere, Category = Shape)
-	bool bHideSources = true;
-
-	/** The seed used to introduce random transform variations when enabled */
-	UPROPERTY(EditAnywhere, Category = Shape, meta = (NoResetToDefault))
-	int32 Seed = FMath::Rand();
-
-	/** Whether or not the pattern items should be projected along the negative Z axis of the plane mechanic */
-	UPROPERTY(EditAnywhere, Category = Shape)
-	bool bProjectElementsDown = false;
-
-	/** How much each pattern item should be moved along the negative Z axis of the plane mechanic if Project Elements Down is enabled */
-	UPROPERTY(EditAnywhere, Category = Shape, meta = (EditCondition = "bProjectElementsDown == true", EditConditionHides))
-	float ProjectionOffset = 0.0f;
 };
 
+/**
+ * Settings for Bounding Box adjustments in the Pattern Tool
+ */
+UCLASS()
+class MESHMODELINGTOOLSEXP_API UPatternTool_BoundingBoxSettings : public UInteractiveToolPropertySet
+{
+	GENERATED_BODY()
+public:
+	/** If true, pattern element bounding boxes are not changed to account for StartScale or StartRotation */
+	UPROPERTY(EditAnywhere, Category = BoundingBox)
+	bool bIgnoreTransforms = false;
+	
+	/** Value added to the all pattern elements' bounding boxes for adjusting the behavior of packed spacing mode manually */
+	UPROPERTY(EditAnywhere, Category = BoundingBox)
+	float Adjustment = 0.0f;
 
+	/** If true, the bounding boxes of each element are rendered in green and the combined bounding box of all source elements is rendered in red */
+	UPROPERTY(EditAnywhere, Category = BoundingBox)
+	bool bVisualize = false;
+};
+	
 /**
  * Settings for Linear Patterns in the Pattern Tool
  */
@@ -145,11 +170,11 @@ public:
 	EPatternToolAxisSpacingMode SpacingMode = EPatternToolAxisSpacingMode::ByCount;
 
 	/** Number of Pattern Elements to place */
-	UPROPERTY(EditAnywhere, Category = LinearPattern, meta = (ClampMin = 1, UIMax = 25, EditCondition = "SpacingMode == EPatternToolAxisSpacingMode::ByCount"))
+	UPROPERTY(EditAnywhere, Category = LinearPattern, meta = (ClampMin = 1, UIMax = 25, EditCondition = "SpacingMode == EPatternToolAxisSpacingMode::ByCount", EditConditionHides))
 	int32 Count = 10;
 
 	/** Fixed Increment used to place Pattern Elements */
-	UPROPERTY(EditAnywhere, Category = LinearPattern, meta = (ClampMin = 0, EditCondition = "SpacingMode == EPatternToolAxisSpacingMode::StepSize"))
+	UPROPERTY(EditAnywhere, Category = LinearPattern, meta = (ClampMin = 0, EditCondition = "SpacingMode == EPatternToolAxisSpacingMode::StepSize", EditConditionHides))
 	double StepSize = 100.0;
 
 	/** Length of Pattern along the Axis */
@@ -177,12 +202,12 @@ public:
 	UPROPERTY(EditAnywhere, Category = GridPatternX)
 	EPatternToolAxisSpacingMode SpacingX = EPatternToolAxisSpacingMode::ByCount;
 
-	/** Number of  Pattern Elements to place along the Main axis */
-	UPROPERTY(EditAnywhere, Category = GridPatternX, meta = (ClampMin = 1, UIMax = 25, EditCondition = "SpacingX == EPatternToolAxisSpacingMode::ByCount"))
+	/** Number of Pattern Elements to place along the Main axis */
+	UPROPERTY(EditAnywhere, Category = GridPatternX, meta = (ClampMin = 1, UIMax = 25, EditCondition = "SpacingX == EPatternToolAxisSpacingMode::ByCount", EditConditionHides))
 	int32 CountX = 10;
 
 	/** Fixed Increment used to place Pattern Elements along the Main axis */
-	UPROPERTY(EditAnywhere, Category = GridPatternX, meta = (ClampMin = 0, EditCondition = "SpacingX == EPatternToolAxisSpacingMode::StepSize"))
+	UPROPERTY(EditAnywhere, Category = GridPatternX, meta = (ClampMin = 0, EditCondition = "SpacingX == EPatternToolAxisSpacingMode::StepSize", EditConditionHides))
 	double StepSizeX = 100.0;
 
 	/** Length/Extent of Pattern falong the Main Axis */
@@ -198,11 +223,11 @@ public:
 	EPatternToolAxisSpacingMode SpacingY = EPatternToolAxisSpacingMode::ByCount;
 
 	/** Number of  Pattern Elements to place along the Secondary axis */
-	UPROPERTY(EditAnywhere, Category = GridPatternY, meta = (ClampMin = 1, UIMax = 25, EditCondition = "SpacingY == EPatternToolAxisSpacingMode::ByCount"))
+	UPROPERTY(EditAnywhere, Category = GridPatternY, meta = (ClampMin = 1, UIMax = 25, EditCondition = "SpacingY == EPatternToolAxisSpacingMode::ByCount", EditConditionHides))
 	int32 CountY = 10;
 
 	/** Fixed Increment used to place Pattern Elements along the Secondary axis */
-	UPROPERTY(EditAnywhere, Category = GridPatternY, meta = (ClampMin = 0, EditCondition = "SpacingY == EPatternToolAxisSpacingMode::StepSize"))
+	UPROPERTY(EditAnywhere, Category = GridPatternY, meta = (ClampMin = 0, EditCondition = "SpacingY == EPatternToolAxisSpacingMode::StepSize", EditConditionHides))
 	double StepSizeY = 100.0;
 
 	/** Length/Extent of Pattern falong the Secondary Axis */
@@ -231,11 +256,11 @@ public:
 	EPatternToolAxisSpacingMode SpacingMode = EPatternToolAxisSpacingMode::ByCount;
 
 	/** Number of  Pattern Elements to place */
-	UPROPERTY(EditAnywhere, Category = RadialPattern, meta = (ClampMin = 1, UIMax = 25, EditCondition = "SpacingMode == EPatternToolAxisSpacingMode::ByCount"))
+	UPROPERTY(EditAnywhere, Category = RadialPattern, meta = (ClampMin = 1, UIMax = 25, EditCondition = "SpacingMode == EPatternToolAxisSpacingMode::ByCount", EditConditionHides))
 	int32 Count = 10;
 
 	/** Fixed Increment (in Degrees) used to position Pattern Elements around the Circle/Arc */
-	UPROPERTY(EditAnywhere, Category = RadialPattern, meta = (Units = "Degrees", ClampMin = 0, EditCondition = "SpacingMode == EPatternToolAxisSpacingMode::StepSize"))
+	UPROPERTY(EditAnywhere, Category = RadialPattern, meta = (Units = "Degrees", ClampMin = 0, EditCondition = "SpacingMode == EPatternToolAxisSpacingMode::StepSize", EditConditionHides))
 	double StepSize = 100.0;
 
 	/** Radius of the Circle/Arc */
@@ -416,6 +441,9 @@ public:
 	UPROPERTY()
 	TObjectPtr<UPatternToolSettings> Settings;
 
+	UPROPERTY()
+	TObjectPtr<UPatternTool_BoundingBoxSettings> BoundingBoxSettings;
+	
 	UPROPERTY()
 	TObjectPtr<UPatternTool_LinearSettings> LinearSettings;
 
@@ -471,6 +499,7 @@ protected:
 	void OnSourceVisibilityToggled(bool bVisible);
 	void OnMainFrameUpdated();
 	void OnShapeUpdated();
+	void OnSpacingModeUpdated();
 	void OnParametersUpdated();
 	void UpdatePattern();
 	void ComputeWorldTransform(FTransform& OutWorldTransform, const FTransform& InElementTransform, const FTransform& InPatternTransform) const;
@@ -478,6 +507,8 @@ protected:
 	void GetPatternTransforms_Grid(TArray<UE::Geometry::FTransformSRT3d>& TransformsOut);
 	void GetPatternTransforms_Radial(TArray<UE::Geometry::FTransformSRT3d>& TransformsOut);
 
+	void RenderBoundingBoxes(IToolsContextRenderAPI* RenderAPI);
+	FToolDataVisualizer BoundingBoxVisualizer;
 	
 	struct FPatternElement
 	{
@@ -491,16 +522,29 @@ protected:
 		UE::Geometry::FTransformSRT3d SourceTransform = UE::Geometry::FTransformSRT3d::Identity();
 		UE::Geometry::FTransformSRT3d BaseRotateScale = UE::Geometry::FTransformSRT3d::Identity();
 
+		// We don't need to store rotation or scale relative to first
+		// element because that is handled by BaseRotateScale
+		FVector3d RelativePosition = FVector3d::ZeroVector;
+
 		UDynamicMesh* SourceDynamicMesh = nullptr;
 		UStaticMesh* SourceStaticMesh = nullptr;
 
-		UE::Geometry::FAxisAlignedBox3d LocalBounds = UE::Geometry::FAxisAlignedBox3d::Empty();
-		UE::Geometry::FAxisAlignedBox3d PatternBounds = UE::Geometry::FAxisAlignedBox3d::Empty();
+		UE::Geometry::FAxisAlignedBox3d LocalBounds = UE::Geometry::FAxisAlignedBox3d::Empty();			// The unchanged bounding box of the source mesh. Only used indirectly through PatternBounds
+		UE::Geometry::FAxisAlignedBox3d PatternBounds = UE::Geometry::FAxisAlignedBox3d::Empty();		// The bounding box used for computing CombinedPatternBounds for packed spacing. By default this box adapts to StartScale/StartRotation
 	};
 	TArray<FPatternElement> Elements;
+	UE::Geometry::FAxisAlignedBox3d CombinedPatternBounds = UE::Geometry::FAxisAlignedBox3d::Empty();	// The bounding box that contains all of the elements' PatternBounds bounding boxes. Used for packed mode spacing
+
+	void ComputePatternBounds(int32 ElemIdx);
+	void ComputeCombinedPatternBounds();
 
+	// Given an Element index and an FTransformSRT3d, determine the bounding box that contains the transformed underlying mesh
+	// BoundingBox is made empty before growing to contain the transformed mesh.
+	void ComputeBoundingBoxWithTransform(int32 ElemIdx, UE::Geometry::FAxisAlignedBox3d& BoundingBox, const UE::Geometry::FTransformSRT3d& Transform);
+	
 	bool bHaveNonUniformScaleElements = false;
 
+	
 	struct FComponentSet
 	{
 		TArray<UPrimitiveComponent*> Components;
