MeshModelingTools: Don't consider open boundaries in PolyCut tool to be an unacceptable failure, though write out a warning. Don't try to fix holes by default, to be like MshBool and to avoid closing holes when cutting an open mesh.

diff --git a/Engine/Plugins/Experimental/MeshModelingToolsetExp/Source/MeshModelingToolsEditorOnlyExp/Private/PolygonOnMeshTool.cpp b/Engine/Plugins/Experimental/MeshModelingToolsetExp/Source/MeshModelingToolsEditorOnlyExp/Private/PolygonOnMeshTool.cpp
--- a/Engine/Plugins/Experimental/MeshModelingToolsetExp/Source/MeshModelingToolsEditorOnlyExp/Private/PolygonOnMeshTool.cpp
+++ b/Engine/Plugins/Experimental/MeshModelingToolsetExp/Source/MeshModelingToolsEditorOnlyExp/Private/PolygonOnMeshTool.cpp
@@ -117,6 +117,11 @@ void UPolygonOnMeshTool::Setup()
 				GetToolManager()->DisplayMessage(LOCTEXT("FailNotification", "Unable to complete cut."),
 					EToolMessageLevel::UserWarning);
 			}
+			else if (EdgesOnFailure.Num() > 0)
+			{
+				GetToolManager()->DisplayMessage(LOCTEXT("BoundaryEdgesNotification", "Cut resulted in one or more boundary edges."),
+					EToolMessageLevel::UserWarning);
+			}
 			else
 			{
 				// This clears the warning if needed
@@ -172,20 +177,20 @@ void UPolygonOnMeshTool::UpdateVisualization()
 		TargetMesh->GetEdgeV(EID, A, B);
 		DrawnLineSet->AddLine((FVector)A, (FVector)B, EmbedEdgeColor, EmbedEdgeThickness, EmbedEdgeDepthBias);
 	}
-	if (!bOperationSucceeded)
+	// We show problematic edges whether or not we consider the operation to be successful because in the case of
+	// booleans, we consider boundary edge results to be important to show but not important enough to be considered
+	// a failure (and in fact are not a failure if you're cutting an open mesh, like a simple rectangle)
+	for (int EID : EdgesOnFailure)
 	{
-		for (int EID : EdgesOnFailure)
-		{
-			TargetMesh->GetEdgeV(EID, A, B);
-			DrawnLineSet->AddLine((FVector)A, (FVector)B, PartialPathEdgeColor, PartialPathEdgeThickness, PartialPathEdgeDepthBias);
-		}
-		// In the case where we don't allow failed results, it can be disorienting to show the broken mesh (especially since sometimes
-		// most of it may be cut away). But user can change this behavior.
-		if (!BasicProperties->bCanAcceptFailedResult && !BasicProperties->bShowIntermediateResultOnFailure)
-		{
-			// Reset the preview.
-			Preview->PreviewMesh->UpdatePreview(OriginalDynamicMesh.Get());
-		}
+		TargetMesh->GetEdgeV(EID, A, B);
+		DrawnLineSet->AddLine((FVector)A, (FVector)B, PartialPathEdgeColor, PartialPathEdgeThickness, PartialPathEdgeDepthBias);
+	}
+	// In the case where we don't allow failed results, it can be disorienting to show the broken mesh (especially since sometimes
+	// most of it may be cut away). But user can change this behavior.
+	if (!BasicProperties->bCanAcceptFailedResult && !BasicProperties->bShowIntermediateResultOnFailure)
+	{
+		// Reset the preview.
+		Preview->PreviewMesh->UpdatePreview(OriginalDynamicMesh.Get());
 	}
 }
 
@@ -276,7 +281,7 @@ TUniquePtr<FDynamicMeshOperator> UPolygonOnMeshTool::MakeNewOperator()
 	bool bOpLeavesOpenBoundaries =
 		BasicProperties->Operation == EEmbeddedPolygonOpMethod::TrimInside ||
 		BasicProperties->Operation == EEmbeddedPolygonOpMethod::TrimOutside;
-	EmbedOp->bAttemptFixHolesOnBoolean = !bOpLeavesOpenBoundaries && BasicProperties->bTryToFixCracks;
+	EmbedOp->bAttemptFixHolesOnBoolean = !bOpLeavesOpenBoundaries && BasicProperties->bTryToFixHoles;
 
 	// Match the world plane in the local space
 	FVector3d LocalOrigin = WorldTransform.InverseTransformPosition(DrawPlaneWorld.Origin);
diff --git a/Engine/Plugins/Experimental/MeshModelingToolsetExp/Source/MeshModelingToolsEditorOnlyExp/Public/PolygonOnMeshTool.h b/Engine/Plugins/Experimental/MeshModelingToolsetExp/Source/MeshModelingToolsEditorOnlyExp/Public/PolygonOnMeshTool.h
--- a/Engine/Plugins/Experimental/MeshModelingToolsetExp/Source/MeshModelingToolsEditorOnlyExp/Public/PolygonOnMeshTool.h
+++ b/Engine/Plugins/Experimental/MeshModelingToolsetExp/Source/MeshModelingToolsEditorOnlyExp/Public/PolygonOnMeshTool.h
@@ -77,7 +77,7 @@ public:
 
 	/** Automatically attempt to fill any open boundaries left by CSG (e.g. due to numerical errors) */
 	UPROPERTY(EditAnywhere, Category = Operation, meta = (EditCondition = "bCutWithBoolean && Operation == EEmbeddedPolygonOpMethod::CutThrough || bCutWithBoolean && Operation == EEmbeddedPolygonOpMethod::InsertPolygon", EditConditionHides))
-	bool bTryToFixCracks = true;
+	bool bTryToFixHoles = false;
 
 	// TODO: re-add if/when extrude is added as a supported operation
 	///** Amount to extrude, if extrude is enabled */
diff --git a/Engine/Plugins/Runtime/MeshModelingToolset/Source/ModelingOperatorsEditorOnly/Private/CuttingOps/EmbedPolygonsOp.cpp b/Engine/Plugins/Runtime/MeshModelingToolset/Source/ModelingOperatorsEditorOnly/Private/CuttingOps/EmbedPolygonsOp.cpp
--- a/Engine/Plugins/Runtime/MeshModelingToolset/Source/ModelingOperatorsEditorOnly/Private/CuttingOps/EmbedPolygonsOp.cpp
+++ b/Engine/Plugins/Runtime/MeshModelingToolset/Source/ModelingOperatorsEditorOnly/Private/CuttingOps/EmbedPolygonsOp.cpp
@@ -180,7 +180,15 @@ void FEmbedPolygonsOp::BooleanPath(FProgressCancel* Progress)
 	Boolean.Progress = Progress;
 	Boolean.bSimplifyAlongNewEdges = true;
 	Boolean.bPutResultInInputSpace = true;
-	bool bBoolSuccess = Boolean.Compute();
+	Boolean.Compute();
+
+	// One might think that we would want to set this to the return value of Boolean.Compute() above,
+	// but that call currently returns false whenever any boundary edges are generated. That type of
+	// "failure" is frequently ok for the user, for instance when cutting an open mesh, like a simple rectangle. 
+	// We should probably change FMeshBoolean to differentiate between types of errors so that we can still
+	// respond appropriately to other types of failures. For now, this is what we do here (and to some
+	// extent, in BooleanMeshesOp.cpp).
+	bOperationSucceeded = true;
 
 	if (Progress && Progress->Cancelled())
 	{
@@ -202,10 +210,32 @@ void FEmbedPolygonsOp::BooleanPath(FProgressCancel* Progress)
 			return;
 		}
 
+		// This hole filling code should probably be put into some common utility function since it's used in
+		// BooleanMeshesOp.h, and we do similar things (with more options) for cap filling in revolves.
+		const bool bNeedsNormalsAndUVs = Boolean.Result->HasAttributes();
+		// UV scale based on the mesh bounds, used for hole-filled regions.
+		const double MeshUVScaleFactor = (1.0 / Boolean.Result->GetBounds().MaxDim());
+
 		for (FEdgeLoop& Loop : OpenBoundary.Loops)
 		{
 			FMinimalHoleFiller Filler(Boolean.Result, Loop);
 			Filler.Fill(ResultMesh->AllocateTriangleGroup());
+
+			if (bNeedsNormalsAndUVs)
+			{
+				// Compute a best-fit plane of the boundary vertices
+				TArray<FVector3d> VertexPositions;
+				Loop.GetVertices(VertexPositions);
+
+				FVector3d PlaneOrigin, PlaneNormal;
+				PolygonTriangulation::ComputePolygonPlane<double>(VertexPositions, PlaneNormal, PlaneOrigin);
+				PlaneNormal *= -1.0;	// Previous function seems to orient the normal opposite to what's expected elsewhere
+
+				FDynamicMeshEditor Editor(ResultMesh.Get());
+				FFrame3d ProjectionFrame(PlaneOrigin, PlaneNormal);
+				Editor.SetTriangleNormals(Filler.NewTriangles);
+				Editor.SetTriangleUVsFromProjection(Filler.NewTriangles, ProjectionFrame, MeshUVScaleFactor);
+			}
 		}
 		for (int EID : Boolean.CreatedBoundaryEdges)
 		{
@@ -214,13 +244,10 @@ void FEmbedPolygonsOp::BooleanPath(FProgressCancel* Progress)
 				EdgesOnFailure.Add(EID);
 			}
 		}
-
-		bOperationSucceeded = EdgesOnFailure.Num() == 0;
 	}
 	else
 	{
 		EdgesOnFailure = Boolean.CreatedBoundaryEdges;
-		bOperationSucceeded = bBoolSuccess;
 	}
 
 	for (int EID : ResultMesh->EdgeIndicesItr())
