Add TriSel tool "By Material (All)" selection option and clarify that the previous "By Material" is "By Material (Connected)"

diff --git a/Engine/Plugins/Experimental/MeshModelingToolsetExp/Source/MeshModelingToolsExp/Private/MeshSelectionTool.cpp b/Engine/Plugins/Experimental/MeshModelingToolsetExp/Source/MeshModelingToolsExp/Private/MeshSelectionTool.cpp
--- a/Engine/Plugins/Experimental/MeshModelingToolsetExp/Source/MeshModelingToolsExp/Private/MeshSelectionTool.cpp
+++ b/Engine/Plugins/Experimental/MeshModelingToolsetExp/Source/MeshModelingToolsExp/Private/MeshSelectionTool.cpp
@@ -559,15 +559,35 @@ void UMeshSelectionTool::UpdateFaceSelection(const FBrushStampData& Stamp, const
 	else if (SelectionProps->SelectionMode == EMeshSelectionToolPrimaryMode::ByMaterial)
 	{
 		const FDynamicMeshMaterialAttribute* MaterialIDs = Mesh->Attributes()->GetMaterialID();
-		TArray<int32> StartROI;
-		StartROI.Add(Stamp.HitResult.FaceIndex);
-		FMeshConnectedComponents::GrowToConnectedTriangles(Mesh, StartROI, LocalROI, &TemporaryBuffer, &TemporarySet,
-			[Mesh, MaterialIDs](int t1, int t2) { return MaterialIDs->GetValue(t1) == MaterialIDs->GetValue(t2); });
-		UseROI = &LocalROI;
+		// If material IDs are missing, this logic will be skipped and selection will fall back to brush selection
+		if (MaterialIDs)
+		{
+			TArray<int32> StartROI;
+			StartROI.Add(Stamp.HitResult.FaceIndex);
+			FMeshConnectedComponents::GrowToConnectedTriangles(Mesh, StartROI, LocalROI, &TemporaryBuffer, &TemporarySet,
+				[Mesh, MaterialIDs](int t1, int t2) { return MaterialIDs->GetValue(t1) == MaterialIDs->GetValue(t2); });
+			UseROI = &LocalROI;
+		}
 	}
-	else if (SelectionProps->SelectionMode == EMeshSelectionToolPrimaryMode::ByUVIsland)
+	else if (SelectionProps->SelectionMode == EMeshSelectionToolPrimaryMode::ByMaterialAll)
 	{
 		const FDynamicMeshMaterialAttribute* MaterialIDs = Mesh->Attributes()->GetMaterialID();
+		// If material IDs are missing, this logic will be skipped and selection will fall back to brush selection
+		if (MaterialIDs)
+		{
+			int32 HitMaterialID = MaterialIDs->GetValue(Stamp.HitResult.FaceIndex);
+			for (int32 TID : Mesh->TriangleIndicesItr())
+			{
+				if (MaterialIDs->GetValue(TID) == HitMaterialID)
+				{
+					LocalROI.Add(TID);
+				}
+			}
+			UseROI = &LocalROI;
+		}
+	}
+	else if (SelectionProps->SelectionMode == EMeshSelectionToolPrimaryMode::ByUVIsland)
+	{
 		TArray<int32> StartROI;
 		StartROI.Add(Stamp.HitResult.FaceIndex);
 		FMeshConnectedComponents::GrowToConnectedTriangles(Mesh, StartROI, LocalROI, &TemporaryBuffer, &TemporarySet,
diff --git a/Engine/Plugins/Experimental/MeshModelingToolsetExp/Source/MeshModelingToolsExp/Public/MeshSelectionTool.h b/Engine/Plugins/Experimental/MeshModelingToolsetExp/Source/MeshModelingToolsExp/Public/MeshSelectionTool.h
--- a/Engine/Plugins/Experimental/MeshModelingToolsetExp/Source/MeshModelingToolsExp/Public/MeshSelectionTool.h
+++ b/Engine/Plugins/Experimental/MeshModelingToolsetExp/Source/MeshModelingToolsExp/Public/MeshSelectionTool.h
@@ -235,8 +235,11 @@ enum class EMeshSelectionToolPrimaryMode
 	/** Select all triangles in groups connected to any triangle inside the brush */
 	AllInGroup,
 
+	/** Select the connected group of triangles with same material as hit triangle */
+	ByMaterial UMETA(DisplayName = "By Material (Connected)"),
+
 	/** Select all triangles with same material as hit triangle */
-	ByMaterial,
+	ByMaterialAll UMETA(DisplayName = "By Material (All)"),
 
 	/** Select all triangles in same UV island as hit triangle */
 	ByUVIsland,
