RenderCaptureBaking: Expose render capture baking function to geometry script blueprints

#rnx

diff --git a/Engine/Plugins/Experimental/GeometryScripting/Source/GeometryScriptingCore/Private/MeshBakeFunctions.cpp b/Engine/Plugins/Experimental/GeometryScripting/Source/GeometryScriptingCore/Private/MeshBakeFunctions.cpp
--- a/Engine/Plugins/Experimental/GeometryScripting/Source/GeometryScriptingCore/Private/MeshBakeFunctions.cpp
+++ b/Engine/Plugins/Experimental/GeometryScripting/Source/GeometryScriptingCore/Private/MeshBakeFunctions.cpp
@@ -2,6 +2,11 @@
 
 #include "GeometryScript/MeshBakeFunctions.h"
 #include "UDynamicMesh.h"
+#include "GameFramework/Actor.h"
+
+#include "ModelingToolTargetUtil.h" // UE::ToolTarget functions
+#include "Scene/SceneCapturePhotoSet.h"
+#include "Baking/RenderCaptureFunctions.h"
 
 #include "Sampling/MeshBakerCommon.h"
 #include "Sampling/MeshMapBaker.h"
@@ -773,7 +778,81 @@ namespace GeometryScriptBakeLocals
 		}
 		return true;
 	}
-}
+
+	// Verbose but more robust than casting
+	EBakeTextureResolution ConvertEnum(EGeometryScriptBakeResolution Input)
+	{
+		switch(Input)
+		{
+			case EGeometryScriptBakeResolution::Resolution16:   return EBakeTextureResolution::Resolution16;
+			case EGeometryScriptBakeResolution::Resolution32:   return EBakeTextureResolution::Resolution32;
+			case EGeometryScriptBakeResolution::Resolution64:   return EBakeTextureResolution::Resolution64;
+			case EGeometryScriptBakeResolution::Resolution128:  return EBakeTextureResolution::Resolution128;
+			case EGeometryScriptBakeResolution::Resolution256:  return EBakeTextureResolution::Resolution256;
+			case EGeometryScriptBakeResolution::Resolution512:  return EBakeTextureResolution::Resolution512;
+			case EGeometryScriptBakeResolution::Resolution1024: return EBakeTextureResolution::Resolution1024;
+			case EGeometryScriptBakeResolution::Resolution2048: return EBakeTextureResolution::Resolution2048;
+			case EGeometryScriptBakeResolution::Resolution4096: return EBakeTextureResolution::Resolution4096;
+			case EGeometryScriptBakeResolution::Resolution8192: return EBakeTextureResolution::Resolution8192;
+		}
+		ensure(false);
+		return EBakeTextureResolution::Resolution16;
+	}
+
+	// Verbose but more robust than casting
+	EBakeTextureSamplesPerPixel ConvertEnum(EGeometryScriptBakeSamplesPerPixel Input)
+	{
+		switch(Input)
+		{
+			case EGeometryScriptBakeSamplesPerPixel::Sample1:   return EBakeTextureSamplesPerPixel::Sample1;
+			case EGeometryScriptBakeSamplesPerPixel::Sample4:   return EBakeTextureSamplesPerPixel::Sample4;
+			case EGeometryScriptBakeSamplesPerPixel::Sample16:  return EBakeTextureSamplesPerPixel::Sample16;
+			case EGeometryScriptBakeSamplesPerPixel::Sample64:  return EBakeTextureSamplesPerPixel::Sample64;
+			case EGeometryScriptBakeSamplesPerPixel::Samples256: return EBakeTextureSamplesPerPixel::Sample256;
+		}
+		ensure(false);
+		return EBakeTextureSamplesPerPixel::Sample1;
+	}
+
+
+	FRenderCaptureOptions MakeRenderCaptureOptions(
+		FGeometryScriptBakeRenderCaptureOptions BakeOptions,
+		FGeometryScriptBakeTargetMeshOptions TargetOptions
+		)
+	{
+		FRenderCaptureOptions Options;
+
+		Options.bBakeDeviceDepth = BakeOptions.CleanupTolerance > 0 &&
+			(
+			BakeOptions.bBaseColorMap ||
+			BakeOptions.bNormalMap    ||
+			BakeOptions.bEmissiveMap  ||
+			BakeOptions.bPackedMRSMap ||
+			BakeOptions.bMetallicMap  ||
+			BakeOptions.bRoughnessMap ||
+			BakeOptions.bSpecularMap
+			);
+		Options.ValidSampleDepthThreshold = BakeOptions.CleanupTolerance;
+
+		Options.RenderCaptureImageSize = GeometryScriptBakeLocals::GetDimensions(BakeOptions.RenderCaptureResolution).GetHeight();
+		Options.bAntiAliasing = BakeOptions.bRenderCaptureAntiAliasing;
+		Options.FieldOfViewDegrees = BakeOptions.FieldOfViewDegrees;
+		Options.NearPlaneDist = BakeOptions.NearPlaneDist;
+		Options.bBakeBaseColor = BakeOptions.bBaseColorMap;
+		Options.bBakeEmissive = BakeOptions.bEmissiveMap;
+		Options.bBakeNormalMap = BakeOptions.bNormalMap;
+
+		// Enforce the PackedMRS precondition here
+		Options.bUsePackedMRS =  BakeOptions.bPackedMRSMap;
+		Options.bBakeMetallic =  BakeOptions.bPackedMRSMap ? false : BakeOptions.bMetallicMap;
+		Options.bBakeRoughness = BakeOptions.bPackedMRSMap ? false : BakeOptions.bRoughnessMap;
+		Options.bBakeSpecular =  BakeOptions.bPackedMRSMap ? false : BakeOptions.bSpecularMap;
+
+		Options.TargetUVLayer = static_cast<int32>(TargetOptions.TargetUVLayer);
+
+		return Options;
+	}
+} // end namespace GeometryScriptBakeLocals
 
 
 FGeometryScriptBakeTypeOptions UGeometryScriptLibrary_MeshBakeFunctions::MakeBakeTypeTangentNormal()
@@ -954,5 +1033,161 @@ UDynamicMesh* UGeometryScriptLibrary_MeshBakeFunctions::BakeVertex(
 }
 
 
+FGeometryScriptRenderCaptureTextures UGeometryScriptLibrary_MeshBakeFunctions::BakeTextureFromRenderCaptures(
+	UDynamicMesh* TargetMesh,
+	FTransform TargetLocalToWorld,
+	FGeometryScriptBakeTargetMeshOptions TargetOptions,
+	const TArray<AActor*> SourceActors,
+	FGeometryScriptBakeRenderCaptureOptions BakeOptions,
+	UGeometryScriptDebug* Debug)
+{
+	FRenderCaptureOptions Options = GeometryScriptBakeLocals::MakeRenderCaptureOptions(BakeOptions, TargetOptions);
+
+	// Its possible to pass nullptrs in SourceActors so we filter these out here
+	TArray<TObjectPtr<AActor>> ValidSourceActors;
+	for (AActor* Actor : SourceActors)
+	{
+		if (Actor != nullptr)
+		{
+			ValidSourceActors.Add(Actor);
+		}
+	}
+
+	if (ValidSourceActors.IsEmpty())
+	{
+		return {};
+	}
+
+	if (TargetMesh == nullptr)
+	{
+		AppendError(Debug, EGeometryScriptErrorType::InvalidInputs,
+			LOCTEXT("BakeTextureFromRenderCaptures_InvalidTargetMesh", "BakeTextureFromRenderCaptures: TargetMesh is Null"));
+		return {};
+	}
+
+	// Transform the mesh, and undo it before we return
+	MeshTransforms::ApplyTransform(TargetMesh->GetMeshRef(), TargetLocalToWorld);
+	ON_SCOPE_EXIT {
+		MeshTransforms::ApplyTransformInverse(TargetMesh->GetMeshRef(), TargetLocalToWorld);
+	};
+
+	const auto HasDegenerateUVs = [&TargetMesh, &TargetOptions]
+	{
+		FDynamicMeshUVOverlay* UVOverlay = TargetMesh->GetMeshRef().Attributes()->GetUVLayer(TargetOptions.TargetUVLayer);
+		FAxisAlignedBox2f Bounds = FAxisAlignedBox2f::Empty();
+		for (const int Index : UVOverlay->ElementIndicesItr())
+		{
+			FVector2f UV;
+			UVOverlay->GetElement(Index, UV);
+			Bounds.Contain(UV);
+		}
+		return Bounds.Min == Bounds.Max;
+	};
+
+	if (TargetMesh->GetMeshRef().Attributes()->GetUVLayer(TargetOptions.TargetUVLayer) == nullptr)
+	{
+		AppendWarning(Debug, EGeometryScriptErrorType::InvalidInputs,
+			LOCTEXT("BakeTextureFromRenderCaptures_TargetMeshMissingUVs", "BakeTextureFromRenderCaptures: TargetMesh UV layer is missing"));
+		return {};
+	}
+
+	if (HasDegenerateUVs())
+	{
+		AppendWarning(Debug, EGeometryScriptErrorType::InvalidInputs,
+			LOCTEXT("BakeTextureFromRenderCaptures_TargetMeshDegenerateUVs", "BakeTextureFromRenderCaptures: TargetMesh UV layer is degenerate"));
+		return {};
+	}
+
+	if (Options.bBakeNormalMap && !FDynamicMeshTangents(TargetMesh->GetMeshPtr()).HasValidTangents(true))
+	{
+		AppendWarning(Debug, EGeometryScriptErrorType::InvalidInputs,
+			LOCTEXT("BakeTextureFromRenderCaptures_InvalidMeshTangents", "BakeTextureFromRenderCaptures: TargetMesh has invalid tangents so the requested normal map cannot be baked"));
+		return {};
+	}
+
+	TUniquePtr<FSceneCapturePhotoSet> SceneCapture = CapturePhotoSet(ValidSourceActors, Options, false);
+
+	const FDynamicMeshAABBTree3 TargetMeshSpatial(TargetMesh->GetMeshPtr());
+	TSharedPtr<FMeshTangentsd, ESPMode::ThreadSafe> TargetMeshTangents = MakeShared<FMeshTangentsd>(TargetMesh->GetMeshPtr());
+	TargetMeshTangents->CopyTriVertexTangents(TargetMesh->GetMeshRef());
+
+	FSceneCapturePhotoSetSampler Sampler(
+		SceneCapture.Get(),
+		BakeOptions.CleanupTolerance,
+		TargetMesh->GetMeshPtr(),
+		&TargetMeshSpatial,
+		TargetMeshTangents.Get());
+
+	FRenderCaptureOcclusionHandler OcclusionHandler(GeometryScriptBakeLocals::GetDimensions(BakeOptions.Resolution));
+
+	TUniquePtr<FMeshMapBaker> Baker = MakeRenderCaptureBaker(
+		TargetMesh->GetMeshPtr(),
+		TargetMeshTangents,
+		SceneCapture.Get(),
+		&Sampler,
+		Options,
+		GeometryScriptBakeLocals::ConvertEnum(BakeOptions.Resolution),
+		GeometryScriptBakeLocals::ConvertEnum(BakeOptions.SamplesPerPixel),
+		&OcclusionHandler);
+	{
+		TRACE_CPUPROFILER_EVENT_SCOPE(MeshBakeFunctions_BakeTextureFromRenderCaptures_Bake);
+		Baker->Bake();
+	}
+
+	FRenderCaptureTextures TexturesOut;
+	GetTexturesFromRenderCaptureBaker(Baker, TexturesOut);
+
+	// Update source data
+	if (Options.bBakeBaseColor && TexturesOut.BaseColorMap)
+	{
+		FTexture2DBuilder::CopyPlatformDataToSourceData(TexturesOut.BaseColorMap, FTexture2DBuilder::ETextureType::Color);
+	}
+	if (Options.bBakeNormalMap && TexturesOut.NormalMap)
+	{
+		FTexture2DBuilder::CopyPlatformDataToSourceData(TexturesOut.NormalMap, FTexture2DBuilder::ETextureType::NormalMap);
+	}
+	if (Options.bUsePackedMRS && TexturesOut.PackedMRSMap)
+	{
+		FTexture2DBuilder::CopyPlatformDataToSourceData(TexturesOut.PackedMRSMap, FTexture2DBuilder::ETextureType::ColorLinear);
+	}
+	if (Options.bBakeMetallic && TexturesOut.MetallicMap)
+	{
+		FTexture2DBuilder::CopyPlatformDataToSourceData(TexturesOut.MetallicMap, FTexture2DBuilder::ETextureType::Metallic);
+	}
+	if (Options.bBakeRoughness && TexturesOut.RoughnessMap)
+	{
+		FTexture2DBuilder::CopyPlatformDataToSourceData(TexturesOut.RoughnessMap, FTexture2DBuilder::ETextureType::Roughness);
+	}
+	if (Options.bBakeSpecular && TexturesOut.SpecularMap)
+	{
+		FTexture2DBuilder::CopyPlatformDataToSourceData(TexturesOut.SpecularMap, FTexture2DBuilder::ETextureType::Specular);
+	}
+	if (Options.bBakeEmissive && TexturesOut.EmissiveMap)
+	{
+		FTexture2DBuilder::CopyPlatformDataToSourceData(TexturesOut.EmissiveMap, FTexture2DBuilder::ETextureType::EmissiveHDR);
+	}
+
+	FGeometryScriptRenderCaptureTextures Result;
+
+	Result.BaseColorMap = TexturesOut.BaseColorMap;
+	Result.NormalMap    = TexturesOut.NormalMap;
+	Result.PackedMRSMap = TexturesOut.PackedMRSMap;
+	Result.MetallicMap  = TexturesOut.MetallicMap;
+	Result.RoughnessMap = TexturesOut.RoughnessMap;
+	Result.SpecularMap  = TexturesOut.SpecularMap;
+	Result.EmissiveMap  = TexturesOut.EmissiveMap;
+
+	Result.bHasBaseColorMap = (Result.BaseColorMap != nullptr);
+	Result.bHasNormalMap    = (Result.NormalMap    != nullptr);
+	Result.bHasPackedMRSMap = (Result.PackedMRSMap != nullptr);
+	Result.bHasMetallicMap  = (Result.MetallicMap  != nullptr);
+	Result.bHasRoughnessMap = (Result.RoughnessMap != nullptr);
+	Result.bHasSpecularMap  = (Result.SpecularMap  != nullptr);
+	Result.bHasEmissiveMap  = (Result.EmissiveMap  != nullptr);
+
+	return Result;
+}
+
+
 #undef LOCTEXT_NAMESPACE
 
diff --git a/Engine/Plugins/Experimental/GeometryScripting/Source/GeometryScriptingCore/Public/GeometryScript/MeshBakeFunctions.h b/Engine/Plugins/Experimental/GeometryScripting/Source/GeometryScriptingCore/Public/GeometryScript/MeshBakeFunctions.h
--- a/Engine/Plugins/Experimental/GeometryScripting/Source/GeometryScriptingCore/Public/GeometryScript/MeshBakeFunctions.h
+++ b/Engine/Plugins/Experimental/GeometryScripting/Source/GeometryScriptingCore/Public/GeometryScript/MeshBakeFunctions.h
@@ -296,6 +296,123 @@ struct GEOMETRYSCRIPTINGCORE_API FGeometryScriptBakeSourceMeshOptions
 	EGeometryScriptBakeNormalSpace SourceNormalSpace = EGeometryScriptBakeNormalSpace::Tangent;
 };
 
+USTRUCT(BlueprintType)
+struct GEOMETRYSCRIPTINGCORE_API FGeometryScriptBakeRenderCaptureOptions
+{
+	GENERATED_BODY()
+
+	/** The pixel resolution of the generated textures */
+	UPROPERTY(EditAnywhere, BlueprintReadWrite, Category = Options)
+	EGeometryScriptBakeResolution Resolution = EGeometryScriptBakeResolution::Resolution256;
+	
+	/** The pixel resolution of render capture photo set */
+	UPROPERTY(EditAnywhere, BlueprintReadWrite, Category = Options)
+	EGeometryScriptBakeResolution RenderCaptureResolution = EGeometryScriptBakeResolution::Resolution256;
+
+	/** Number of samples per pixel */
+	UPROPERTY(EditAnywhere, BlueprintReadWrite, Category = Options)
+	EGeometryScriptBakeSamplesPerPixel SamplesPerPixel = EGeometryScriptBakeSamplesPerPixel::Sample1;
+	
+	/** Whether to use anti-aliasing in the render captures, this may introduce artefacts if pixels at different scene depths get blended */
+	UPROPERTY(EditAnywhere, BlueprintReadWrite, Category = Options, DisplayName="Render Capture Anti-aliasing")
+	bool bRenderCaptureAntiAliasing = false;
+
+	/**
+	 * This threshold is used to detect occlusion artefacts (e.g., silhouettes/blotches in the base color) so that texels
+	 * exhibiting them can be cleaned up i.e., rejected and filled in using the values of nearby valid texels instead.
+	 * 
+	 * If the threshold is zero, the cleanup step is skipped.
+	 * If the threshold is too large, texels with artefacts won't be detected and the cleanup step is effectively skipped.
+	 * If the threshold is too small, texels without artefacts get detected and results will be bad because there won't
+	 * be enough nearby valid texels from which to infill values.
+	 * 
+	 * A good starting point is to choose a threshold around the size of the distance (in world space) between the
+	 * target and source meshes; if the target was generated by VoxWrap then the voxel size estimates this distance.
+	 */
+	UPROPERTY(EditAnywhere, BlueprintReadWrite, Category = Options)
+	float CleanupTolerance = 0.f;
+
+	UPROPERTY(EditAnywhere, BlueprintReadWrite, Category = Options)
+	bool bBaseColorMap = true;
+
+	/** World space normal map */
+	UPROPERTY(EditAnywhere, BlueprintReadWrite, Category = Options)
+	bool bNormalMap = true;
+
+	/** Packed Metallic/Roughness/Specular Map. If this is true the separated channel booleans are ignored */
+	UPROPERTY(EditAnywhere, BlueprintReadWrite, Category = Options, DisplayName="Packed MRS Map")
+	bool bPackedMRSMap = true;
+
+	UPROPERTY(EditAnywhere, BlueprintReadWrite, Category = Options)
+	bool bMetallicMap = false;
+
+	UPROPERTY(EditAnywhere, BlueprintReadWrite, Category = Options)
+	bool bRoughnessMap = false;
+
+	UPROPERTY(EditAnywhere, BlueprintReadWrite, Category = Options)
+	bool bSpecularMap = false;
+
+	UPROPERTY(EditAnywhere, BlueprintReadWrite, Category = Options)
+	bool bEmissiveMap = true;
+
+	// These are hidden in the UI right now, we might want to expose them if they turn out to be useful for very large
+	// or very small objects (not tested yet) TODO Figure out if we want to expose these options
+
+	UPROPERTY()
+	double FieldOfViewDegrees = 45.0;
+
+	UPROPERTY()
+	double NearPlaneDist = 1.0;
+};
+
+USTRUCT(BlueprintType)
+struct GEOMETRYSCRIPTINGCORE_API FGeometryScriptRenderCaptureTextures
+{
+	GENERATED_BODY();
+
+	UPROPERTY(EditAnywhere, BlueprintReadWrite, Category = Options)
+	TObjectPtr<UTexture2D> BaseColorMap = nullptr;
+
+	UPROPERTY(EditAnywhere, BlueprintReadWrite, Category = Options)
+	bool bHasBaseColorMap = false;
+
+	UPROPERTY(EditAnywhere, BlueprintReadWrite, Category = Options)
+	TObjectPtr<UTexture2D> NormalMap = nullptr;
+
+	UPROPERTY(EditAnywhere, BlueprintReadWrite, Category = Options)
+	bool bHasNormalMap = false;
+
+	UPROPERTY(EditAnywhere, BlueprintReadWrite, Category = Options, DisplayName="Packed MRS Map")
+	TObjectPtr<UTexture2D> PackedMRSMap = nullptr;
+
+	UPROPERTY(EditAnywhere, BlueprintReadWrite, Category = Options, DisplayName="Has Packed MRS Map")
+	bool bHasPackedMRSMap = false;
+
+	UPROPERTY(EditAnywhere, BlueprintReadWrite, Category = Options)
+	TObjectPtr<UTexture2D> MetallicMap = nullptr;
+
+	UPROPERTY(EditAnywhere, BlueprintReadWrite, Category = Options)
+	bool bHasMetallicMap = false;
+
+	UPROPERTY(EditAnywhere, BlueprintReadWrite, Category = Options)
+	TObjectPtr<UTexture2D> RoughnessMap = nullptr;
+
+	UPROPERTY(EditAnywhere, BlueprintReadWrite, Category = Options)
+	bool bHasRoughnessMap = false;
+
+	UPROPERTY(EditAnywhere, BlueprintReadWrite, Category = Options)
+	TObjectPtr<UTexture2D> SpecularMap = nullptr;
+
+	UPROPERTY(EditAnywhere, BlueprintReadWrite, Category = Options)
+	bool bHasSpecularMap = false;
+
+	UPROPERTY(EditAnywhere, BlueprintReadWrite, Category = Options)
+	TObjectPtr<UTexture2D> EmissiveMap = nullptr;
+
+	UPROPERTY(EditAnywhere, BlueprintReadWrite, Category = Options)
+	bool bHasEmissiveMap = false;
+};
+
 UCLASS(meta = (ScriptName = "GeometryScript_Bake"))
 class GEOMETRYSCRIPTINGCORE_API UGeometryScriptLibrary_MeshBakeFunctions : public UBlueprintFunctionLibrary
 {
@@ -373,4 +490,14 @@ public:
 		FGeometryScriptBakeOutputType BakeTypes,
 		FGeometryScriptBakeVertexOptions BakeOptions,
 		UGeometryScriptDebug* Debug = nullptr);
+
+	UFUNCTION(BlueprintCallable, Category = "GeometryScript|Bake")
+	static UPARAM(DisplayName="Textures Out") FGeometryScriptRenderCaptureTextures
+	BakeTextureFromRenderCaptures(
+		UDynamicMesh* TargetMesh,
+		FTransform TargetLocalToWorld,
+		FGeometryScriptBakeTargetMeshOptions TargetOptions,
+		const TArray<AActor*> SourceActors,
+		FGeometryScriptBakeRenderCaptureOptions BakeOptions,
+		UGeometryScriptDebug* Debug = nullptr);
 };
diff --git a/Engine/Plugins/Experimental/MeshModelingToolsetExp/Source/MeshModelingToolsExp/Private/BakeMeshAttributeMapsToolBase.cpp b/Engine/Plugins/Experimental/MeshModelingToolsetExp/Source/MeshModelingToolsExp/Private/BakeMeshAttributeMapsToolBase.cpp
--- a/Engine/Plugins/Experimental/MeshModelingToolsetExp/Source/MeshModelingToolsExp/Private/BakeMeshAttributeMapsToolBase.cpp
+++ b/Engine/Plugins/Experimental/MeshModelingToolsetExp/Source/MeshModelingToolsExp/Private/BakeMeshAttributeMapsToolBase.cpp
@@ -1,6 +1,7 @@
 // Copyright Epic Games, Inc. All Rights Reserved.
 
 #include "BakeMeshAttributeMapsToolBase.h"
+#include "Baking/BakingTypes.h"
 #include "InteractiveToolManager.h"
 #include "ToolBuilderUtil.h"
 #include "ToolSetupUtil.h"
diff --git a/Engine/Plugins/Experimental/MeshModelingToolsetExp/Source/MeshModelingToolsExp/Private/BakeRenderCaptureTool.cpp b/Engine/Plugins/Experimental/MeshModelingToolsetExp/Source/MeshModelingToolsExp/Private/BakeRenderCaptureTool.cpp
--- a/Engine/Plugins/Experimental/MeshModelingToolsetExp/Source/MeshModelingToolsExp/Private/BakeRenderCaptureTool.cpp
+++ b/Engine/Plugins/Experimental/MeshModelingToolsetExp/Source/MeshModelingToolsExp/Private/BakeRenderCaptureTool.cpp
@@ -17,11 +17,10 @@
 
 #include "EngineAnalytics.h"
 
+#include "Baking/RenderCaptureFunctions.h"
+#include "Baking/BakingTypes.h"
 #include "Sampling/MeshImageBakingCache.h"
 #include "Sampling/MeshMapBaker.h"
-#include "Sampling/RenderCaptureMapEvaluator.h"
-#include "Image/ImageInfilling.h"
-#include "Algo/NoneOf.h"
 #include "Misc/ScopedSlowTask.h"
 
 #include UE_INLINE_GENERATED_CPP_BY_NAME(BakeRenderCaptureTool)
@@ -31,81 +30,6 @@ using namespace UE::Geometry;
 
 #define LOCTEXT_NAMESPACE "UBakeRenderCaptureTool"
 
-//
-// Implementation details
-//
-
-
-
-class FSceneCapturePhotoSetSampler : public FMeshBakerDynamicMeshSampler
-{
-public:
-	FSceneCapturePhotoSetSampler(
-		FSceneCapturePhotoSet* SceneCapture,
-		float ValidSampleDepthThreshold,
-		TFunctionRef<bool(const FVector3d&, const FVector3d&)> VisibilityFunction,
-		const FDynamicMesh3* Mesh,
-		const FDynamicMeshAABBTree3* Spatial,
-		const FMeshTangentsd* Tangents)
-	:
-		FMeshBakerDynamicMeshSampler(Mesh, Spatial, Tangents),
-		SceneCapture(SceneCapture),
-		ValidSampleDepthThreshold(ValidSampleDepthThreshold),
-		VisibilityFunction(VisibilityFunction)
-	{
-		check(SceneCapture != nullptr);
-		check(Mesh != nullptr);
-		check(Spatial != nullptr);
-		check(Tangents != nullptr);
-
-		bool bHasDepth = SceneCapture->GetCaptureTypeEnabled(ERenderCaptureType::DeviceDepth);
-		if (bHasDepth)
-		{
-			ensure(ValidSampleDepthThreshold > 0); // We only need the depth capture if this threshold is positive
-		}
-	}
-
-	virtual bool SupportsCustomCorrespondence() const override
-	{
-		return true;
-	}
-
-	// Warning: Expects that Sample.BaseSample.SurfacePoint and Sample.BaseNormal are set when the function is called
-	virtual void* ComputeCustomCorrespondence(const FMeshUVSampleInfo& SampleInfo, FMeshMapEvaluator::FCorrespondenceSample& Sample) const override
-	{
-		// Perform a ray-cast to determine which photo/coordinate, if any, should be sampled
-		int PhotoIndex;
-		FVector2d PhotoCoords;
-		SceneCapture->ComputeSampleLocation(
-			Sample.BaseSample.SurfacePoint,
-			Sample.BaseNormal,
-			ValidSampleDepthThreshold,
-			VisibilityFunction,
-			PhotoIndex,
-			PhotoCoords);
-
-		// Store the photo coordinates and index in the correspondence sample
-		Sample.DetailMesh = SceneCapture;
-		Sample.DetailTriID = PhotoIndex;
-		Sample.DetailBaryCoords.X = PhotoCoords.X;
-		Sample.DetailBaryCoords.Y = PhotoCoords.Y;
-
-		// This will be set to Sample.DetailMesh but we can already do that internally so it's kindof redundant
-		return SceneCapture;
-	}
-
-	virtual bool IsValidCorrespondence(const FMeshMapEvaluator::FCorrespondenceSample& Sample) const override
-	{
-		return Sample.DetailTriID != IndexConstants::InvalidID;
-	}
-
-public:
-	FSceneCapturePhotoSet* SceneCapture = nullptr;
-	float ValidSampleDepthThreshold = 0;
-	TFunctionRef<bool(const FVector3d&, const FVector3d&)> VisibilityFunction;
-};
-
-
 
 static FString BaseColorTexParamName = TEXT("BaseColor");
 static FString RoughnessTexParamName = TEXT("Roughness");
@@ -115,199 +39,40 @@ static FString EmissiveTexParamName = TEXT("Emissive");
 static FString NormalTexParamName = TEXT("NormalMap");
 static FString PackedMRSTexParamName = TEXT("PackedMRS");
 
-class FRenderCaptureSettings
-{
-public:
-	enum class ETextureSizePolicy : uint8
-	{
-		TextureSize = 0,
-		TexelDensity = 1
-	};
 
-	/**
-	 * Input options to Actor Approximation process
-	 */
-	struct FOptions
-	{
-		//
-		// Material approximation settings
-		//
-		int32 RenderCaptureImageSize = 1024;
-		bool bAntiAliasing = false;
-
-		// render capture parameters
-		double FieldOfViewDegrees = 45.0;
-		double NearPlaneDist = 1.0;
-		double ValidSampleDepthThreshold = 0;
-
-		//
-		// Material output settings
-		//
-
-		// A new MIC derived from this material will be created and assigned to the generated mesh
-		UMaterialInterface* BakeMaterial = nullptr;		// if null, will use /MeshModelingToolsetExp/Materials/FullMaterialBakePreviewMaterial_PackedMRS instead
-		bool bBakeBaseColor = true;
-		bool bBakeRoughness = true;
-		bool bBakeMetallic = true;
-		bool bBakeSpecular = true;
-		bool bBakeEmissive = true;
-		bool bBakeNormalMap = true;
-		bool bBakeDeviceDepth = true;
-		
-		bool bUsePackedMRS = true;
-
-		//
-		// Mesh settings
-		//
-
-		//  Which UV layer of the Target mesh (the one we're baking to) should be used
-		int32 TargetUVLayer = 0;
-	};
-
-	/**
-	 * Construct an FOptions from the provided FMeshApproximationSettings.
-	 */
-	static FOptions ConstructOptions(
-		const URenderCaptureProperties& RenderCaptureProperties,
-		const UBakeRenderCaptureToolProperties& ToolProperties,
-		const UBakeRenderCaptureInputToolProperties& InputMeshSettings)
-	{
-		//
-		// Construct options for ApproximateActors operation
-		//
-		FOptions Options;
-
-		Options.TargetUVLayer = InputMeshSettings.GetTargetUVLayerIndex();
-
-		Options.RenderCaptureImageSize = static_cast<int32>(RenderCaptureProperties.Resolution);
-		Options.ValidSampleDepthThreshold = ToolProperties.ValidSampleDepthThreshold;
-
-		Options.bBakeBaseColor = RenderCaptureProperties.bBaseColorMap;
-		Options.bBakeNormalMap = RenderCaptureProperties.bNormalMap;
-		Options.bBakeEmissive =  RenderCaptureProperties.bEmissiveMap;
-		Options.bBakeDeviceDepth = RenderCaptureProperties.bDeviceDepthMap;
-		
-		// Enforce the PackedMRS precondition here so we don't have to check it at each usage site.  Note: We don't
-		// apply this precondition on the RenderCaptureProperties because we don't want the user to have to re-enable
-		// options which enabling PackedMRS disabled.
-		Options.bUsePackedMRS =  RenderCaptureProperties.bPackedMRSMap;
-		Options.bBakeMetallic =  RenderCaptureProperties.bPackedMRSMap ? false : RenderCaptureProperties.bMetallicMap;
-		Options.bBakeRoughness = RenderCaptureProperties.bPackedMRSMap ? false : RenderCaptureProperties.bRoughnessMap;
-		Options.bBakeSpecular =  RenderCaptureProperties.bPackedMRSMap ? false : RenderCaptureProperties.bSpecularMap;
-
-		Options.bAntiAliasing = RenderCaptureProperties.bAntiAliasing;
-		Options.FieldOfViewDegrees = RenderCaptureProperties.CaptureFieldOfView;
-		Options.NearPlaneDist = RenderCaptureProperties.NearPlaneDist;
-
-		return Options;
-	}
-};
-
-static TUniquePtr<FSceneCapturePhotoSet> CapturePhotoSet(
-	const TArray<TObjectPtr<AActor>>& Actors,
-	const FRenderCaptureSettings::FOptions& Options,
-	bool bAllowCancel
-)
+FRenderCaptureOptions
+MakeRenderCaptureOptions(
+	const URenderCaptureProperties& RenderCaptureProperties,
+	const UBakeRenderCaptureToolProperties& ToolProperties,
+	const UBakeRenderCaptureInputToolProperties& InputMeshSettings)
 {
-	TRACE_CPUPROFILER_EVENT_SCOPE(CapturePhotoSet);
-
-	FScopedSlowTask Progress(1.f, LOCTEXT("CapturingScene", "Capturing Scene..."));
-	Progress.EnterProgressFrame(1.f);
-	Progress.MakeDialog(bAllowCancel);
-
-	double FieldOfView = Options.FieldOfViewDegrees;
-	double NearPlaneDist = Options.NearPlaneDist;
-
-	FImageDimensions CaptureDimensions(Options.RenderCaptureImageSize, Options.RenderCaptureImageSize);
-
-	TUniquePtr<FSceneCapturePhotoSet> SceneCapture = MakeUnique<FSceneCapturePhotoSet>();
-	SceneCapture->SetAllowCancel(bAllowCancel);
+	FRenderCaptureOptions Options;
 
-	SceneCapture->SetCaptureTypeEnabled(ERenderCaptureType::DeviceDepth, Options.bBakeDeviceDepth);
-	SceneCapture->SetCaptureTypeEnabled(ERenderCaptureType::BaseColor,   Options.bBakeBaseColor);
-	SceneCapture->SetCaptureTypeEnabled(ERenderCaptureType::WorldNormal, Options.bBakeNormalMap);
-	SceneCapture->SetCaptureTypeEnabled(ERenderCaptureType::Emissive,    Options.bBakeEmissive);
-	SceneCapture->SetCaptureTypeEnabled(ERenderCaptureType::CombinedMRS, Options.bUsePackedMRS);
-	SceneCapture->SetCaptureTypeEnabled(ERenderCaptureType::Metallic,    Options.bBakeMetallic);
-	SceneCapture->SetCaptureTypeEnabled(ERenderCaptureType::Roughness,   Options.bBakeRoughness);
-	SceneCapture->SetCaptureTypeEnabled(ERenderCaptureType::Specular,    Options.bBakeSpecular);
+	Options.TargetUVLayer = InputMeshSettings.GetTargetUVLayerIndex();
 
-	FRenderCaptureConfig Config;
-	Config.bAntiAliasing = Options.bAntiAliasing;
-	SceneCapture->SetCaptureConfig(ERenderCaptureType::BaseColor,   Config);
-	SceneCapture->SetCaptureConfig(ERenderCaptureType::WorldNormal, Config);
-	SceneCapture->SetCaptureConfig(ERenderCaptureType::CombinedMRS, Config);
-	SceneCapture->SetCaptureConfig(ERenderCaptureType::Metallic,    Config);
-	SceneCapture->SetCaptureConfig(ERenderCaptureType::Roughness,   Config);
-	SceneCapture->SetCaptureConfig(ERenderCaptureType::Specular,    Config);
-	SceneCapture->SetCaptureConfig(ERenderCaptureType::Emissive,    Config);
+	Options.RenderCaptureImageSize = static_cast<int32>(RenderCaptureProperties.Resolution);
+	Options.ValidSampleDepthThreshold = ToolProperties.ValidSampleDepthThreshold;
 
-	SceneCapture->SetCaptureSceneActors(Actors[0]->GetWorld(), Actors);
-
-	// SceneCapture->SetEnableWriteDebugImages(true);
-
-	SceneCapture->AddStandardExteriorCapturesFromBoundingBox(
-		CaptureDimensions, FieldOfView, NearPlaneDist,
-		true, true, true, true, true);
-
-	return SceneCapture;
+	Options.bBakeBaseColor = RenderCaptureProperties.bBaseColorMap;
+	Options.bBakeNormalMap = RenderCaptureProperties.bNormalMap;
+	Options.bBakeEmissive =  RenderCaptureProperties.bEmissiveMap;
+	Options.bBakeDeviceDepth = RenderCaptureProperties.bDeviceDepthMap;
+	
+	// Enforce the PackedMRS precondition here so we don't have to check it at each usage site.  Note: We don't
+	// apply this precondition on the RenderCaptureProperties because we don't want the user to have to re-enable
+	// options which enabling PackedMRS disabled.
+	Options.bUsePackedMRS =  RenderCaptureProperties.bPackedMRSMap;
+	Options.bBakeMetallic =  RenderCaptureProperties.bPackedMRSMap ? false : RenderCaptureProperties.bMetallicMap;
+	Options.bBakeRoughness = RenderCaptureProperties.bPackedMRSMap ? false : RenderCaptureProperties.bRoughnessMap;
+	Options.bBakeSpecular =  RenderCaptureProperties.bPackedMRSMap ? false : RenderCaptureProperties.bSpecularMap;
+
+	Options.bAntiAliasing = RenderCaptureProperties.bAntiAliasing;
+	Options.FieldOfViewDegrees = RenderCaptureProperties.CaptureFieldOfView;
+	Options.NearPlaneDist = RenderCaptureProperties.NearPlaneDist;
+
+	return Options;
 }
 
-template <ERenderCaptureType CaptureType>
-TSharedPtr<FRenderCaptureMapEvaluator<FVector4f>>
-MakeColorEvaluator(
-	const FSceneCapturePhotoSet::FSceneSample& DefaultSample,
-	const FSceneCapturePhotoSet* SceneCapture,
-	const FImageDimensions& CaptureDimensions)
-{
-	TSharedPtr<FRenderCaptureMapEvaluator<FVector4f>> Evaluator = MakeShared<FRenderCaptureMapEvaluator<FVector4f>>();
-
-	switch (CaptureType) {
-	case ERenderCaptureType::BaseColor:
-		Evaluator->Channel = ERenderCaptureChannel::BaseColor;
-		break;
-	case ERenderCaptureType::Roughness:
-		Evaluator->Channel = ERenderCaptureChannel::Roughness;
-		break;
-	case ERenderCaptureType::Metallic:
-		Evaluator->Channel = ERenderCaptureChannel::Metallic;
-		break;
-	case ERenderCaptureType::Specular:
-		Evaluator->Channel = ERenderCaptureChannel::Specular;
-		break;
-	case ERenderCaptureType::Emissive:
-		Evaluator->Channel = ERenderCaptureChannel::Emissive;
-		break;
-	case ERenderCaptureType::WorldNormal:
-		Evaluator->Channel = ERenderCaptureChannel::WorldNormal;
-		break;
-	case ERenderCaptureType::CombinedMRS:
-		Evaluator->Channel = ERenderCaptureChannel::CombinedMRS;
-		break;
-	case ERenderCaptureType::DeviceDepth:
-		Evaluator->Channel = ERenderCaptureChannel::DeviceDepth;
-		break;
-	}
-
-	Evaluator->DefaultResult = DefaultSample.GetValue4f(CaptureType);
-
-	Evaluator->EvaluateSampleCallback = [&DefaultSample, SceneCapture, &CaptureDimensions](const FMeshMapEvaluator::FCorrespondenceSample& Sample)
-	{
-		const int PhotoIndex = Sample.DetailTriID;
-		const FVector2d PhotoCoords(Sample.DetailBaryCoords.X, Sample.DetailBaryCoords.Y);
-		const FVector4f SampleColor = SceneCapture->ComputeSampleNearest<CaptureType>(PhotoIndex, PhotoCoords, DefaultSample);
-		return SampleColor;
-	};
-
-	Evaluator->EvaluateColorCallback = [](const int DataIdx, float*& In)
-	{
-		const FVector4f Out(In[0], In[1], In[2], In[3]);
-		In += 4;
-		return Out;
-	};
-
-	return Evaluator;
-}
 
 //
 // Tool Operator
@@ -318,8 +83,8 @@ class FRenderCaptureMapBakerOp : public TGenericDataOperator<FMeshMapBaker>
 public:
 	UE::Geometry::FDynamicMesh3* BaseMesh = nullptr;
 	TSharedPtr<UE::Geometry::FMeshTangentsd, ESPMode::ThreadSafe> BaseMeshTangents;
-	FRenderCaptureSettings::FOptions Options;
-	int32 TextureImageSize;
+	FRenderCaptureOptions Options;
+	EBakeTextureResolution TextureImageSize;
 	EBakeTextureSamplesPerPixel SamplesPerPixel;
 	FSceneCapturePhotoSet* SceneCapture = nullptr;
 
@@ -331,250 +96,30 @@ public:
 // Bake textures onto the base/target mesh by projecting/sampling the set of captured photos
 void FRenderCaptureMapBakerOp::CalculateResult(FProgressCancel*)
 {
-	TRACE_CPUPROFILER_EVENT_SCOPE(FRenderCaptureMapBakerOp_CalculateResult);
+	const FDynamicMeshAABBTree3 BaseMeshSpatial(BaseMesh);
 
-	check(BaseMesh != nullptr);
-	check(BaseMeshTangents.IsValid());
-	check(SceneCapture != nullptr);
+	FSceneCapturePhotoSetSampler Sampler(
+		SceneCapture,
+		Options.ValidSampleDepthThreshold,
+		BaseMesh,
+		&BaseMeshSpatial,
+		BaseMeshTangents.Get());
 
-	FDynamicMeshAABBTree3 BaseMeshSpatial(BaseMesh, true);
+	const FImageDimensions TextureDimensions(
+		static_cast<int32>(TextureImageSize),
+		static_cast<int32>(TextureImageSize));
 
-	// Compute an offset shift surface positions toward the render camera position so we don't immediately self intersect
-	// The Max expression ensures the code works when the base mesh is 2D along one coordinate direction
-	const double RayOffsetHackDist = (double)(100.0 * FMathf::ZeroTolerance * FMath::Max(BaseMesh->GetBounds().MinDim(), 0.01));
+	FRenderCaptureOcclusionHandler OcclusionHandler(TextureDimensions);
 
-	auto VisibilityFunction = [&BaseMeshSpatial, RayOffsetHackDist](const FVector3d& SurfPos, const FVector3d& ViewPos)
-	{
-		FVector3d RayDir = ViewPos - SurfPos;
-		double Dist = Normalize(RayDir);
-		FVector3d RayOrigin = SurfPos + RayOffsetHackDist * RayDir;
-		int32 HitTID = BaseMeshSpatial.FindNearestHitTriangle(FRay3d(RayOrigin, RayDir), IMeshSpatial::FQueryOptions(Dist));
-		return (HitTID == IndexConstants::InvalidID);
-	};
-
-	struct FInfillData
-	{
-		struct FSampleStats {
-			uint16 NumValid = 0;
-			uint16 NumInvalid = 0;
-
-			// The ==, !=, += operators and the Zero() function are required by the TMarchingPixelInfill implementation
-
-			bool operator==(const FSampleStats& Other) const
-			{
-				return (NumValid == Other.NumValid) && (NumInvalid == Other.NumInvalid);
-			}
-
-			bool operator!=(const FSampleStats& Other) const
-			{
-				return !(*this == Other);
-			}
-
-			FSampleStats& operator+=(const FSampleStats& Other)
-			{
-				NumValid += Other.NumValid;
-				NumInvalid += Other.NumInvalid;
-				return *this;
-			}
-
-			static FSampleStats Zero()
-			{
-				return FSampleStats{0, 0};
-			}
-		};
-
-		// Collect some sample stats per pixel, used to determine if a pixel requires infill or not
-		TImageBuilder<FSampleStats> SampleStats;
-
-		// The i-th element of this array indicates if the i-th evaluator needs infill
-		TArray<bool> EvaluatorNeedsInfill;
-	} InfillData;
-
-	InfillData.SampleStats.SetDimensions(FImageDimensions(TextureImageSize, TextureImageSize));
-	InfillData.SampleStats.Clear(FInfillData::FSampleStats{});
-
-	auto RegisterSampleStats = [&InfillData](bool bSampleValid, const FMeshMapEvaluator::FCorrespondenceSample& Sample, const FVector2d& UVPosition, const FVector2i& ImageCoords)
-	{
-		checkSlow(InfillData.SampleStats.GetDimensions().IsValidCoords(ImageCoords));
-		if (bSampleValid)
-		{
-			InfillData.SampleStats.GetPixel(ImageCoords).NumValid += 1;
-		}
-		else
-		{
-			InfillData.SampleStats.GetPixel(ImageCoords).NumInvalid += 1;
-		}
-	};
-
-	auto ComputeAndApplyInfill = [&InfillData](TArray<TUniquePtr<TImageBuilder<FVector4f>>>& BakeResults)
-	{
-		check(BakeResults.Num() == InfillData.EvaluatorNeedsInfill.Num());
-
-		if (BakeResults.IsEmpty() || Algo::NoneOf(InfillData.EvaluatorNeedsInfill))
-		{
-			return;
-		}
-
-		// Find pixels that need infill
-		TArray<FVector2i> MissingPixels;
-		FCriticalSection MissingPixelsLock;
-		ParallelFor(InfillData.SampleStats.GetDimensions().GetHeight(), [&MissingPixels, &MissingPixelsLock, &InfillData](int32 Y)
-		{
-			for (int32 X = 0; X < InfillData.SampleStats.GetDimensions().GetWidth(); X++)
-			{
-				const FInfillData::FSampleStats& Stats = InfillData.SampleStats.GetPixel(X, Y);
-				// TODO experiment with other classifications
-				if (Stats.NumInvalid > 0 && Stats.NumValid == 0)
-				{
-					MissingPixelsLock.Lock();
-					MissingPixels.Add(FVector2i(X, Y));
-					MissingPixelsLock.Unlock();
-				}
-			}
-		});
-
-		auto NormalizeFunc = [](FVector4f SumValue, int32 Count)
-		{
-			float InvSum = (Count == 0) ? 1.0f : (1.0f / Count);
-			return FVector4f(SumValue.X * InvSum, SumValue.Y * InvSum, SumValue.Z * InvSum, 1.0f);
-		};
-		
-		auto DummyNormalizeStatsFunc = [](FInfillData::FSampleStats SumValue, int32 Count)
-		{
-			// The return value must be different from MissingValue below so that ComputeInfill works correctly
-			return FInfillData::FSampleStats{TNumericLimits<uint16>::Max(), TNumericLimits<uint16>::Max()};
-		};
-
-		TMarchingPixelInfill<FInfillData::FSampleStats> Infill;
-		// This must be the same as the value of exterior pixels, otherwise infill will spread the exterior values into the texture
-		FInfillData::FSampleStats MissingValue{0, 0};
-		Infill.ComputeInfill(InfillData.SampleStats, MissingPixels, MissingValue, DummyNormalizeStatsFunc);
-		for (int EvaluatorIndex = 0; EvaluatorIndex < BakeResults.Num(); EvaluatorIndex++)
-		{
-			if (InfillData.EvaluatorNeedsInfill[EvaluatorIndex])
-			{
-				Infill.ApplyInfill<FVector4f>(*BakeResults[EvaluatorIndex], NormalizeFunc);
-			}
-		}
-	};
-
-	Result->SetTargetMesh(BaseMesh);
-	Result->SetTargetMeshTangents(BaseMeshTangents);
-	Result->SetDimensions(FImageDimensions(TextureImageSize, TextureImageSize));
-	Result->SetSamplesPerPixel(static_cast<int32>(SamplesPerPixel));
-	Result->SetFilter(FMeshMapBaker::EBakeFilterType::BSpline);
-	Result->SetTargetMeshUVLayer(Options.TargetUVLayer);
-	Result->InteriorSampleCallback = RegisterSampleStats;
-	Result->PostWriteToImageCallback = ComputeAndApplyInfill;
-
-	FSceneCapturePhotoSetSampler Sampler(SceneCapture, Options.ValidSampleDepthThreshold, VisibilityFunction, BaseMesh, &BaseMeshSpatial, BaseMeshTangents.Get());
-	Result->SetDetailSampler(&Sampler);
-	Result->SetCorrespondenceStrategy(FMeshMapBaker::ECorrespondenceStrategy::Custom);
-
-	// Pixels in the output textures which don't map onto the mesh have a light grey color (except the normal map which
-	// will show a color corresponding to a unit z tangent space normal)
-	const FVector4f InvalidColor(.42, .42, .42, 1);
-	const FVector3f DefaultNormal = FVector3f::UnitZ();
-
-	FSceneCapturePhotoSet::FSceneSample DefaultColorSample;
-	DefaultColorSample.BaseColor = FVector3f(InvalidColor.X, InvalidColor.Y, InvalidColor.Z);
-	DefaultColorSample.Roughness = InvalidColor.X;
-	DefaultColorSample.Specular = InvalidColor.X;
-	DefaultColorSample.Metallic = InvalidColor.X;
-	DefaultColorSample.Emissive = FVector3f(InvalidColor.X, InvalidColor.Y, InvalidColor.Z);
-	DefaultColorSample.WorldNormal = FVector4f((DefaultNormal + FVector3f::One()) * .5f, InvalidColor.W);
-	
-	// We use 0 here since this corresponds to the infinite far plane value. Also, we don't preview the depth texture
-	DefaultColorSample.DeviceDepth = 0;
-
-	auto AddColorEvaluator = [this, &InfillData] (const TSharedPtr<FRenderCaptureMapEvaluator<FVector4f>>& Evaluator)
-	{
-		Result->AddEvaluator(Evaluator);
-		InfillData.EvaluatorNeedsInfill.Add(true);
-	};
-	
-	const FImageDimensions CaptureDimensions(Options.RenderCaptureImageSize, Options.RenderCaptureImageSize);
-
-	if (Options.bBakeDeviceDepth)
-	{
-		AddColorEvaluator(MakeColorEvaluator<ERenderCaptureType::DeviceDepth>(DefaultColorSample, SceneCapture, CaptureDimensions));
-	}
-	if (Options.bBakeBaseColor)
-	{
-		AddColorEvaluator(MakeColorEvaluator<ERenderCaptureType::BaseColor>(DefaultColorSample, SceneCapture, CaptureDimensions));
-	}
-	if (Options.bUsePackedMRS)
-	{
-		AddColorEvaluator(MakeColorEvaluator<ERenderCaptureType::CombinedMRS>(DefaultColorSample, SceneCapture, CaptureDimensions));
-	}
-	if (Options.bBakeRoughness)
-	{
-		AddColorEvaluator(MakeColorEvaluator<ERenderCaptureType::Roughness>(DefaultColorSample, SceneCapture, CaptureDimensions));
-	}
-	if (Options.bBakeMetallic)
-	{
-		AddColorEvaluator(MakeColorEvaluator<ERenderCaptureType::Metallic>(DefaultColorSample, SceneCapture, CaptureDimensions));
-	}
-	if (Options.bBakeSpecular)
-	{
-		AddColorEvaluator(MakeColorEvaluator<ERenderCaptureType::Specular>(DefaultColorSample, SceneCapture, CaptureDimensions));
-	}
-	if (Options.bBakeEmissive)
-	{
-		AddColorEvaluator(MakeColorEvaluator<ERenderCaptureType::Emissive>(DefaultColorSample, SceneCapture, CaptureDimensions));
-	}
-	if (Options.bBakeNormalMap)
-	{
-		TSharedPtr<FRenderCaptureMapEvaluator<FVector3f>> Evaluator = MakeShared<FRenderCaptureMapEvaluator<FVector3f>>();
-
-		Evaluator->Channel = ERenderCaptureChannel::WorldNormal;
-
-		Evaluator->DefaultResult = DefaultNormal;
-
-		Evaluator->EvaluateSampleCallback = [this, &DefaultColorSample, &CaptureDimensions](const FMeshMapEvaluator::FCorrespondenceSample& Sample)
-		{
-			const int32 TriangleID = Sample.BaseSample.TriangleIndex;
-			const FVector3d BaryCoords = Sample.BaseSample.BaryCoords;
-			const int PhotoIndex = Sample.DetailTriID;
-			const FVector2d PhotoCoords(Sample.DetailBaryCoords.X, Sample.DetailBaryCoords.Y);
-
-			const FVector4f NormalColor = SceneCapture->ComputeSampleNearest<ERenderCaptureType::WorldNormal>(PhotoIndex, PhotoCoords, DefaultColorSample);
-
-			// Map from color components [0,1] to normal components [-1,1]
-			const FVector3f WorldSpaceNormal(
-				(NormalColor.X - 0.5f) * 2.0f,
-				(NormalColor.Y - 0.5f) * 2.0f,
-				(NormalColor.Z - 0.5f) * 2.0f);
-
-			// Get tangents on base mesh
-			FVector3d BaseTangentX, BaseTangentY;
-			BaseMeshTangents->GetInterpolatedTriangleTangent(TriangleID, BaryCoords, BaseTangentX, BaseTangentY);
-
-			// Compute normal in tangent space
-			const FVector3f TangentSpaceNormal(
-					(float)WorldSpaceNormal.Dot(FVector3f(BaseTangentX)),
-					(float)WorldSpaceNormal.Dot(FVector3f(BaseTangentY)),
-					(float)WorldSpaceNormal.Dot(FVector3f(Sample.BaseNormal)));
-
-			return TangentSpaceNormal;
-		};
-
-		Evaluator->EvaluateColorCallback = [](const int DataIdx, float*& In)
-		{
-			// Map normal space [-1,1] to color space [0,1]
-			const FVector3f Normal(In[0], In[1], In[2]);
-			const FVector3f Color = (Normal + FVector3f::One()) * 0.5f;
-			const FVector4f Out(Color.X, Color.Y, Color.Z, 1.0f);
-			In += 3;
-			return Out;
-		};
-
-		Result->AddEvaluator(Evaluator);
-
-		// Note: No infill on normal map for now, doesn't make sense to do after mapping to tangent space!
-		//  (should we build baked normal map in world space, and then resample to tangent space??)
-		InfillData.EvaluatorNeedsInfill.Add(false);
-	}
+	Result = MakeRenderCaptureBaker(
+		BaseMesh,
+		BaseMeshTangents,
+		SceneCapture,
+		&Sampler,
+		Options,
+		TextureImageSize,
+		SamplesPerPixel,
+		&OcclusionHandler);
 
 	{
 		TRACE_CPUPROFILER_EVENT_SCOPE(FRenderCaptureMapBakerOp_CalculateResult_Bake);
@@ -588,7 +133,6 @@ void FRenderCaptureMapBakerOp::CalculateResult(FProgressCancel*)
 //
 
 
-
 const FToolTargetTypeRequirements& UBakeRenderCaptureToolBuilder::GetTargetRequirements() const
 {
 	static FToolTargetTypeRequirements TypeRequirements({
@@ -619,6 +163,7 @@ UMultiSelectionMeshEditingTool* UBakeRenderCaptureToolBuilder::CreateNewTool(con
 
 
 
+
 void UBakeRenderCaptureTool::Setup()
 {
 	TRACE_CPUPROFILER_EVENT_SCOPE(UBakeRenderCaptureTool::Setup);
@@ -931,8 +476,8 @@ TUniquePtr<TGenericDataOperator<FMeshMapBaker>> UBakeRenderCaptureTool::MakeNewO
 	TUniquePtr<FRenderCaptureMapBakerOp> Op = MakeUnique<FRenderCaptureMapBakerOp>();
 	Op->BaseMesh = &TargetMesh;
 	Op->BaseMeshTangents = TargetMeshTangents;
-	Op->Options = FRenderCaptureSettings::ConstructOptions(*RenderCaptureProperties, *Settings, *InputMeshSettings);
-	Op->TextureImageSize = static_cast<int32>(Settings->TextureSize);
+	Op->Options = MakeRenderCaptureOptions(*RenderCaptureProperties, *Settings, *InputMeshSettings);
+	Op->TextureImageSize = Settings->TextureSize;
 	Op->SamplesPerPixel = Settings->SamplesPerPixel;
 	Op->SceneCapture = SceneCapture.Get();
 	return Op;
@@ -942,118 +487,19 @@ TUniquePtr<TGenericDataOperator<FMeshMapBaker>> UBakeRenderCaptureTool::MakeNewO
 
 void UBakeRenderCaptureTool::OnMapsUpdatedRC(const TUniquePtr<FMeshMapBaker>& NewResult)
 {
-	// We do this to defer work I guess, it was like this in the original ApproximateActors implementation :DeferredPopulateSourceData 
-	constexpr bool bPopulateSourceData = false;
-
 	TRACE_CPUPROFILER_EVENT_SCOPE(BakeRenderCaptureTool_Textures_BuildTextures);
 
-	const int32 NumEval = NewResult->NumEvaluators();
-	for (int32 EvalIdx = 0; EvalIdx < NumEval; ++EvalIdx)
-	{
-		FMeshMapEvaluator* BaseEval = NewResult->GetEvaluator(EvalIdx);
-
-		check(BaseEval->DataLayout().Num() == 1);
+	FRenderCaptureTextures TexturesOut;
+	GetTexturesFromRenderCaptureBaker(NewResult, TexturesOut);
 
-		switch (BaseEval->DataLayout()[0])
-		{
-		case FMeshMapEvaluator::EComponents::Float4:
-
-			{
-				FRenderCaptureMapEvaluator<FVector4f>* Eval = static_cast<FRenderCaptureMapEvaluator<FVector4f>*>(BaseEval);
-				TUniquePtr<TImageBuilder<FVector4f>> ImageBuilder = MoveTemp(NewResult->GetBakeResults(EvalIdx)[0]);
-				
-				if (ensure(ImageBuilder.IsValid()) == false) return;
-
-				switch (Eval->Channel)
-				{
-				case ERenderCaptureChannel::BaseColor:
-
-					ResultSettings->BaseColorMap = FTexture2DBuilder::BuildTextureFromImage(
-						*ImageBuilder,
-						FTexture2DBuilder::ETextureType::Color,
-						true,
-						bPopulateSourceData);
-					break;
-
-				case ERenderCaptureChannel::Roughness:
-
-					ResultSettings->RoughnessMap = FTexture2DBuilder::BuildTextureFromImage(
-						*ImageBuilder,
-						FTexture2DBuilder::ETextureType::Roughness,
-						false,
-						bPopulateSourceData);
-					break;
-
-				case ERenderCaptureChannel::Metallic:
-
-					ResultSettings->MetallicMap = FTexture2DBuilder::BuildTextureFromImage(
-						*ImageBuilder,
-						FTexture2DBuilder::ETextureType::Metallic,
-						false,
-						bPopulateSourceData);
-					break;
-
-				case ERenderCaptureChannel::Specular:
-
-					ResultSettings->SpecularMap = FTexture2DBuilder::BuildTextureFromImage(
-						*ImageBuilder,
-						FTexture2DBuilder::ETextureType::Specular,
-						false,
-						bPopulateSourceData);
-					break;
-
-				case ERenderCaptureChannel::Emissive:
-
-					ResultSettings->EmissiveMap = FTexture2DBuilder::BuildTextureFromImage(
-						*ImageBuilder,
-						FTexture2DBuilder::ETextureType::EmissiveHDR,
-						false,
-						bPopulateSourceData);
-					ResultSettings->EmissiveMap->CompressionSettings = TC_HDR_Compressed;
-					break;
-
-				case ERenderCaptureChannel::CombinedMRS:
-
-					ResultSettings->PackedMRSMap = FTexture2DBuilder::BuildTextureFromImage(
-						*ImageBuilder,
-						FTexture2DBuilder::ETextureType::ColorLinear,
-						false,
-						bPopulateSourceData);
-					break;
-
-				case ERenderCaptureChannel::DeviceDepth:
-					// Do nothing, the depth capture is used internally and not presented to the user
-					break;
-
-				default:
-					ensure(false);
-					return;
-				}
-
-			} break; // Float4
-
-		case FMeshMapEvaluator::EComponents::Float3:
-
-			{
-				FRenderCaptureMapEvaluator<FVector3f>* Eval = static_cast<FRenderCaptureMapEvaluator<FVector3f>*>(BaseEval);
-				TUniquePtr<TImageBuilder<FVector4f>> ImageBuilder = MoveTemp(NewResult->GetBakeResults(EvalIdx)[0]);
-				
-				if (ensure(ImageBuilder.IsValid()) == false) return;
-				if (ensure(Eval->Channel == ERenderCaptureChannel::WorldNormal) == false) return;
-
-				ResultSettings->NormalMap = FTexture2DBuilder::BuildTextureFromImage(
-					*ImageBuilder,
-					FTexture2DBuilder::ETextureType::NormalMap,
-					false,
-					bPopulateSourceData);
-
-			} break; // Float3
-			
-		default:
-			ensure(false);
-			return;
-		}
-	}
+	// Unpack TexturesOut to store in ResultSettings
+	ResultSettings->BaseColorMap = TexturesOut.BaseColorMap;
+	ResultSettings->NormalMap = TexturesOut.NormalMap;
+	ResultSettings->PackedMRSMap = TexturesOut.PackedMRSMap;
+	ResultSettings->MetallicMap = TexturesOut.MetallicMap;
+	ResultSettings->RoughnessMap = TexturesOut.RoughnessMap;
+	ResultSettings->SpecularMap = TexturesOut.SpecularMap;
+	ResultSettings->EmissiveMap = TexturesOut.EmissiveMap;
 
 	GatherAnalytics(*NewResult);
 	UpdateVisualization();
@@ -1157,7 +603,7 @@ void UBakeRenderCaptureTool::UpdateResult()
 
 	//
 	// create a set of spatially located render captures of the scene ("photo set"). We need to recompute this if the
-	// render capture properties changed. Note we only compare the URenderCaptureProperties only, and not the
+	// render capture properties changed. Note we only compare the URenderCaptureProperties, and not the
 	// ValidSampleDepthThreshold, this is intentional so that we only trigger a scene capture recompute when we go from
 	// a zero to a positive threshold (we need to compute the depth capture), or a positive to a zero threshold (we can
 	// save memory and not compute the depth capture), we don't need to recompute the scene capture when the user is
@@ -1174,8 +620,13 @@ void UBakeRenderCaptureTool::UpdateResult()
 		const bool bAllowCancel = (bFirstEverSceneCapture == false);
 
 		SceneCapture.Reset();
-		FRenderCaptureSettings::FOptions Options = FRenderCaptureSettings::ConstructOptions(*RenderCaptureProperties, *Settings, *InputMeshSettings);
-		SceneCapture = CapturePhotoSet(Actors, Options, bAllowCancel);
+		FRenderCaptureOptions Options = MakeRenderCaptureOptions(*RenderCaptureProperties, *Settings, *InputMeshSettings);
+		{
+			FScopedSlowTask Progress(1.f, LOCTEXT("CapturingScene", "Capturing Scene..."));
+			Progress.EnterProgressFrame(1.f);
+			Progress.MakeDialog(bAllowCancel);
+			SceneCapture = CapturePhotoSet(Actors, Options, bAllowCancel);
+		}
 
 		for (int Idx = 1; Idx < Targets.Num(); ++Idx)
 		{
diff --git a/Engine/Plugins/Experimental/MeshModelingToolsetExp/Source/MeshModelingToolsExp/Public/BakeMeshAttributeMapsTool.h b/Engine/Plugins/Experimental/MeshModelingToolsetExp/Source/MeshModelingToolsExp/Public/BakeMeshAttributeMapsTool.h
--- a/Engine/Plugins/Experimental/MeshModelingToolsetExp/Source/MeshModelingToolsExp/Public/BakeMeshAttributeMapsTool.h
+++ b/Engine/Plugins/Experimental/MeshModelingToolsetExp/Source/MeshModelingToolsExp/Public/BakeMeshAttributeMapsTool.h
@@ -11,6 +11,7 @@
 #include "ModelingOperators.h"
 #include "PreviewMesh.h"
 #include "ModelingToolTargetUtil.h"
+#include "Baking/BakingTypes.h"
 #include "BakeMeshAttributeMapsToolBase.h"
 #include "BakeMeshAttributeMapsTool.generated.h"
 
diff --git a/Engine/Plugins/Experimental/MeshModelingToolsetExp/Source/MeshModelingToolsExp/Public/BakeMeshAttributeMapsToolBase.h b/Engine/Plugins/Experimental/MeshModelingToolsetExp/Source/MeshModelingToolsExp/Public/BakeMeshAttributeMapsToolBase.h
--- a/Engine/Plugins/Experimental/MeshModelingToolsetExp/Source/MeshModelingToolsExp/Public/BakeMeshAttributeMapsToolBase.h
+++ b/Engine/Plugins/Experimental/MeshModelingToolsetExp/Source/MeshModelingToolsExp/Public/BakeMeshAttributeMapsToolBase.h
@@ -14,47 +14,9 @@
 #include "ModelingOperators.h"
 #include "PreviewMesh.h"
 #include "Sampling/MeshMapBaker.h"
+#include "Baking/BakingTypes.h"
 #include "BakeMeshAttributeMapsToolBase.generated.h"
 
-/**
-* Bake maps enums
-*/
-
-
-UENUM()
-enum class EBakeTextureResolution 
-{
-	Resolution16 = 16 UMETA(DisplayName = "16 x 16"),
-	Resolution32 = 32 UMETA(DisplayName = "32 x 32"),
-	Resolution64 = 64 UMETA(DisplayName = "64 x 64"),
-	Resolution128 = 128 UMETA(DisplayName = "128 x 128"),
-	Resolution256 = 256 UMETA(DisplayName = "256 x 256"),
-	Resolution512 = 512 UMETA(DisplayName = "512 x 512"),
-	Resolution1024 = 1024 UMETA(DisplayName = "1024 x 1024"),
-	Resolution2048 = 2048 UMETA(DisplayName = "2048 x 2048"),
-	Resolution4096 = 4096 UMETA(DisplayName = "4096 x 4096"),
-	Resolution8192 = 8192 UMETA(DisplayName = "8192 x 8192")
-};
-
-
-UENUM()
-enum class EBakeTextureBitDepth
-{
-	ChannelBits8 UMETA(DisplayName = "8 bits/channel"),
-	ChannelBits16 UMETA(DisplayName = "16 bits/channel")
-};
-
-
-UENUM()
-enum class EBakeTextureSamplesPerPixel
-{
-	Sample1 = 1 UMETA(DisplayName = "1"),
-	Sample4 = 4 UMETA(DisplayName = "4"),
-	Sample16 = 16 UMETA(DisplayName = "16"),
-	Sample64 = 64 UMETA(DisplayName = "64"),
-	Sample256 = 256 UMETA(DisplayName = "256")
-};
-
 
 UCLASS()
 class MESHMODELINGTOOLSEXP_API UBakeMeshAttributeMapsResultToolProperties : public UInteractiveToolPropertySet
diff --git a/Engine/Plugins/Experimental/MeshModelingToolsetExp/Source/MeshModelingToolsExp/Public/BakeMultiMeshAttributeMapsTool.h b/Engine/Plugins/Experimental/MeshModelingToolsetExp/Source/MeshModelingToolsExp/Public/BakeMultiMeshAttributeMapsTool.h
--- a/Engine/Plugins/Experimental/MeshModelingToolsetExp/Source/MeshModelingToolsExp/Public/BakeMultiMeshAttributeMapsTool.h
+++ b/Engine/Plugins/Experimental/MeshModelingToolsetExp/Source/MeshModelingToolsExp/Public/BakeMultiMeshAttributeMapsTool.h
@@ -12,6 +12,7 @@
 #include "Scene/MeshSceneAdapter.h"
 #include "ModelingOperators.h"
 #include "PreviewMesh.h"
+#include "Baking/BakingTypes.h"
 #include "BakeMeshAttributeMapsToolBase.h"
 #include "BakeMultiMeshAttributeMapsTool.generated.h"
 
diff --git a/Engine/Plugins/Experimental/MeshModelingToolsetExp/Source/MeshModelingToolsExp/Public/BakeRenderCaptureTool.h b/Engine/Plugins/Experimental/MeshModelingToolsetExp/Source/MeshModelingToolsExp/Public/BakeRenderCaptureTool.h
--- a/Engine/Plugins/Experimental/MeshModelingToolsetExp/Source/MeshModelingToolsExp/Public/BakeRenderCaptureTool.h
+++ b/Engine/Plugins/Experimental/MeshModelingToolsetExp/Source/MeshModelingToolsExp/Public/BakeRenderCaptureTool.h
@@ -3,9 +3,9 @@
 #pragma once
 
 #include "CoreMinimal.h"
-#include "MultiSelectionTool.h"
 #include "InteractiveToolBuilder.h"
 #include "BaseTools/MultiSelectionMeshEditingTool.h"
+#include "Baking/BakingTypes.h"
 #include "BakeMeshAttributeMapsToolBase.h"
 
 // Render Capture algorithm includes
@@ -15,6 +15,16 @@
 
 class UTexture2D;
 
+namespace UE
+{
+namespace Geometry
+{
+
+class FSceneCapturePhotoSet;
+
+}
+}
+
 //
 // Tool Result
 //
@@ -132,7 +142,6 @@ public:
 
 	bool operator==(const URenderCaptureProperties& Other) const
 	{
-		// Intentionally omitted TransientToolProperty
 		return Resolution == Other.Resolution
 			&& bBaseColorMap == Other.bBaseColorMap
 			&& bNormalMap == Other.bNormalMap
diff --git a/Engine/Plugins/Runtime/MeshModelingToolset/Source/ModelingComponents/Private/Baking/RenderCaptureFunctions.cpp b/Engine/Plugins/Runtime/MeshModelingToolset/Source/ModelingComponents/Private/Baking/RenderCaptureFunctions.cpp
new file mode 100644
--- /dev/null
+++ b/Engine/Plugins/Runtime/MeshModelingToolset/Source/ModelingComponents/Private/Baking/RenderCaptureFunctions.cpp
@@ -0,0 +1,577 @@
+﻿// Copyright Epic Games, Inc. All Rights Reserved.
+
+#include "Baking/RenderCaptureFunctions.h"
+#include "Scene/SceneCapturePhotoSet.h"
+#include "Sampling/MeshMapBaker.h"
+#include "Sampling/RenderCaptureMapEvaluator.h"
+#include "AssetUtils/Texture2DBuilder.h"
+
+#include "Engine/Texture2D.h"
+#include "GameFramework/Actor.h"
+#include "Algo/NoneOf.h"
+
+
+using namespace UE::Geometry;
+
+
+
+
+FRenderCaptureOcclusionHandler::FRenderCaptureOcclusionHandler(FImageDimensions Dimensions)
+{
+	SampleStats.SetDimensions(Dimensions);
+	SampleStats.Clear({});
+}
+
+void FRenderCaptureOcclusionHandler::RegisterSample(const FVector2i& ImageCoords, bool bSampleValid)
+{
+	checkSlow(SampleStats.GetDimensions().IsValidCoords(ImageCoords));
+	if (bSampleValid)
+	{
+		SampleStats.GetPixel(ImageCoords).NumValid += 1;
+	}
+	else
+	{
+		SampleStats.GetPixel(ImageCoords).NumInvalid += 1;
+	}
+}
+
+void FRenderCaptureOcclusionHandler::PushInfillRequired(bool bInfillRequired)
+{
+	InfillRequired.Add(bInfillRequired);
+}
+
+void FRenderCaptureOcclusionHandler::ComputeAndApplyInfill(TArray<TUniquePtr<TImageBuilder<FVector4f>>>& Images)
+{
+	check(Images.Num() == InfillRequired.Num());
+
+	if (Images.IsEmpty() || Algo::NoneOf(InfillRequired))
+	{
+		return;
+	}
+
+	ComputeInfill();
+
+	for (int ImageIndex = 0; ImageIndex < Images.Num(); ImageIndex++)
+	{
+		if (InfillRequired[ImageIndex])
+		{
+			ApplyInfill(*Images[ImageIndex]);
+		}
+	}
+}
+
+void FRenderCaptureOcclusionHandler::ComputeInfill()
+{
+	// Find pixels that need infill
+	TArray<FVector2i> MissingPixels;
+	FCriticalSection MissingPixelsLock;
+	ParallelFor(SampleStats.GetDimensions().GetHeight(), [this, &MissingPixels, &MissingPixelsLock](int32 Y)
+	{
+		for (int32 X = 0; X < SampleStats.GetDimensions().GetWidth(); X++)
+		{
+			const FSampleStats& Stats = SampleStats.GetPixel(X, Y);
+			// TODO experiment with other classifications
+			if (Stats.NumInvalid > 0 && Stats.NumValid == 0)
+			{
+				MissingPixelsLock.Lock();
+				MissingPixels.Add(FVector2i(X, Y));
+				MissingPixelsLock.Unlock();
+			}
+		}
+	});
+	
+	auto DummyNormalizeStatsFunc = [](FSampleStats SumValue, int32 Count)
+	{
+		// The return value must be different from MissingValue below so that ComputeInfill works correctly
+		return FSampleStats{TNumericLimits<uint16>::Max(), TNumericLimits<uint16>::Max()};
+	};
+
+	// This must be the same as the value of exterior pixels, otherwise infill will spread the exterior values into the texture
+	FSampleStats MissingValue{0, 0};
+	Infill.ComputeInfill(SampleStats, MissingPixels, MissingValue, DummyNormalizeStatsFunc);
+}
+
+void FRenderCaptureOcclusionHandler::ApplyInfill(TImageBuilder<FVector4f>& Image) const
+{
+	auto NormalizeFunc = [](FVector4f SumValue, int32 Count)
+	{
+		float InvSum = (Count == 0) ? 1.0f : (1.0f / Count);
+		return FVector4f(SumValue.X * InvSum, SumValue.Y * InvSum, SumValue.Z * InvSum, 1.0f);
+	};
+
+	Infill.ApplyInfill<FVector4f>(Image, NormalizeFunc);
+}
+
+bool FRenderCaptureOcclusionHandler::FSampleStats::operator==(const FSampleStats& Other) const
+{
+	return (NumValid == Other.NumValid) && (NumInvalid == Other.NumInvalid);
+}
+
+bool FRenderCaptureOcclusionHandler::FSampleStats::operator!=(const FSampleStats& Other) const
+{
+	return !(*this == Other);
+}
+
+FRenderCaptureOcclusionHandler::FSampleStats& FRenderCaptureOcclusionHandler::FSampleStats::operator+=(const FSampleStats& Other)
+{
+	NumValid += Other.NumValid;
+	NumInvalid += Other.NumInvalid;
+	return *this;
+}
+
+FRenderCaptureOcclusionHandler::FSampleStats FRenderCaptureOcclusionHandler::FSampleStats::Zero()
+{
+	return FSampleStats{0, 0};
+}
+
+
+
+
+
+
+
+
+FSceneCapturePhotoSetSampler::FSceneCapturePhotoSetSampler(
+		FSceneCapturePhotoSet* SceneCapture,
+		float ValidSampleDepthThreshold,
+		const FDynamicMesh3* Mesh,
+		const FDynamicMeshAABBTree3* Spatial,
+		const FMeshTangentsd* Tangents) :
+	FMeshBakerDynamicMeshSampler(Mesh, Spatial, Tangents),
+	SceneCapture(SceneCapture),
+	ValidSampleDepthThreshold(ValidSampleDepthThreshold)
+{
+	check(SceneCapture != nullptr);
+	check(Mesh != nullptr);
+	check(Spatial != nullptr);
+	check(Tangents != nullptr);
+
+	bool bHasDepth = SceneCapture->GetCaptureTypeEnabled(ERenderCaptureType::DeviceDepth);
+	if (bHasDepth)
+	{
+		ensure(ValidSampleDepthThreshold > 0); // We only need the depth capture if this threshold is positive
+	}
+
+	// Compute an offset shift surface positions toward the render camera position so we don't 	immediately self intersect
+	// The Max expression ensures the code works when the base mesh is 2D along one coordinate direction
+	const double RayOffsetHackDist((double)(100.0 * FMathf::ZeroTolerance * FMath::Max(Mesh->GetBounds().MinDim(), 0.01)));
+
+	VisibilityFunction = [this, RayOffsetHackDist](const FVector3d& SurfPos, const FVector3d& ViewPos)
+	{
+		FVector3d RayDir = ViewPos - SurfPos;
+		double Dist = Normalize(RayDir);
+		FVector3d RayOrigin = SurfPos + RayOffsetHackDist * RayDir;
+		int32 HitTID = DetailSpatial->FindNearestHitTriangle(FRay3d(RayOrigin, RayDir), IMeshSpatial::FQueryOptions(Dist));
+		return (HitTID == IndexConstants::InvalidID);
+	};
+}
+
+bool FSceneCapturePhotoSetSampler::SupportsCustomCorrespondence() const
+{
+	return true;
+}
+
+void* FSceneCapturePhotoSetSampler::ComputeCustomCorrespondence(const FMeshUVSampleInfo& SampleInfo, FMeshMapEvaluator::FCorrespondenceSample& Sample) const
+{
+	
+	// Perform a ray-cast to determine which photo/coordinate, if any, should be sampled
+	int PhotoIndex;
+	FVector2d PhotoCoords;
+	SceneCapture->ComputeSampleLocation(
+		Sample.BaseSample.SurfacePoint,
+		Sample.BaseNormal,
+		ValidSampleDepthThreshold,
+		VisibilityFunction,
+		PhotoIndex,
+		PhotoCoords);
+
+	// Store the photo coordinates and index in the correspondence sample
+	Sample.DetailMesh = SceneCapture;
+	Sample.DetailTriID = PhotoIndex;
+	Sample.DetailBaryCoords.X = PhotoCoords.X;
+	Sample.DetailBaryCoords.Y = PhotoCoords.Y;
+
+	// This will be set to Sample.DetailMesh but we can already do that internally so it's kindof redundant
+	return SceneCapture;
+}
+
+bool FSceneCapturePhotoSetSampler::IsValidCorrespondence(const FMeshMapEvaluator::FCorrespondenceSample& Sample) const
+{
+	return Sample.DetailTriID != IndexConstants::InvalidID;
+}
+
+
+
+
+
+
+
+TUniquePtr<FSceneCapturePhotoSet> UE::Geometry::CapturePhotoSet(
+	const TArray<TObjectPtr<AActor>>& Actors,
+	const FRenderCaptureOptions& Options,
+	bool bAllowCancel
+)
+{
+	TRACE_CPUPROFILER_EVENT_SCOPE(CapturePhotoSet);
+
+	double FieldOfView = Options.FieldOfViewDegrees;
+	double NearPlaneDist = Options.NearPlaneDist;
+
+	FImageDimensions CaptureDimensions(Options.RenderCaptureImageSize, Options.RenderCaptureImageSize);
+
+	TUniquePtr<FSceneCapturePhotoSet> SceneCapture = MakeUnique<FSceneCapturePhotoSet>();
+	SceneCapture->SetAllowCancel(bAllowCancel);
+
+	SceneCapture->SetCaptureTypeEnabled(ERenderCaptureType::DeviceDepth, Options.bBakeDeviceDepth);
+	SceneCapture->SetCaptureTypeEnabled(ERenderCaptureType::BaseColor,   Options.bBakeBaseColor);
+	SceneCapture->SetCaptureTypeEnabled(ERenderCaptureType::WorldNormal, Options.bBakeNormalMap);
+	SceneCapture->SetCaptureTypeEnabled(ERenderCaptureType::Emissive,    Options.bBakeEmissive);
+	SceneCapture->SetCaptureTypeEnabled(ERenderCaptureType::CombinedMRS, Options.bUsePackedMRS);
+	SceneCapture->SetCaptureTypeEnabled(ERenderCaptureType::Metallic,    Options.bBakeMetallic);
+	SceneCapture->SetCaptureTypeEnabled(ERenderCaptureType::Roughness,   Options.bBakeRoughness);
+	SceneCapture->SetCaptureTypeEnabled(ERenderCaptureType::Specular,    Options.bBakeSpecular);
+
+	FRenderCaptureConfig Config;
+	Config.bAntiAliasing = Options.bAntiAliasing;
+	SceneCapture->SetCaptureConfig(ERenderCaptureType::BaseColor,   Config);
+	SceneCapture->SetCaptureConfig(ERenderCaptureType::WorldNormal, Config);
+	SceneCapture->SetCaptureConfig(ERenderCaptureType::CombinedMRS, Config);
+	SceneCapture->SetCaptureConfig(ERenderCaptureType::Metallic,    Config);
+	SceneCapture->SetCaptureConfig(ERenderCaptureType::Roughness,   Config);
+	SceneCapture->SetCaptureConfig(ERenderCaptureType::Specular,    Config);
+	SceneCapture->SetCaptureConfig(ERenderCaptureType::Emissive,    Config);
+	
+	SceneCapture->SetCaptureSceneActors(Actors[0]->GetWorld(), Actors);
+
+	SceneCapture->AddStandardExteriorCapturesFromBoundingBox(
+		CaptureDimensions, FieldOfView, NearPlaneDist,
+		true, true, true, true, true);
+
+	return SceneCapture;
+}
+
+
+
+
+
+template <ERenderCaptureType CaptureType>
+TSharedPtr<FRenderCaptureMapEvaluator<FVector4f>>
+MakeColorEvaluator(
+	const FSceneCapturePhotoSet::FSceneSample& DefaultSample,
+	const FSceneCapturePhotoSet* SceneCapture)
+{
+	TSharedPtr<FRenderCaptureMapEvaluator<FVector4f>> Evaluator = MakeShared<FRenderCaptureMapEvaluator<FVector4f>>();
+
+	switch (CaptureType) {
+	case ERenderCaptureType::BaseColor:
+		Evaluator->Channel = ERenderCaptureChannel::BaseColor;
+		break;
+	case ERenderCaptureType::Roughness:
+		Evaluator->Channel = ERenderCaptureChannel::Roughness;
+		break;
+	case ERenderCaptureType::Metallic:
+		Evaluator->Channel = ERenderCaptureChannel::Metallic;
+		break;
+	case ERenderCaptureType::Specular:
+		Evaluator->Channel = ERenderCaptureChannel::Specular;
+		break;
+	case ERenderCaptureType::Emissive:
+		Evaluator->Channel = ERenderCaptureChannel::Emissive;
+		break;
+	case ERenderCaptureType::WorldNormal:
+		Evaluator->Channel = ERenderCaptureChannel::WorldNormal;
+		break;
+	case ERenderCaptureType::CombinedMRS:
+		Evaluator->Channel = ERenderCaptureChannel::CombinedMRS;
+		break;
+	case ERenderCaptureType::DeviceDepth:
+		Evaluator->Channel = ERenderCaptureChannel::DeviceDepth;
+		break;
+	}
+
+	Evaluator->DefaultResult = DefaultSample.GetValue4f(CaptureType);
+
+	Evaluator->EvaluateSampleCallback = [DefaultSample, SceneCapture](const FMeshMapEvaluator::FCorrespondenceSample& Sample)
+	{
+		const int PhotoIndex = Sample.DetailTriID;
+		const FVector2d PhotoCoords(Sample.DetailBaryCoords.X, Sample.DetailBaryCoords.Y);
+		const FVector4f SampleColor = SceneCapture->ComputeSampleNearest<CaptureType>(PhotoIndex, PhotoCoords, DefaultSample);
+		return SampleColor;
+	};
+
+	Evaluator->EvaluateColorCallback = [](const int DataIdx, float*& In)
+	{
+		const FVector4f Out(In[0], In[1], In[2], In[3]);
+		In += 4;
+		return Out;
+	};
+
+	return Evaluator;
+}
+
+
+TUniquePtr<FMeshMapBaker> UE::Geometry::MakeRenderCaptureBaker(
+	FDynamicMesh3* BaseMesh,
+	TSharedPtr<UE::Geometry::FMeshTangentsd, ESPMode::ThreadSafe> BaseMeshTangents,
+	FSceneCapturePhotoSet* SceneCapture,
+	FSceneCapturePhotoSetSampler* Sampler,
+	FRenderCaptureOptions Options,
+	EBakeTextureResolution TextureImageSize,
+	EBakeTextureSamplesPerPixel SamplesPerPixel,
+	FRenderCaptureOcclusionHandler* OcclusionHandler)
+{
+	TRACE_CPUPROFILER_EVENT_SCOPE(MakeRenderCaptureBaker);
+
+	check(BaseMesh != nullptr);
+	check(BaseMeshTangents.IsValid());
+	check(SceneCapture != nullptr);
+	check(Sampler != nullptr);
+	check(OcclusionHandler != nullptr);
+
+	auto RegisterSampleStats = [OcclusionHandler](bool bSampleValid, const FMeshMapEvaluator::FCorrespondenceSample& Sample, const FVector2d& UVPosition, const FVector2i& ImageCoords)
+	{
+		OcclusionHandler->RegisterSample(ImageCoords, bSampleValid);
+	};
+
+	auto ComputeAndApplyInfill = [OcclusionHandler](TArray<TUniquePtr<TImageBuilder<FVector4f>>>& BakeResults)
+	{
+		OcclusionHandler->ComputeAndApplyInfill(BakeResults);
+	};
+
+	TUniquePtr<FMeshMapBaker> Result = MakeUnique<FMeshMapBaker>();
+	Result->SetTargetMesh(BaseMesh);
+	Result->SetTargetMeshTangents(BaseMeshTangents);
+	Result->SetDimensions(FImageDimensions(static_cast<int32>(TextureImageSize), static_cast<int32>(TextureImageSize)));
+	Result->SetSamplesPerPixel(static_cast<int32>(SamplesPerPixel));
+	Result->SetFilter(FMeshMapBaker::EBakeFilterType::BSpline);
+	Result->SetTargetMeshUVLayer(Options.TargetUVLayer);
+	Result->InteriorSampleCallback = RegisterSampleStats;
+	Result->PostWriteToImageCallback = ComputeAndApplyInfill;
+	Result->SetDetailSampler(Sampler);
+	Result->SetCorrespondenceStrategy(FMeshMapBaker::ECorrespondenceStrategy::Custom);
+
+	// Pixels in the output textures which don't map onto the mesh have a light grey color (except the normal map which
+	// will show a color corresponding to a unit z tangent space normal)
+	const FVector4f InvalidColor(.42, .42, .42, 1);
+	const FVector3f DefaultNormal = FVector3f::UnitZ();
+
+	FSceneCapturePhotoSet::FSceneSample DefaultColorSample;
+	DefaultColorSample.BaseColor = FVector3f(InvalidColor.X, InvalidColor.Y, InvalidColor.Z);
+	DefaultColorSample.Roughness = InvalidColor.X;
+	DefaultColorSample.Specular = InvalidColor.X;
+	DefaultColorSample.Metallic = InvalidColor.X;
+	DefaultColorSample.Emissive = FVector3f(InvalidColor.X, InvalidColor.Y, InvalidColor.Z);
+	DefaultColorSample.WorldNormal = FVector4f((DefaultNormal + FVector3f::One()) * .5f, InvalidColor.W);
+	
+	// We use 0 here since this corresponds to the infinite far plane value. Also, we don't preview the depth texture
+	DefaultColorSample.DeviceDepth = 0;
+
+	auto AddColorEvaluator = [&Result, OcclusionHandler] (const TSharedPtr<FRenderCaptureMapEvaluator<FVector4f>>& Evaluator)
+	{
+		Result->AddEvaluator(Evaluator);
+		OcclusionHandler->PushInfillRequired(true);
+	};
+	
+	if (Options.bBakeDeviceDepth)
+	{
+		AddColorEvaluator(MakeColorEvaluator<ERenderCaptureType::DeviceDepth>(DefaultColorSample, SceneCapture));
+	}
+	if (Options.bBakeBaseColor)
+	{
+		AddColorEvaluator(MakeColorEvaluator<ERenderCaptureType::BaseColor>(DefaultColorSample, SceneCapture));
+	}
+	if (Options.bUsePackedMRS)
+	{
+		AddColorEvaluator(MakeColorEvaluator<ERenderCaptureType::CombinedMRS>(DefaultColorSample, SceneCapture));
+	}
+	if (Options.bBakeRoughness)
+	{
+		AddColorEvaluator(MakeColorEvaluator<ERenderCaptureType::Roughness>(DefaultColorSample, SceneCapture));
+	}
+	if (Options.bBakeMetallic)
+	{
+		AddColorEvaluator(MakeColorEvaluator<ERenderCaptureType::Metallic>(DefaultColorSample, SceneCapture));
+	}
+	if (Options.bBakeSpecular)
+	{
+		AddColorEvaluator(MakeColorEvaluator<ERenderCaptureType::Specular>(DefaultColorSample, SceneCapture));
+	}
+	if (Options.bBakeEmissive)
+	{
+		AddColorEvaluator(MakeColorEvaluator<ERenderCaptureType::Emissive>(DefaultColorSample, SceneCapture));
+	}
+	if (Options.bBakeNormalMap)
+	{
+		TSharedPtr<FRenderCaptureMapEvaluator<FVector3f>> Evaluator = MakeShared<FRenderCaptureMapEvaluator<FVector3f>>();
+
+		Evaluator->Channel = ERenderCaptureChannel::WorldNormal;
+
+		Evaluator->DefaultResult = DefaultNormal;
+
+		Evaluator->EvaluateSampleCallback = [SceneCapture, BaseMeshTangents, DefaultColorSample](const FMeshMapEvaluator::FCorrespondenceSample& Sample)
+		{
+			const int32 TriangleID = Sample.BaseSample.TriangleIndex;
+			const FVector3d BaryCoords = Sample.BaseSample.BaryCoords;
+			const int PhotoIndex = Sample.DetailTriID;
+			const FVector2d PhotoCoords(Sample.DetailBaryCoords.X, Sample.DetailBaryCoords.Y);
+
+			const FVector4f NormalColor = SceneCapture->ComputeSampleNearest<ERenderCaptureType::WorldNormal>(PhotoIndex, PhotoCoords, DefaultColorSample);
+
+			// Map from color components [0,1] to normal components [-1,1]
+			const FVector3f WorldSpaceNormal(
+				(NormalColor.X - 0.5f) * 2.0f,
+				(NormalColor.Y - 0.5f) * 2.0f,
+				(NormalColor.Z - 0.5f) * 2.0f);
+
+			// Get tangents on base mesh
+			FVector3d BaseTangentX, BaseTangentY;
+			BaseMeshTangents->GetInterpolatedTriangleTangent(TriangleID, BaryCoords, BaseTangentX, BaseTangentY);
+
+			// Compute normal in tangent space
+			const FVector3f TangentSpaceNormal(
+					(float)WorldSpaceNormal.Dot(FVector3f(BaseTangentX)),
+					(float)WorldSpaceNormal.Dot(FVector3f(BaseTangentY)),
+					(float)WorldSpaceNormal.Dot(FVector3f(Sample.BaseNormal)));
+
+			return TangentSpaceNormal;
+		};
+
+		Evaluator->EvaluateColorCallback = [](const int DataIdx, float*& In)
+		{
+			// Map normal space [-1,1] to color space [0,1]
+			const FVector3f Normal(In[0], In[1], In[2]);
+			const FVector3f Color = (Normal + FVector3f::One()) * 0.5f;
+			const FVector4f Out(Color.X, Color.Y, Color.Z, 1.0f);
+			In += 3;
+			return Out;
+		};
+
+		Result->AddEvaluator(Evaluator);
+
+		// Note: No infill on normal map for now, doesn't make sense to do after mapping to tangent space!
+		//  (should we build baked normal map in world space, and then resample to tangent space??)
+		OcclusionHandler->PushInfillRequired(false);
+	}
+
+	return Result;
+}
+
+
+void UE::Geometry::GetTexturesFromRenderCaptureBaker(const TUniquePtr<FMeshMapBaker>& Baker, FRenderCaptureTextures& TexturesOut)
+{
+	// We do this to defer work I guess, it was like this in the original ApproximateActors implementation :DeferredPopulateSourceData 
+	constexpr bool bPopulateSourceData = false;
+
+	const int32 NumEval = Baker->NumEvaluators();
+	for (int32 EvalIdx = 0; EvalIdx < NumEval; ++EvalIdx)
+	{
+		FMeshMapEvaluator* BaseEval = Baker->GetEvaluator(EvalIdx);
+		check(BaseEval->DataLayout().Num() == 1);
+		switch (BaseEval->DataLayout()[0])
+		{
+		case FMeshMapEvaluator::EComponents::Float4:
+
+			{
+				FRenderCaptureMapEvaluator<FVector4f>* Eval = static_cast<FRenderCaptureMapEvaluator<FVector4f>*>(BaseEval);
+				TUniquePtr<TImageBuilder<FVector4f>> ImageBuilder = MoveTemp(Baker->GetBakeResults(EvalIdx)[0]);
+				
+				if (ensure(ImageBuilder.IsValid()) == false) return;
+
+				switch (Eval->Channel)
+				{
+				case ERenderCaptureChannel::BaseColor:
+
+					TexturesOut.BaseColorMap = FTexture2DBuilder::BuildTextureFromImage(
+						*ImageBuilder,
+						FTexture2DBuilder::ETextureType::Color,
+						true,
+						bPopulateSourceData);
+					break;
+
+				case ERenderCaptureChannel::Roughness:
+
+					TexturesOut.RoughnessMap = FTexture2DBuilder::BuildTextureFromImage(
+						*ImageBuilder,
+						FTexture2DBuilder::ETextureType::Roughness,
+						false,
+						bPopulateSourceData);
+					break;
+
+				case ERenderCaptureChannel::Metallic:
+
+					TexturesOut.MetallicMap = FTexture2DBuilder::BuildTextureFromImage(
+						*ImageBuilder,
+						FTexture2DBuilder::ETextureType::Metallic,
+						false,
+						bPopulateSourceData);
+					break;
+
+				case ERenderCaptureChannel::Specular:
+
+					TexturesOut.SpecularMap = FTexture2DBuilder::BuildTextureFromImage(
+						*ImageBuilder,
+						FTexture2DBuilder::ETextureType::Specular,
+						false,
+						bPopulateSourceData);
+					break;
+
+				case ERenderCaptureChannel::Emissive:
+
+					TexturesOut.EmissiveMap = FTexture2DBuilder::BuildTextureFromImage(
+						*ImageBuilder,
+						FTexture2DBuilder::ETextureType::EmissiveHDR,
+						false,
+						bPopulateSourceData);
+					TexturesOut.EmissiveMap->CompressionSettings = TC_HDR_Compressed;
+					break;
+
+				case ERenderCaptureChannel::CombinedMRS:
+
+					TexturesOut.PackedMRSMap = FTexture2DBuilder::BuildTextureFromImage(
+						*ImageBuilder,
+						FTexture2DBuilder::ETextureType::ColorLinear,
+						false,
+						bPopulateSourceData);
+					break;
+
+				case ERenderCaptureChannel::DeviceDepth:
+
+					// Add a null pointer for this evaluator, the depth capture is used internally and not presented to the user
+					break;
+
+				case ERenderCaptureChannel::WorldNormal:
+
+					// This should be handled in the Float3 branch
+					ensure(false);
+					return;
+
+				default:
+					ensure(false);
+					return;
+				}
+
+			} break; // Float4
+
+		case FMeshMapEvaluator::EComponents::Float3:
+
+			{
+				FRenderCaptureMapEvaluator<FVector3f>* Eval = static_cast<FRenderCaptureMapEvaluator<FVector3f>*>(BaseEval);
+				TUniquePtr<TImageBuilder<FVector4f>> ImageBuilder = MoveTemp(Baker->GetBakeResults(EvalIdx)[0]);
+				
+				if (ensure(ImageBuilder.IsValid()) == false) return;
+				if (ensure(Eval->Channel == ERenderCaptureChannel::WorldNormal) == false) return;
+
+				TexturesOut.NormalMap = FTexture2DBuilder::BuildTextureFromImage(
+					*ImageBuilder,
+					FTexture2DBuilder::ETextureType::NormalMap,
+					false,
+					bPopulateSourceData);
+
+			} break; // Float3
+			
+		default:
+			ensure(false);
+			return;
+		}
+	}
+}
diff --git a/Engine/Plugins/Runtime/MeshModelingToolset/Source/ModelingComponents/Private/Scene/WorldRenderCapture.cpp b/Engine/Plugins/Runtime/MeshModelingToolset/Source/ModelingComponents/Private/Scene/WorldRenderCapture.cpp
--- a/Engine/Plugins/Runtime/MeshModelingToolset/Source/ModelingComponents/Private/Scene/WorldRenderCapture.cpp
+++ b/Engine/Plugins/Runtime/MeshModelingToolset/Source/ModelingComponents/Private/Scene/WorldRenderCapture.cpp
@@ -183,6 +183,7 @@ void FWorldRenderCapture::SetVisibleActors(const TArray<AActor*>& Actors)
 
 	// Find all components that need to be included in rendering.
 	// This also descends into any ChildActorComponents
+	// TODO Consider using Actor->GetActorBounds() or ForEachComponent here
 	TArray<UActorComponent*> ComponentQueue;
 	for (AActor* Actor : Actors)
 	{
diff --git a/Engine/Plugins/Runtime/MeshModelingToolset/Source/ModelingComponents/Public/Baking/BakingTypes.h b/Engine/Plugins/Runtime/MeshModelingToolset/Source/ModelingComponents/Public/Baking/BakingTypes.h
new file mode 100644
--- /dev/null
+++ b/Engine/Plugins/Runtime/MeshModelingToolset/Source/ModelingComponents/Public/Baking/BakingTypes.h
@@ -0,0 +1,44 @@
+// Copyright Epic Games, Inc. All Rights Reserved.
+
+#pragma once
+
+/**
+* Bake maps enums
+*/
+
+#include "CoreMinimal.h"
+
+UENUM()
+enum class EBakeTextureResolution 
+{
+	Resolution16 = 16 UMETA(DisplayName = "16 x 16"),
+	Resolution32 = 32 UMETA(DisplayName = "32 x 32"),
+	Resolution64 = 64 UMETA(DisplayName = "64 x 64"),
+	Resolution128 = 128 UMETA(DisplayName = "128 x 128"),
+	Resolution256 = 256 UMETA(DisplayName = "256 x 256"),
+	Resolution512 = 512 UMETA(DisplayName = "512 x 512"),
+	Resolution1024 = 1024 UMETA(DisplayName = "1024 x 1024"),
+	Resolution2048 = 2048 UMETA(DisplayName = "2048 x 2048"),
+	Resolution4096 = 4096 UMETA(DisplayName = "4096 x 4096"),
+	Resolution8192 = 8192 UMETA(DisplayName = "8192 x 8192")
+};
+
+
+UENUM()
+enum class EBakeTextureBitDepth
+{
+	ChannelBits8 UMETA(DisplayName = "8 bits/channel"),
+	ChannelBits16 UMETA(DisplayName = "16 bits/channel")
+};
+
+
+UENUM()
+enum class EBakeTextureSamplesPerPixel
+{
+	Sample1 = 1 UMETA(DisplayName = "1"),
+	Sample4 = 4 UMETA(DisplayName = "4"),
+	Sample16 = 16 UMETA(DisplayName = "16"),
+	Sample64 = 64 UMETA(DisplayName = "64"),
+	Sample256 = 256 UMETA(DisplayName = "256")
+};
+
diff --git a/Engine/Plugins/Runtime/MeshModelingToolset/Source/ModelingComponents/Public/Baking/RenderCaptureFunctions.h b/Engine/Plugins/Runtime/MeshModelingToolset/Source/ModelingComponents/Public/Baking/RenderCaptureFunctions.h
new file mode 100644
--- /dev/null
+++ b/Engine/Plugins/Runtime/MeshModelingToolset/Source/ModelingComponents/Public/Baking/RenderCaptureFunctions.h
@@ -0,0 +1,184 @@
+﻿// Copyright Epic Games, Inc. All Rights Reserved.
+
+#include "Scene/SceneCapturePhotoSet.h"
+#include "Sampling/MeshBakerCommon.h"
+#include "Image/ImageInfilling.h"
+#include "Baking/BakingTypes.h"
+#include "DynamicMesh/MeshTangents.h"
+#include "CoreMinimal.h"
+
+class UTexture2D;
+class AActor;
+
+namespace UE
+{
+namespace Geometry
+{
+
+class FMeshMapBaker;
+
+
+
+
+
+
+class MODELINGCOMPONENTS_API FRenderCaptureOcclusionHandler
+{
+public:
+	FRenderCaptureOcclusionHandler(FImageDimensions Dimensions);
+
+	void RegisterSample(const FVector2i& ImageCoords, bool bSampleValid);
+
+	void PushInfillRequired(bool bInfillRequired);
+
+	void ComputeAndApplyInfill(TArray<TUniquePtr<TImageBuilder<FVector4f>>>& Images);
+
+private:
+
+	struct FSampleStats {
+		uint16 NumValid = 0;
+		uint16 NumInvalid = 0;
+
+		// These are required by the TMarchingPixelInfill implementation
+		bool operator==(const FSampleStats& Other) const;
+		bool operator!=(const FSampleStats& Other) const;
+		FSampleStats& operator+=(const FSampleStats& Other);
+		static FSampleStats Zero();
+	};
+
+	void ComputeInfill();
+
+	void ApplyInfill(TImageBuilder<FVector4f>& Image) const;
+
+	// Collect some sample stats per pixel, used to determine if a pixel requires infill or not
+	TImageBuilder<FSampleStats> SampleStats;
+
+	// InfillRequire[i] indicates if the i-th image passed to ComputeAndApplyInfill needs infill
+	TArray<bool> InfillRequired;
+
+	TMarchingPixelInfill<FSampleStats> Infill;
+};
+
+
+
+
+
+class MODELINGCOMPONENTS_API FSceneCapturePhotoSetSampler : public FMeshBakerDynamicMeshSampler
+{
+public:
+	FSceneCapturePhotoSetSampler(
+		FSceneCapturePhotoSet* SceneCapture,
+		float ValidSampleDepthThreshold,
+		const FDynamicMesh3* Mesh,
+		const FDynamicMeshAABBTree3* Spatial,
+		const FMeshTangentsd* Tangents); 
+
+	virtual bool SupportsCustomCorrespondence() const override;
+
+	// Warning: Expects that Sample.BaseSample.SurfacePoint and Sample.BaseNormal are set when the function is called
+	virtual void* ComputeCustomCorrespondence(const FMeshUVSampleInfo& SampleInfo, FMeshMapEvaluator::FCorrespondenceSample& Sample) const override;
+
+	virtual bool IsValidCorrespondence(const FMeshMapEvaluator::FCorrespondenceSample& Sample) const override;
+
+private:
+	FSceneCapturePhotoSet* SceneCapture = nullptr;
+	float ValidSampleDepthThreshold = 0;
+	TFunction<bool(const FVector3d&, const FVector3d&)> VisibilityFunction;
+};
+
+
+
+
+
+
+struct MODELINGCOMPONENTS_API FRenderCaptureOptions
+{
+	//
+	// Material approximation settings
+	//
+
+	int32 RenderCaptureImageSize = 1024;
+	bool bAntiAliasing = false;
+
+	// render capture parameters
+	double FieldOfViewDegrees = 45.0;
+	double NearPlaneDist = 1.0;
+	double ValidSampleDepthThreshold = 0;
+
+	//
+	// Material output settings
+	//
+
+	bool bBakeBaseColor = true;
+	bool bBakeRoughness = true;
+	bool bBakeMetallic = true;
+	bool bBakeSpecular = true;
+	bool bBakeEmissive = true;
+	bool bBakeNormalMap = true;
+	bool bBakeDeviceDepth = true;
+	
+	bool bUsePackedMRS = true;
+
+	//
+	// Mesh settings
+	//
+
+	//  Which UV layer of the Target mesh (the one we're baking to) should be used
+	int32 TargetUVLayer = 0;
+
+	//
+	// For internal use only
+	//
+
+	// A new MIC derived from this material will be created and assigned to the generated mesh
+	// if null, will use /MeshModelingToolsetExp/Materials/FullMaterialBakePreviewMaterial_PackedMRS instead
+	UMaterialInterface* BakeMaterial = nullptr;
+};
+
+
+
+MODELINGCOMPONENTS_API
+TUniquePtr<FSceneCapturePhotoSet> CapturePhotoSet(
+	const TArray<TObjectPtr<AActor>>& Actors,
+	const FRenderCaptureOptions& Options,
+	bool bAllowCancel);
+
+
+
+// Return a render capture baker, note the lifetime of all arguments such match the lifetime of the returned baker
+MODELINGCOMPONENTS_API
+TUniquePtr<FMeshMapBaker> MakeRenderCaptureBaker(
+	FDynamicMesh3* BaseMesh,
+	TSharedPtr<FMeshTangentsd, ESPMode::ThreadSafe> BaseMeshTangents,
+	FSceneCapturePhotoSet* SceneCapture,
+	FSceneCapturePhotoSetSampler* Sampler,
+	FRenderCaptureOptions Options,
+	EBakeTextureResolution TextureImageSize,
+	EBakeTextureSamplesPerPixel SamplesPerPixel,
+	FRenderCaptureOcclusionHandler* OcclusionHandler);
+
+
+
+
+
+
+struct MODELINGCOMPONENTS_API FRenderCaptureTextures
+{
+	UTexture2D* BaseColorMap = nullptr;
+	UTexture2D* NormalMap = nullptr;
+	UTexture2D* PackedMRSMap = nullptr;
+	UTexture2D* MetallicMap = nullptr;
+	UTexture2D* RoughnessMap = nullptr;
+	UTexture2D* SpecularMap = nullptr;
+	UTexture2D* EmissiveMap = nullptr;
+};
+
+// Note: The source data in the textures is *not* updated by this function
+MODELINGCOMPONENTS_API
+void GetTexturesFromRenderCaptureBaker(
+	const TUniquePtr<FMeshMapBaker>& Baker,
+	FRenderCaptureTextures& TexturesOut);
+
+
+} // end namespace Geometry
+} // end namespace UE
\ No newline at end of file
