ModelingComponents: add FGroupVisualizationCache, this is a helper object that stores precomputed info about each polygroup in a mesh, used for visualization in various Tools
ModelingTools: add the ability to draw IDs of visible polygroups in UMeshGroupPaintTool, as well as draw the ID of the polygroup under the cursor. Add the ability to draw Tri/Edge/Vert/Group IDs in the Inspector Tool. Also expand the set of override render materials, and add a polygroup layer picker.

diff --git a/Engine/Plugins/Experimental/MeshModelingToolsetExp/Source/MeshModelingToolsExp/Private/MeshGroupPaintTool.cpp b/Engine/Plugins/Experimental/MeshModelingToolsetExp/Source/MeshModelingToolsExp/Private/MeshGroupPaintTool.cpp
--- a/Engine/Plugins/Experimental/MeshModelingToolsetExp/Source/MeshModelingToolsExp/Private/MeshGroupPaintTool.cpp
+++ b/Engine/Plugins/Experimental/MeshModelingToolsetExp/Source/MeshModelingToolsExp/Private/MeshGroupPaintTool.cpp
@@ -100,7 +100,7 @@ void UMeshGroupPaintTool::Setup()
 	}
 
 	DynamicMeshComponent->SetInvalidateProxyOnChangeEnabled(false);
-	OnDynamicMeshComponentChangedHandle = DynamicMeshComponent->OnMeshVerticesChanged.AddUObject(this, &UMeshGroupPaintTool::OnDynamicMeshComponentChanged);
+	OnDynamicMeshComponentChangedHandle = DynamicMeshComponent->OnMeshChanged.AddUObject(this, &UMeshGroupPaintTool::OnDynamicMeshComponentChanged);
 
 	FDynamicMesh3* Mesh = GetSculptMesh();
 	Mesh->EnableVertexColors(FVector3f::One());
@@ -637,6 +637,7 @@ bool UMeshGroupPaintTool::SyncMeshWithGroupBuffer(FDynamicMesh3* Mesh)
 			NumModified++;
 		}
 	}
+
 	return (NumModified > 0);
 }
 
@@ -1071,7 +1072,6 @@ bool UMeshGroupPaintTool::OnUpdateHover(const FInputDeviceRay& DevicePos)
 	return true;
 }
 
-
 void UMeshGroupPaintTool::DrawHUD(FCanvas* Canvas, IToolsContextRenderAPI* RenderAPI)
 {
 	if (PolyLassoMechanic)
@@ -1080,6 +1080,54 @@ void UMeshGroupPaintTool::DrawHUD(FCanvas* Canvas, IToolsContextRenderAPI* Rende
 		PolyLassoMechanic->LineColor = GetInEraseStroke() ? FLinearColor::Red : FLinearColor::Green;
 		PolyLassoMechanic->DrawHUD(Canvas, RenderAPI);
 	}
+
+
+	float DPIScale = Canvas->GetDPIScale();
+	UFont* UseFont = GEngine->GetSmallFont();
+	FViewCameraState CamState = RenderAPI->GetCameraState();
+	const FSceneView* SceneView = RenderAPI->GetSceneView();
+	FVector3d LocalEyePosition(CurTargetTransform.InverseTransformPosition((FVector3d)CamState.Position));
+
+	FDynamicMesh3* Mesh = GetSculptMesh();
+
+	// draw the group number for the group under the cursor
+	if (FilterProperties->bShowHitGroup)
+	{
+		int CursorHitTID = Octree.FindNearestHitObject(FRay3d(LocalEyePosition, Normalized(HoverStamp.LocalFrame.Origin - LocalEyePosition)));
+		if (Mesh->IsTriangle(CursorHitTID))
+		{
+			int32 GroupID = ActiveGroupSet->GetTriangleGroup(CursorHitTID);
+			FVector2D CursorPixelPos;
+			SceneView->WorldToPixel(HoverStamp.WorldFrame.Origin, CursorPixelPos);
+			FString CursorString = FString::Printf(TEXT("%d"), GroupID);
+			Canvas->DrawShadowedString(CursorPixelPos.X / (double)DPIScale, CursorPixelPos.Y / (double)DPIScale, *CursorString, UseFont, FLinearColor::White);
+		}
+	}
+
+	// draw the group number for all groups
+	if (FilterProperties->bShowAllGroups)
+	{
+		if (bDrawGroupsDataValid == false)
+		{
+			GroupVisualizationCache.UpdateGroupInfo_ConnectedComponents(*Mesh, *ActiveGroupSet);
+			bDrawGroupsDataValid = true;
+		}
+		for (const FGroupVisualizationCache::FGroupInfo& GroupInfo : GroupVisualizationCache)
+		{
+			FRay3d EyeRay;
+			EyeRay.Origin = LocalEyePosition;
+			EyeRay.Direction = Normalized(GroupInfo.Center - LocalEyePosition);
+			int HitTID = Octree.FindNearestHitObject(EyeRay);
+			if (HitTID == GroupInfo.CenterTris.A || HitTID == GroupInfo.CenterTris.B)
+			{
+				FVector2D PixelPos;
+				SceneView->WorldToPixel(CurTargetTransform.TransformPosition(GroupInfo.Center), PixelPos);
+				FString String = FString::Printf(TEXT("%d"), GroupInfo.GroupID);
+				Canvas->DrawShadowedString(PixelPos.X / (double)DPIScale, PixelPos.Y / (double)DPIScale, *String, UseFont, FLinearColor::White);
+			}
+		}
+	}
+
 }
 
 
@@ -1112,6 +1160,7 @@ void UMeshGroupPaintTool::OnTick(float DeltaTime)
 		// post rendering update
 		DynamicMeshComponent->FastNotifyTriangleVerticesUpdated(AccumulatedTriangleROI, EMeshRenderAttributeFlags::VertexColors);
 		GetToolManager()->PostInvalidation();
+		bDrawGroupsDataValid = false;
 
 		// ignore stamp and wait for next tick to do anything else
 		bUndoUpdatePending = false;
@@ -1186,9 +1235,6 @@ void UMeshGroupPaintTool::OnTick(float DeltaTime)
 
 
 
-
-
-
 void UMeshGroupPaintTool::AllocateNewGroupAndSetAsCurrentAction()
 {
 	FilterProperties->SetGroup = ActiveGroupSet->MaxGroupID;
@@ -1457,6 +1503,9 @@ void UMeshGroupPaintTool::EndChange()
 	};
 
 	GetToolManager()->EmitObjectChange(DynamicMeshComponent, MoveTemp(NewChange), LOCTEXT("GroupPaintChange", "Group Stroke"));
+
+	// debug groups are invalid now
+	bDrawGroupsDataValid = false;
 }
 
 
@@ -1468,7 +1517,7 @@ void UMeshGroupPaintTool::WaitForPendingUndoRedo()
 	}
 }
 
-void UMeshGroupPaintTool::OnDynamicMeshComponentChanged(UDynamicMeshComponent* Component, const FMeshVertexChange* Change, bool bRevert)
+void UMeshGroupPaintTool::OnDynamicMeshComponentChanged()
 {
 	// update octree
 	FDynamicMesh3* Mesh = GetSculptMesh();
@@ -1479,13 +1528,13 @@ void UMeshGroupPaintTool::OnDynamicMeshComponentChanged(UDynamicMeshComponent* C
 		// we should never hit this anymore, because of pre-change calling WaitForPendingUndoRedo()
 		WaitForPendingUndoRedo();
 
-		// this is not right because now we are going to do extra recomputation, but it's very messy otherwise...
-		UE::Geometry::VertexToTriangleOneRing(Mesh, Change->Vertices, AccumulatedTriangleROI);
+		// TODO: do we need to read from mesh change here??
+		//UE::Geometry::VertexToTriangleOneRing(Mesh, Change->Vertices, AccumulatedTriangleROI);
 	}
 	else
 	{
-		AccumulatedTriangleROI.Reset();
-		UE::Geometry::VertexToTriangleOneRing(Mesh, Change->Vertices, AccumulatedTriangleROI);
+		// TODO: do we need to read from mesh change here??
+		//UE::Geometry::VertexToTriangleOneRing(Mesh, Change->Vertices, AccumulatedTriangleROI);
 	}
 
 	// note that we have a pending update
diff --git a/Engine/Plugins/Experimental/MeshModelingToolsetExp/Source/MeshModelingToolsExp/Private/MeshInspectorTool.cpp b/Engine/Plugins/Experimental/MeshModelingToolsetExp/Source/MeshModelingToolsExp/Private/MeshInspectorTool.cpp
--- a/Engine/Plugins/Experimental/MeshModelingToolsetExp/Source/MeshModelingToolsExp/Private/MeshInspectorTool.cpp
+++ b/Engine/Plugins/Experimental/MeshModelingToolsetExp/Source/MeshModelingToolsExp/Private/MeshInspectorTool.cpp
@@ -9,6 +9,7 @@
 #include "DynamicMesh/MeshNormals.h"
 
 #include "Components/DynamicMeshComponent.h"
+#include "Materials/Material.h"
 #include "Materials/MaterialInstanceDynamic.h"
 
 #include "Properties/MeshStatisticsProperties.h"
@@ -16,6 +17,8 @@
 
 #include "Drawing/LineSetComponent.h"
 #include "ToolDataVisualizer.h"
+#include "Polygroups/PolygroupUtil.h"
+#include "Util/ColorConstants.h"
 
 #include "SceneManagement.h" // for FPrimitiveDrawInterface
 
@@ -23,12 +26,136 @@
 #include "AssetUtils/MeshDescriptionUtil.h"
 #include "ModelingToolTargetUtil.h"
 
+#include "CanvasTypes.h"
+#include "CanvasItem.h"
+
 #include UE_INLINE_GENERATED_CPP_BY_NAME(MeshInspectorTool)
 
 using namespace UE::Geometry;
 
 #define LOCTEXT_NAMESPACE "UMeshInspectorTool"
 
+
+
+
+
+void UMeshInspectorMaterialProperties::UpdateUVChannels(int32 UVChannelIndex, const TArray<FString>& UVChannelNames, bool bUpdateSelection)
+{
+	UVChannelNamesList = UVChannelNames;
+	if (bUpdateSelection)
+	{
+		UVChannel = 0 <= UVChannelIndex && UVChannelIndex < UVChannelNames.Num() ? UVChannelNames[UVChannelIndex] : TEXT("");
+	}
+}
+
+
+const TArray<FString>& UMeshInspectorMaterialProperties::GetUVChannelNamesFunc() const
+{
+	return UVChannelNamesList;
+}
+
+void UMeshInspectorMaterialProperties::RestoreProperties(UInteractiveTool* RestoreToTool, const FString& CacheIdentifier)
+{
+	Super::RestoreProperties(RestoreToTool, CacheIdentifier);
+	Setup();
+}
+
+void UMeshInspectorMaterialProperties::Setup()
+{
+	UMaterial* CheckerMaterialBase = LoadObject<UMaterial>(nullptr, TEXT("/MeshModelingToolsetExp/Materials/CheckerMaterial"));
+	if (CheckerMaterialBase != nullptr)
+	{
+		CheckerMaterial = UMaterialInstanceDynamic::Create(CheckerMaterialBase, this);
+		if (CheckerMaterial != nullptr)
+		{
+			CheckerMaterial->SetScalarParameterValue("Density", CheckerDensity);
+			CheckerMaterial->SetScalarParameterValue("UVChannel", static_cast<float>(UVChannelNamesList.IndexOfByKey(UVChannel)));
+		}
+	}
+}
+
+void UMeshInspectorMaterialProperties::UpdateMaterials()
+{
+	if (MaterialMode == EMeshInspectorMaterialMode::Checkerboard)
+	{
+		if (CheckerMaterial != nullptr)
+		{
+			CheckerMaterial->SetScalarParameterValue("Density", CheckerDensity);
+			CheckerMaterial->SetScalarParameterValue("UVChannel", static_cast<float>(UVChannelNamesList.IndexOfByKey(UVChannel)));
+		}
+	}
+	else if (MaterialMode == EMeshInspectorMaterialMode::FlatShaded)
+	{
+		ActiveCustomMaterial = UMaterialInstanceDynamic::Create(ToolSetupUtil::GetDefaultSculptMaterial(nullptr), this);
+		ActiveCustomMaterial->SetScalarParameterValue(TEXT("FlatShading"), 1.0f);
+	}
+	else if (MaterialMode == EMeshInspectorMaterialMode::Grey)
+	{
+		ActiveCustomMaterial = UMaterialInstanceDynamic::Create(ToolSetupUtil::GetImageBasedSculptMaterial(nullptr, ToolSetupUtil::ImageMaterialType::DefaultBasic), this);
+		ActiveCustomMaterial->SetScalarParameterValue(TEXT("FlatShading"), 0.0f);
+	}
+	else if (MaterialMode == EMeshInspectorMaterialMode::Transparent)
+	{
+		ActiveCustomMaterial = ToolSetupUtil::GetTransparentSculptMaterial(nullptr, TransparentMaterialColor, Opacity, bTwoSided);
+		ActiveCustomMaterial->SetVectorParameterValue(TEXT("Color"), Color);
+		ActiveCustomMaterial->SetScalarParameterValue(TEXT("Opacity"), Opacity);
+	}
+	else if (MaterialMode == EMeshInspectorMaterialMode::VertexColor)
+	{
+		ActiveCustomMaterial = ToolSetupUtil::GetVertexColorMaterial(nullptr);
+		ActiveCustomMaterial->SetScalarParameterValue(TEXT("FlatShading"), (bFlatShading) ? 1.0f : 0.0f);
+	}
+	else if (MaterialMode == EMeshInspectorMaterialMode::TangentNormal)
+	{
+		ActiveCustomMaterial = UMaterialInstanceDynamic::Create(ToolSetupUtil::GetImageBasedSculptMaterial(nullptr, ToolSetupUtil::ImageMaterialType::TangentNormalFromView), this);
+	}
+	else if (MaterialMode == EMeshInspectorMaterialMode::GroupColor)
+	{
+		ActiveCustomMaterial = ToolSetupUtil::GetVertexColorMaterial(nullptr);
+		ActiveCustomMaterial->SetScalarParameterValue(TEXT("FlatShading"), (bFlatShading) ? 1.0f : 0.0f);
+	}
+}
+
+
+UMaterialInterface* UMeshInspectorMaterialProperties::GetActiveOverrideMaterial() const
+{
+	if (MaterialMode == EMeshInspectorMaterialMode::Checkerboard)
+	{
+		return CheckerMaterial;
+	}
+	else if (MaterialMode == EMeshInspectorMaterialMode::Override)
+	{
+		return OverrideMaterial;
+	}
+	else if (MaterialMode == EMeshInspectorMaterialMode::FlatShaded)
+	{
+		return ActiveCustomMaterial;
+	} 
+	else if (MaterialMode == EMeshInspectorMaterialMode::Grey)
+	{
+		return ActiveCustomMaterial;
+	}
+	else if (MaterialMode == EMeshInspectorMaterialMode::Transparent)
+	{
+		return ActiveCustomMaterial;
+	}
+	else if (MaterialMode == EMeshInspectorMaterialMode::VertexColor)
+	{
+		return ActiveCustomMaterial;
+	}
+	else if (MaterialMode == EMeshInspectorMaterialMode::TangentNormal)
+	{
+		return ActiveCustomMaterial;
+	}
+	else if (MaterialMode == EMeshInspectorMaterialMode::GroupColor)
+	{
+		return ActiveCustomMaterial;
+	}
+
+	return nullptr;
+}
+
+
 /*
  * ToolBuilder
  */
@@ -54,7 +181,8 @@ void UMeshInspectorTool::Setup()
 	PreviewMesh->bBuildSpatialDataStructure = false;
 	PreviewMesh->CreateInWorld(GetTargetWorld(), FTransform::Identity);
 	ToolSetupUtil::ApplyRenderingConfigurationToPreview(PreviewMesh, Target);
-	PreviewMesh->SetTransform((FTransform)UE::ToolTarget::GetLocalToWorldTransform(Target));
+	LocalToWorldTransform = (FTransform)UE::ToolTarget::GetLocalToWorldTransform(Target);
+	PreviewMesh->SetTransform(LocalToWorldTransform);
 
 	FComponentMaterialSet MaterialSet = UE::ToolTarget::GetMaterialSet(Target);
 	PreviewMesh->SetMaterials(MaterialSet.Materials);
@@ -80,7 +208,15 @@ void UMeshInspectorTool::Setup()
 	Settings->RestoreProperties(this);
 	AddToolPropertySource(Settings);
 
-	MaterialSettings = NewObject<UExistingMeshMaterialProperties>(this);
+	PolygroupLayerProperties = NewObject<UPolygroupLayersProperties>(this);
+	PolygroupLayerProperties->RestoreProperties(this, TEXT("MeshGroupPaintTool"));
+	PolygroupLayerProperties->InitializeGroupLayers(PreviewMesh->GetMesh());
+	PolygroupLayerProperties->WatchProperty(PolygroupLayerProperties->ActiveGroupLayer, [&](FName) { OnSelectedGroupLayerChanged(); });
+	UpdateActiveGroupLayer();
+	AddToolPropertySource(PolygroupLayerProperties);
+
+
+	MaterialSettings = NewObject<UMeshInspectorMaterialProperties>(this);
 	const FDynamicMesh3* TargetMesh = PreviewMesh->GetPreviewDynamicMesh();
 	TArray<FString> UVChannelNamesList;
 	for (int32 k = 0; k < TargetMesh->Attributes()->NumUVLayers(); ++k)
@@ -110,6 +246,16 @@ void UMeshInspectorTool::Setup()
 }
 
 
+UE::Geometry::FDynamicMeshAABBTree3* UMeshInspectorTool::GetSpatial()
+{
+	if (MeshAABBTree.IsValid() == false)
+	{
+		MeshAABBTree = MakeUnique<FDynamicMeshAABBTree3>(PreviewMesh->GetMesh(), true);
+	}
+	return MeshAABBTree.Get();
+}
+
+
 void UMeshInspectorTool::OnShutdown(EToolShutdownType ShutdownType)
 {
 	UE::ToolTarget::ShowSourceObject(Target);
@@ -224,6 +370,145 @@ void UMeshInspectorTool::Render(IToolsContextRenderAPI* RenderAPI)
 }
 
 
+void UMeshInspectorTool::DrawHUD(FCanvas* Canvas, IToolsContextRenderAPI* RenderAPI)
+{
+	float DPIScale = Canvas->GetDPIScale();
+	UFont* UseFont = GEngine->GetSmallFont();
+	FViewCameraState CamState = RenderAPI->GetCameraState();
+	const FSceneView* SceneView = RenderAPI->GetSceneView();
+	FVector3d LocalEyePosition(LocalToWorldTransform.InverseTransformPosition((FVector3d)CamState.Position));
+	FVector3d LocalEyeDirection(LocalToWorldTransform.InverseTransformVectorNoScale((FVector3d)CamState.Forward()));
+
+	const FDynamicMesh3* Mesh = PreviewMesh->GetMesh();
+	FDynamicMeshAABBTree3* Spatial = GetSpatial();
+
+	// Code below draws the different type of mesh indices. The maximum number of labels to draw is
+	// capped. Points are clipped against the view frustum, so that only visible labels are drawn.
+	// In addition raycasts are used to check that the label target point is visible.
+	// This is all a bit expensive, it would be nice to do this in parallel and/or cache 
+	// the result for the current view
+
+	// config settings for drawing numbers
+	int32 DrawMaxNumbers = 500;
+	bool bCanShowAllValues = false;
+	switch (Settings->ShowIndices)
+	{
+	case EMeshInspectorToolDrawIndexMode::TriangleID: bCanShowAllValues = (Mesh->TriangleCount() < DrawMaxNumbers); break;
+	case EMeshInspectorToolDrawIndexMode::VertexID: bCanShowAllValues = (Mesh->VertexCount() < DrawMaxNumbers); break;
+	case EMeshInspectorToolDrawIndexMode::EdgeID: bCanShowAllValues = (Mesh->EdgeCount() < DrawMaxNumbers); break;
+	case EMeshInspectorToolDrawIndexMode::GroupID: bCanShowAllValues = true; break;
+	}
+
+	// figure out 2D pixel-space rectangle to use for clipping
+	FIntRect ViewRect = SceneView->UnconstrainedViewRect;
+	ViewRect.InflateRect( (bCanShowAllValues) ? 0 : -100 );
+	FAxisAlignedBox2d ViewRectBox(
+		FVector2d((double)ViewRect.Min.X, (double)ViewRect.Min.Y),
+		FVector2d((double)ViewRect.Max.X, (double)ViewRect.Max.Y));
+
+	// IsVisibleFunc checks if the camera projection of a 3D world-space point is considered "visible"
+	auto IsVisibleFunc = [SceneView, ViewRectBox](const FVector3d WorldPosition)
+	{
+		FVector2D PixelLocation;
+		if (SceneView->WorldToPixel(WorldPosition, PixelLocation) )
+		{
+			return ViewRectBox.Contains(PixelLocation);
+		}
+		return false;
+	};
+
+	int32 NumbersDrawnCount = 0;
+
+	// DrawNumberIfVisible draws the given number at the given 3D position, projected into the 2D canvas, if it is visible
+	// PositionHitTestFunc argument is called after raycast checks w/ the hit triangle, ray-distance, and local-space hit position, 
+	// to determine if visible point along the eye ray is suitable to draw the number at (ie if it is visible, etc)
+	auto DrawNumberIfVisible = [Mesh, Spatial, DPIScale, UseFont, this, LocalEyePosition, LocalEyeDirection, SceneView, Canvas, &NumbersDrawnCount, DrawMaxNumbers, &IsVisibleFunc]
+	(int32 Number, FVector3d LocalPosition, TFunctionRef<bool(int32, double, FVector3d)> PositionHitTestFunc)
+	{
+		FVector3d WorldPosition = LocalToWorldTransform.TransformPosition(LocalPosition);
+		if (NumbersDrawnCount >= DrawMaxNumbers || IsVisibleFunc(WorldPosition) == false )
+		{
+			return;
+		}
+		FRay3d LocalEyeRay;
+		LocalEyeRay.Origin = LocalEyePosition;
+		LocalEyeRay.Direction = Normalized(LocalPosition - LocalEyePosition);
+		if (LocalEyeRay.Direction.Dot(LocalEyeDirection) < 0)
+		{
+			return;		// ray to Position is pointing away from the camera
+		}
+		double LocalRayHitT;
+		int HitTID;
+		FVector3d HitBaryCoords;
+		if (Spatial->FindNearestHitTriangle(LocalEyeRay, LocalRayHitT, HitTID, HitBaryCoords))
+		{
+			FVector3d LocalHitPos = LocalEyeRay.PointAt(LocalRayHitT);
+			if (PositionHitTestFunc(HitTID, LocalRayHitT, LocalHitPos))
+			{
+				FVector2D PixelPos;
+				SceneView->WorldToPixel(WorldPosition, PixelPos);
+				FString String = FString::Printf(TEXT("%d"), Number);
+				Canvas->DrawShadowedString(PixelPos.X / (double)DPIScale, PixelPos.Y / (double)DPIScale, *String, UseFont, FLinearColor::White);
+				NumbersDrawnCount++;
+			}
+		}
+		else
+		{
+			// if we did not hit the mesh then this item must be at a grazing angle in which case it should be considered visible
+			FVector2D PixelPos;
+			SceneView->WorldToPixel(WorldPosition, PixelPos);
+			FString String = FString::Printf(TEXT("%d"), Number);
+			Canvas->DrawShadowedString(PixelPos.X / (double)DPIScale, PixelPos.Y / (double)DPIScale, *String, UseFont, FLinearColor::White);
+			NumbersDrawnCount++;
+		}
+	};
+
+
+	if (Settings->ShowIndices == EMeshInspectorToolDrawIndexMode::TriangleID)
+	{
+		for (int32 tid : Mesh->TriangleIndicesItr())
+		{
+			DrawNumberIfVisible(tid, Mesh->GetTriCentroid(tid),
+				[&](int32 HitTID, double, FVector3d) { return HitTID == tid; });
+		}
+	}
+	else if (Settings->ShowIndices == EMeshInspectorToolDrawIndexMode::VertexID)
+	{
+		FRandomStream TargetJitter(31337);
+		for (int32 vid : Mesh->VertexIndicesItr())
+		{
+			// slightly jitter the position here to avoid issues w/ rays that directly pass through vertices 'missing' due to precision limitations
+			FVector3d Position = Mesh->GetVertex(vid) + 0.0001 * TargetJitter.GetUnitVector();
+			DrawNumberIfVisible(vid, Position, [&](int32 HitTID, double LocalRayHitT, FVector3d LocalHitPos) { return Distance(LocalHitPos, Position) < 0.001; });
+		}
+	}
+	else if (Settings->ShowIndices == EMeshInspectorToolDrawIndexMode::EdgeID)
+	{
+		FRandomStream TargetJitter(31337);
+		for (int32 eid : Mesh->EdgeIndicesItr())
+		{
+			FVector3d Position = Mesh->GetEdgePoint(eid, 0.5) + 0.0001 * TargetJitter.GetUnitVector();
+			FIndex2i EdgeT = Mesh->GetEdgeT(eid);
+			DrawNumberIfVisible(eid, Position,
+				[&](int32 HitTID, double RayHitT, FVector3d) { return EdgeT.Contains(HitTID); });
+		}
+	}
+	else if (Settings->ShowIndices == EMeshInspectorToolDrawIndexMode::GroupID)
+	{
+		if (bDrawGroupsDataValid == false)
+		{
+			GroupVisualizationCache.UpdateGroupInfo_ConnectedComponents(*Mesh, *ActiveGroupSet);
+			bDrawGroupsDataValid = true;
+		}
+		for (const FGroupVisualizationCache::FGroupInfo& GroupInfo : GroupVisualizationCache)
+		{
+			DrawNumberIfVisible(GroupInfo.GroupID, GroupInfo.Center,
+				[&](int32 HitTID, double, FVector3d) { return HitTID == GroupInfo.CenterTris.A || HitTID == GroupInfo.CenterTris.B; });
+		}
+	}
+}
+
+
 void UMeshInspectorTool::OnPropertyModified(UObject* PropertySet, FProperty* Property)
 {
 	GetToolManager()->PostInvalidation();
@@ -232,9 +517,26 @@ void UMeshInspectorTool::OnPropertyModified(UObject* PropertySet, FProperty* Pro
 
 void UMeshInspectorTool::UpdateVisualization()
 {
+	if (!MaterialSettings) return;
+
 	PreviewMesh->EnableWireframe(Settings->bWireframe);
 
+	// determine custom material
 	MaterialSettings->UpdateMaterials();
+
+	// if override color function was set, clear it
+	if (MaterialSettings->MaterialMode != EMeshInspectorMaterialMode::GroupColor)
+	{
+		PreviewMesh->ClearTriangleColorFunction(UPreviewMesh::ERenderUpdateMode::FastUpdate);
+	}
+	else
+	{
+		PreviewMesh->SetTriangleColorFunction([this](const FDynamicMesh3* Mesh, int TriangleID)
+		{
+			return LinearColors::SelectFColor(ActiveGroupSet->GetTriangleGroup(TriangleID));
+		}, UPreviewMesh::ERenderUpdateMode::FastUpdate);
+	}
+
 	UMaterialInterface* OverrideMaterial = MaterialSettings->GetActiveOverrideMaterial();
 	if (OverrideMaterial == nullptr)
 	{
@@ -418,5 +720,31 @@ void UMeshInspectorTool::DecreaseLineWidthAction()
 
 
 
+void UMeshInspectorTool::OnSelectedGroupLayerChanged()
+{
+	UpdateActiveGroupLayer();
+}
+
+void UMeshInspectorTool::UpdateActiveGroupLayer()
+{
+	if (PolygroupLayerProperties->HasSelectedPolygroup() == false)
+	{
+		ActiveGroupSet = MakeUnique<UE::Geometry::FPolygroupSet>(PreviewMesh->GetMesh());
+	}
+	else
+	{
+		FName SelectedName = PolygroupLayerProperties->ActiveGroupLayer;
+		const FDynamicMeshPolygroupAttribute* FoundAttrib = UE::Geometry::FindPolygroupLayerByName(*PreviewMesh->GetMesh(), SelectedName);
+		ensureMsgf(FoundAttrib, TEXT("Selected Attribute Not Found! Falling back to Default group layer."));
+		ActiveGroupSet = MakeUnique<UE::Geometry::FPolygroupSet>(PreviewMesh->GetMesh(), FoundAttrib);
+	}
+	bDrawGroupsDataValid = false;
+
+	// force visualization update
+	UpdateVisualization();
+}
+
+
+
 #undef LOCTEXT_NAMESPACE
 
diff --git a/Engine/Plugins/Experimental/MeshModelingToolsetExp/Source/MeshModelingToolsExp/Public/MeshGroupPaintTool.h b/Engine/Plugins/Experimental/MeshModelingToolsetExp/Source/MeshModelingToolsExp/Public/MeshGroupPaintTool.h
--- a/Engine/Plugins/Experimental/MeshModelingToolsetExp/Source/MeshModelingToolsExp/Public/MeshGroupPaintTool.h
+++ b/Engine/Plugins/Experimental/MeshModelingToolsetExp/Source/MeshModelingToolsExp/Public/MeshGroupPaintTool.h
@@ -20,6 +20,7 @@
 #include "TransformTypes.h"
 #include "Changes/MeshPolygroupChange.h"
 #include "Polygroups/PolygroupSet.h"
+#include "FaceGroupUtil.h"
 
 #include "MeshGroupPaintTool.generated.h"
 
@@ -176,6 +177,15 @@ public:
 	/** Number of vertices in a triangle the Lasso must hit to be counted as "inside" */
 	UPROPERTY(EditAnywhere, Category = Filters, AdvancedDisplay, meta = (UIMin = 1, UIMax = 3, EditCondition = "SubToolType == EMeshGroupPaintInteractionType::PolyLasso"))
 	int MinTriVertCount = 1;
+
+
+	/** Display the Group ID of the last triangle under the cursor */
+	UPROPERTY(EditAnywhere, Category = Visualization)
+	bool bShowHitGroup = false;
+
+	/** Display the Group ID for all visible groups in the mesh */
+	UPROPERTY(EditAnywhere, Category = Visualization)
+	bool bShowAllGroups = false;
 };
 
 
@@ -288,9 +298,9 @@ public:
 
 	virtual void Setup() override;
 	virtual void Shutdown(EToolShutdownType ShutdownType) override;
-	virtual void DrawHUD(FCanvas* Canvas, IToolsContextRenderAPI* RenderAPI) override;
 
 	virtual void OnTick(float DeltaTime) override;
+	virtual void DrawHUD(FCanvas* Canvas, IToolsContextRenderAPI* RenderAPI) override;
 
 	virtual bool HasCancel() const override { return true; }
 	virtual bool HasAccept() const override { return true; }
@@ -404,7 +414,7 @@ protected:
 	TObjectPtr<UMeshElementsVisualizer> MeshElementsDisplay;
 
 	// realtime visualization
-	void OnDynamicMeshComponentChanged(UDynamicMeshComponent* Component, const FMeshVertexChange* Change, bool bRevert);
+	void OnDynamicMeshComponentChanged();
 	FDelegateHandle OnDynamicMeshComponentChangedHandle;
 
 	TUniquePtr<UE::Geometry::FPolygroupSet> ActiveGroupSet;
@@ -469,10 +479,11 @@ protected:
 	void PrecomputeFilterData();
 
 
+	bool bDrawGroupsDataValid = false;
+	UE::Geometry::FGroupVisualizationCache GroupVisualizationCache;
+
 protected:
 	virtual bool ShowWorkPlane() const override { return false; }
-
-
 };
 
 
diff --git a/Engine/Plugins/Experimental/MeshModelingToolsetExp/Source/MeshModelingToolsExp/Public/MeshInspectorTool.h b/Engine/Plugins/Experimental/MeshModelingToolsetExp/Source/MeshModelingToolsExp/Public/MeshInspectorTool.h
--- a/Engine/Plugins/Experimental/MeshModelingToolsetExp/Source/MeshModelingToolsExp/Public/MeshInspectorTool.h
+++ b/Engine/Plugins/Experimental/MeshModelingToolsetExp/Source/MeshModelingToolsExp/Public/MeshInspectorTool.h
@@ -12,6 +12,9 @@
 #include "Properties/MeshMaterialProperties.h"
 #include "PreviewMesh.h"
 #include "BaseTools/SingleSelectionMeshEditingTool.h"
+#include "PropertySets/PolygroupLayersProperties.h"
+#include "Polygroups/PolygroupSet.h"
+#include "FaceGroupUtil.h"
 #include "MeshInspectorTool.generated.h"
 
 
@@ -32,6 +35,16 @@ public:
 };
 
 
+UENUM()
+enum class EMeshInspectorToolDrawIndexMode : uint8
+{
+	None,
+	VertexID,
+	TriangleID,
+	GroupID,
+	EdgeID
+};
+
 
 UCLASS()
 class MESHMODELINGTOOLSEXP_API UMeshInspectorProperties : public UInteractiveToolPropertySet
@@ -80,16 +93,125 @@ public:
 	bool bTangentVectors = false;
 
 	/** Length of line segments representing normal vectors */
-	UPROPERTY(EditAnywhere, Category = Options, meta = (EditCondition = "bNormalVectors", 
+	UPROPERTY(EditAnywhere, Category = Options, AdvancedDisplay, meta = (EditCondition = "bNormalVectors",
 		UIMin="0", UIMax="400", ClampMin = "0", ClampMax = "1000000000.0"))
 	float NormalLength = 5.0f;
 
 	/** Length of line segments representing tangent vectors */
-	UPROPERTY(EditAnywhere, Category = Options, meta = (EditCondition = "bTangentVectors", 
+	UPROPERTY(EditAnywhere, Category = Options, AdvancedDisplay, meta = (EditCondition = "bTangentVectors",
 		UIMin = "0", UIMax = "400", ClampMin = "0", ClampMax = "1000000000.0"))
 	float TangentLength = 5.0f;
+
+	/** Draw the mesh indices of the selected type. A maximum of 500 visible indices will be rendered. */
+	UPROPERTY(EditAnywhere, Category=Options)
+	EMeshInspectorToolDrawIndexMode ShowIndices = EMeshInspectorToolDrawIndexMode::None;
+};
+
+
+
+
+/** Material Modes for Mesh Inspector Tool */
+UENUM()
+enum class EMeshInspectorMaterialMode : uint8
+{
+	/** Input material */
+	Original,
+	/** Flat Shaded Material, ie with per-triangle normals */
+	FlatShaded,
+	/** Grey material */
+	Grey,
+	/** Transparent material, with opacity/color controls */
+	Transparent,
+	/** Tangent/Normal material */
+	TangentNormal,
+	/** Vertex Color material */
+	VertexColor,
+	/** Polygroup Color material */
+	GroupColor,
+	/** Checkerboard material */
+	Checkerboard,
+	/** Override material */
+	Override
+};
+
+
+// material settings for mesh inspector tool
+UCLASS()
+class MESHMODELINGTOOLSEXP_API UMeshInspectorMaterialProperties : public UInteractiveToolPropertySet
+{
+	GENERATED_BODY()
+public:
+
+	/** Material that will be used to render the mesh */
+	UPROPERTY(EditAnywhere, Category = PreviewMaterial)
+	EMeshInspectorMaterialMode MaterialMode = EMeshInspectorMaterialMode::Original;
+
+	/** Number of checkerboard tiles within the 0 to 1 range; only available when Checkerboard is selected as material mode */
+	UPROPERTY(EditAnywhere, Category = PreviewMaterial,
+		meta = (UIMin = "10.0", UIMax = "100.0", ClampMin = "0.01", ClampMax = "10000.0", EditConditionHides, EditCondition = "MaterialMode == EMeshInspectorMaterialMode::Checkerboard"))
+	float CheckerDensity = 50.0f;
+
+	/** Material to use instead of the original material; only available when Override is selected as material mode */
+	UPROPERTY(EditAnywhere, Category = PreviewMaterial, meta = (EditConditionHides, EditCondition = "MaterialMode == EMeshInspectorMaterialMode::Override"))
+	TObjectPtr<UMaterialInterface> OverrideMaterial = nullptr;
+
+	/** Which UV channel to use for visualizing the checkerboard material on the mesh; note that this does not affect the preview layout */
+	UPROPERTY(EditAnywhere, Category = PreviewMaterial,
+		meta = (DisplayName = "Preview UV Channel", GetOptions = GetUVChannelNamesFunc, EditConditionHides, EditCondition =
+			"MaterialMode == EMeshInspectorMaterialMode::Checkerboard", NoResetToDefault))
+	FString UVChannel;
+
+	UPROPERTY(meta = (TransientToolProperty))
+	TArray<FString> UVChannelNamesList;
+
+	UFUNCTION()
+	const TArray<FString>& GetUVChannelNamesFunc() const;
+
+
+	/** Toggle flat shading on/off */
+	UPROPERTY(EditAnywhere, Category = PreviewMaterial, meta = (EditConditionHides, EditCondition = "MaterialMode == EMeshInspectorMaterialMode::VertexColor || MaterialMode == EMeshInspectorMaterialMode::GroupColor") )
+	bool bFlatShading = false;
+
+	/** Main Color of Material */
+	UPROPERTY(EditAnywhere, Category = PreviewMaterial, meta = (EditConditionHides, EditCondition = "MaterialMode == EMeshInspectorMaterialMode::Diffuse"))
+	FLinearColor Color = FLinearColor(0.4f, 0.4f, 0.4f);
+
+	/** Opacity of transparent material */
+	UPROPERTY(EditAnywhere, Category = PreviewMaterial, meta = (EditConditionHides, EditCondition = "MaterialMode == EMeshInspectorMaterialMode::Transparent", ClampMin = "0", ClampMax = "1.0"))
+	double Opacity = 0.65;
+
+	//~ Could have used the same property as Color, above, but the user may want different saved values for the two
+	UPROPERTY(EditAnywhere, Category = PreviewMaterial, meta = (EditConditionHides, EditCondition = "MaterialMode == EMeshInspectorMaterialMode::Transparent", DisplayName = "Color"))
+	FLinearColor TransparentMaterialColor = FLinearColor(0.0606, 0.309, 0.842);
+
+	/** Although a two-sided transparent material causes rendering issues with overlapping faces, it is still frequently useful to see the shape when sculpting around other objects. */
+	UPROPERTY(EditAnywhere, Category = PreviewMaterial, meta = (EditConditionHides, EditCondition = "MaterialMode == EMeshInspectorMaterialMode::Transparent"))
+	bool bTwoSided = true;
+
+	UPROPERTY(meta = (TransientToolProperty))
+	TObjectPtr<UMaterialInstanceDynamic> CheckerMaterial = nullptr;
+
+	UPROPERTY(meta = (TransientToolProperty))
+	TObjectPtr<UMaterialInstanceDynamic> ActiveCustomMaterial = nullptr;
+
+	// Needs custom restore in order to call setup
+	virtual void RestoreProperties(UInteractiveTool* RestoreToTool, const FString& CacheIdentifier = TEXT("")) override;
+
+	void Setup();
+
+	void UpdateMaterials();
+	UMaterialInterface* GetActiveOverrideMaterial() const;
+	
+	void UpdateUVChannels(int32 UVChannelIndex, const TArray<FString>& UVChannelNames, bool bUpdateSelection = true);
 };
 
+
+
+
+
+
+
+
 /**
  * Mesh Inspector Tool for visualizing mesh information
  */
@@ -107,6 +229,7 @@ public:
 	virtual void OnShutdown(EToolShutdownType ShutdownType) override;
 
 	virtual void Render(IToolsContextRenderAPI* RenderAPI) override;
+	virtual void DrawHUD(FCanvas* Canvas, IToolsContextRenderAPI* RenderAPI) override;
 
 	virtual bool HasCancel() const override { return false; }
 	virtual bool HasAccept() const override;
@@ -125,8 +248,10 @@ protected:
 	TObjectPtr<UMeshInspectorProperties> Settings = nullptr;
 
 	UPROPERTY()
-	TObjectPtr<UExistingMeshMaterialProperties> MaterialSettings = nullptr;
+	TObjectPtr<UPolygroupLayersProperties> PolygroupLayerProperties = nullptr;
 
+	UPROPERTY()
+	TObjectPtr<UMeshInspectorMaterialProperties> MaterialSettings = nullptr;
 
 	float LineWidthMultiplier = 1.0f;
 
@@ -150,4 +275,17 @@ protected:
 
 	void UpdateVisualization();
 	void Precompute();
+
+	FViewCameraState CameraState;
+	FTransform LocalToWorldTransform;
+
+	TUniquePtr<UE::Geometry::FDynamicMeshAABBTree3> MeshAABBTree;
+	UE::Geometry::FDynamicMeshAABBTree3* GetSpatial();
+
+	TUniquePtr<UE::Geometry::FPolygroupSet> ActiveGroupSet;
+	void OnSelectedGroupLayerChanged();
+	void UpdateActiveGroupLayer();
+
+	bool bDrawGroupsDataValid = false;
+	UE::Geometry::FGroupVisualizationCache GroupVisualizationCache;
 };
diff --git a/Engine/Plugins/Runtime/GeometryProcessing/Source/DynamicMesh/Private/FaceGroupUtil.cpp b/Engine/Plugins/Runtime/GeometryProcessing/Source/DynamicMesh/Private/FaceGroupUtil.cpp
--- a/Engine/Plugins/Runtime/GeometryProcessing/Source/DynamicMesh/Private/FaceGroupUtil.cpp
+++ b/Engine/Plugins/Runtime/GeometryProcessing/Source/DynamicMesh/Private/FaceGroupUtil.cpp
@@ -1,6 +1,8 @@
 // Copyright Epic Games, Inc. All Rights Reserved.
 
 #include "FaceGroupUtil.h"
+#include "Selections/MeshConnectedComponents.h"
+#include "Async/ParallelFor.h"
 
 using namespace UE::Geometry;
 
@@ -201,4 +203,89 @@ void FaceGroupUtil::SeparateMeshByGroups(FDynamicMesh3& Mesh, TArray<FDynamicMes
 		check(M.TriangleCount() > 0); // SplitMesh should never add an empty mesh
 		GroupIDs.Add(M.GetTriangleGroup(0));
 	}
+}
+
+
+
+void FGroupVisualizationCache::UpdateGroupInfo_ConnectedComponents(
+	const FDynamicMesh3& SourceMesh,
+	const FPolygroupSet& GroupSet,
+	bool bParallel)
+{
+	
+	// find connected group components
+	FMeshConnectedComponents Components(&SourceMesh);
+	Components.FindConnectedTriangles([&](int32 t1, int32 t2) { return GroupSet.GetGroup(t1) == GroupSet.GetGroup(t2); });
+
+	GroupInfo.SetNum(Components.Num());
+
+	EParallelForFlags UseParallelFlags = (bParallel) ? EParallelForFlags::None : EParallelForFlags::ForceSingleThread;
+	ParallelFor(Components.Num(), [&](int32 ci) {
+		const FMeshConnectedComponents::FComponent& Component = Components[ci];
+		GroupInfo[ci].GroupID = GroupSet.GetGroup(Component.Indices[0]);
+
+		// compute bounds
+		GroupInfo[ci].Bounds = FAxisAlignedBox3d::Empty();
+		for (int32 tid : Component.Indices)
+		{
+			GroupInfo[ci].Bounds.Contain(SourceMesh.GetTriBounds(tid));
+		}
+
+		if (Component.Indices.Num() == 1)
+		{
+			GroupInfo[ci].Center = SourceMesh.GetTriCentroid(Component.Indices[0]);
+			GroupInfo[ci].CenterTris = FIndex2i(Component.Indices[0], Component.Indices[0]);
+		}
+		else if (Component.Indices.Num() == 2 &&
+			SourceMesh.GetTriNeighbourTris(Component.Indices[0]).Contains(Component.Indices[1]))
+		{
+			int32 eid = SourceMesh.FindEdgeFromTriPair(Component.Indices[0], Component.Indices[1]);
+			if (eid != IndexConstants::InvalidID)
+			{
+				GroupInfo[ci].Center = SourceMesh.GetEdgePoint(eid, 0.5);
+				GroupInfo[ci].CenterTris = SourceMesh.GetEdgeT(eid);
+			}
+			else
+			{
+				GroupInfo[ci].Center = SourceMesh.GetTriCentroid(Component.Indices[0]);
+				GroupInfo[ci].CenterTris = FIndex2i(Component.Indices[0], Component.Indices[0]);
+			}
+		}
+		else
+		{
+			int32 CenterTriID = -1;
+			TSet<int32> TrisInGroup(Component.Indices);
+			TArray<int32> BorderTris;
+			while (CenterTriID == -1)
+			{
+				BorderTris.Reset();
+				for (int32 tid : TrisInGroup)
+				{
+					FIndex3i NbrTris = SourceMesh.GetTriNeighbourTris(tid);
+					bool bIsBorder = (TrisInGroup.Contains(NbrTris.A) == false || TrisInGroup.Contains(NbrTris.B) == false || TrisInGroup.Contains(NbrTris.C) == false);
+					if (bIsBorder)
+					{
+						BorderTris.Add(tid);
+					}
+				}
+				for (int32 tid : BorderTris)
+				{
+					TrisInGroup.Remove(tid);
+					if (TrisInGroup.Num() == 1)
+					{
+						CenterTriID = *TrisInGroup.begin();
+						break;
+					}
+				}
+			}
+			GroupInfo[ci].Center = SourceMesh.GetTriCentroid(CenterTriID);
+			GroupInfo[ci].CenterTris = FIndex2i(CenterTriID, CenterTriID);
+		}
+
+		if (bStorePerGroupTriangleIDs)
+		{
+			GroupInfo[ci].TriangleIDs = MoveTemp(Components[ci].Indices);
+		}
+
+	}, UseParallelFlags);
 }
\ No newline at end of file
diff --git a/Engine/Plugins/Runtime/GeometryProcessing/Source/DynamicMesh/Public/FaceGroupUtil.h b/Engine/Plugins/Runtime/GeometryProcessing/Source/DynamicMesh/Public/FaceGroupUtil.h
--- a/Engine/Plugins/Runtime/GeometryProcessing/Source/DynamicMesh/Public/FaceGroupUtil.h
+++ b/Engine/Plugins/Runtime/GeometryProcessing/Source/DynamicMesh/Public/FaceGroupUtil.h
@@ -7,6 +7,7 @@
 #include "MathUtil.h"
 #include "DynamicMesh/DynamicMesh3.h"
 #include "DynamicMeshEditor.h"
+#include "Polygroups/PolygroupSet.h"
 
 /**
  * Utility functions for dealing with Mesh FaceGroups
@@ -68,5 +69,66 @@ namespace FaceGroupUtil
 	*/
 	DYNAMICMESH_API void SeparateMeshByGroups(FDynamicMesh3& Mesh, TArray<FDynamicMesh3>& SplitMeshes, TArray<int32>& GroupIDs);
 
-
 }
+
+
+namespace UE
+{
+namespace Geometry
+{
+
+
+/**
+ * FGroupVisualizationCache is intended to be used as a cache for per-group information that 
+ * would be used to visualize PolyGroups, in particular drawing the Polygroup ID for each group.
+ */
+class DYNAMICMESH_API FGroupVisualizationCache
+{
+public:
+
+	struct FGroupInfo
+	{
+		int32 GroupID;
+		FAxisAlignedBox3d Bounds;		// bounding box of the triangles in the group
+
+		FVector3d Center;				// "center point" of the group - may or may not actually be on the mesh surface
+		FIndex2i CenterTris;			// "center triangles" of the group - one or two triangles, may be invalid
+
+		TArray<int32> TriangleIDs;		// list of triangles in the group, only initialized if bStorePerGroupTriangleIDs is true
+	};
+
+	// if enabled, each FGroupInfo keeps track of the TriangleIDs that belong to it. Can have significant memory overhead.
+	bool bStorePerGroupTriangleIDs = false;
+
+	/**
+	 * Computed Per-Group Info, for each group found during the Update functions below
+	 */
+	TArray<FGroupInfo> GroupInfo;
+
+
+	/**
+	 * Update the GroupInfo array by finding group-connected components.
+	 * Note that this may result in multiple GroupInfo items with the same GroupID, if they are not connected-components.
+	 * Computing connected components is somewhat expensive. In addition, for each component, if it is more
+	 * than 2 triangles it will be "eroded" down a single triangle to determine the group "center", also somewhat expensive.
+	 * This step can be computed in parallel if bParallel is true.
+	 */
+	void UpdateGroupInfo_ConnectedComponents(
+		const FDynamicMesh3& SourceMesh,
+		const FPolygroupSet& GroupSet,
+		bool bParallel = true);
+
+public:
+	/**
+	 * DO NOT USE DIRECTLY
+	 * STL-like iterators to enable ranged-based for loop support
+	 */
+	auto begin() { return GroupInfo.begin(); }
+	auto begin() const { return GroupInfo.begin(); }
+	auto end() { return GroupInfo.end(); }
+	auto end() const { return GroupInfo.end(); }
+};
+
+
+} // end namespace UE::Geometry
+} // end namespace UE
