Fixed non-unity/pch errors

diff --git a/Engine/Plugins/Experimental/MeshModelingToolsetExp/Source/MeshModelingToolsExp/Private/MeshGroupPaintTool.cpp b/Engine/Plugins/Experimental/MeshModelingToolsetExp/Source/MeshModelingToolsExp/Private/MeshGroupPaintTool.cpp
--- a/Engine/Plugins/Experimental/MeshModelingToolsetExp/Source/MeshModelingToolsExp/Private/MeshGroupPaintTool.cpp
+++ b/Engine/Plugins/Experimental/MeshModelingToolsetExp/Source/MeshModelingToolsExp/Private/MeshGroupPaintTool.cpp
@@ -8,6 +8,7 @@
 #include "Drawing/MeshElementsVisualizer.h"
 #include "Async/ParallelFor.h"
 #include "Async/Async.h"
+#include "SceneView.h"
 #include "ToolSetupUtil.h"
 #include "ModelingToolTargetUtil.h"
 #include "MeshWeights.h"
diff --git a/Engine/Plugins/Experimental/VirtualCamera/Source/VCamCore/Private/VCamOutputRemoteSession.cpp b/Engine/Plugins/Experimental/VirtualCamera/Source/VCamCore/Private/VCamOutputRemoteSession.cpp
--- a/Engine/Plugins/Experimental/VirtualCamera/Source/VCamCore/Private/VCamOutputRemoteSession.cpp
+++ b/Engine/Plugins/Experimental/VirtualCamera/Source/VCamCore/Private/VCamOutputRemoteSession.cpp
@@ -1,6 +1,7 @@
 // Copyright Epic Games, Inc. All Rights Reserved.
 
 #include "VCamOutputRemoteSession.h"
+#include "Engine/GameInstance.h"
 #include "UObject/Package.h"
 #include "VCamOutputComposure.h"
 #include "Channels/RemoteSessionImageChannel.h"
