add missing headers to fix nonunity error

diff --git a/Engine/Plugins/Experimental/MeshModelingToolsetExp/Source/MeshModelingToolsExp/Private/MeshGroupPaintTool.cpp b/Engine/Plugins/Experimental/MeshModelingToolsetExp/Source/MeshModelingToolsExp/Private/MeshGroupPaintTool.cpp
--- a/Engine/Plugins/Experimental/MeshModelingToolsetExp/Source/MeshModelingToolsExp/Private/MeshGroupPaintTool.cpp
+++ b/Engine/Plugins/Experimental/MeshModelingToolsetExp/Source/MeshModelingToolsExp/Private/MeshGroupPaintTool.cpp
@@ -34,6 +34,7 @@
 
 #include "CanvasTypes.h"
 #include "CanvasItem.h"
+#include "Engine/Engine.h"  // for GEngine->GetSmallFont()
 
 #include UE_INLINE_GENERATED_CPP_BY_NAME(MeshGroupPaintTool)
 
diff --git a/Engine/Plugins/Experimental/MeshModelingToolsetExp/Source/MeshModelingToolsExp/Private/MeshInspectorTool.cpp b/Engine/Plugins/Experimental/MeshModelingToolsetExp/Source/MeshModelingToolsExp/Private/MeshInspectorTool.cpp
--- a/Engine/Plugins/Experimental/MeshModelingToolsetExp/Source/MeshModelingToolsExp/Private/MeshInspectorTool.cpp
+++ b/Engine/Plugins/Experimental/MeshModelingToolsetExp/Source/MeshModelingToolsExp/Private/MeshInspectorTool.cpp
@@ -29,6 +29,7 @@
 
 #include "CanvasTypes.h"
 #include "CanvasItem.h"
+#include "Engine/Engine.h"  // for GEngine->GetSmallFont()
 
 #include UE_INLINE_GENERATED_CPP_BY_NAME(MeshInspectorTool)
 
