add missing headers to fix nonunity compile issue

diff --git a/Engine/Plugins/Experimental/MeshModelingToolsetExp/Source/MeshModelingToolsExp/Private/MeshGroupPaintTool.cpp b/Engine/Plugins/Experimental/MeshModelingToolsetExp/Source/MeshModelingToolsExp/Private/MeshGroupPaintTool.cpp
--- a/Engine/Plugins/Experimental/MeshModelingToolsetExp/Source/MeshModelingToolsExp/Private/MeshGroupPaintTool.cpp
+++ b/Engine/Plugins/Experimental/MeshModelingToolsetExp/Source/MeshModelingToolsExp/Private/MeshGroupPaintTool.cpp
@@ -35,6 +35,7 @@
 #include "CanvasTypes.h"
 #include "CanvasItem.h"
 #include "Engine/Engine.h"  // for GEngine->GetSmallFont()
+#include "SceneView.h"
 
 #include UE_INLINE_GENERATED_CPP_BY_NAME(MeshGroupPaintTool)
 
diff --git a/Engine/Plugins/Experimental/MeshModelingToolsetExp/Source/MeshModelingToolsExp/Private/MeshInspectorTool.cpp b/Engine/Plugins/Experimental/MeshModelingToolsetExp/Source/MeshModelingToolsExp/Private/MeshInspectorTool.cpp
--- a/Engine/Plugins/Experimental/MeshModelingToolsetExp/Source/MeshModelingToolsExp/Private/MeshInspectorTool.cpp
+++ b/Engine/Plugins/Experimental/MeshModelingToolsetExp/Source/MeshModelingToolsExp/Private/MeshInspectorTool.cpp
@@ -30,6 +30,7 @@
 #include "CanvasTypes.h"
 #include "CanvasItem.h"
 #include "Engine/Engine.h"  // for GEngine->GetSmallFont()
+#include "SceneView.h"
 
 #include UE_INLINE_GENERATED_CPP_BY_NAME(MeshInspectorTool)
 
