MeshModelingTools: Fix bad merge resolution from 22901795 that broke PolyCut behavior downstream.

diff --git a/Engine/Plugins/Experimental/MeshModelingToolsetExp/Source/MeshModelingToolsEditorOnlyExp/Private/PolygonOnMeshTool.cpp b/Engine/Plugins/Experimental/MeshModelingToolsetExp/Source/MeshModelingToolsEditorOnlyExp/Private/PolygonOnMeshTool.cpp
--- a/Engine/Plugins/Experimental/MeshModelingToolsetExp/Source/MeshModelingToolsEditorOnlyExp/Private/PolygonOnMeshTool.cpp
+++ b/Engine/Plugins/Experimental/MeshModelingToolsetExp/Source/MeshModelingToolsEditorOnlyExp/Private/PolygonOnMeshTool.cpp
@@ -185,9 +185,10 @@ void UPolygonOnMeshTool::UpdateVisualization()
 		TargetMesh->GetEdgeV(EID, A, B);
 		DrawnLineSet->AddLine((FVector)A, (FVector)B, PartialPathEdgeColor, PartialPathEdgeThickness, PartialPathEdgeDepthBias);
 	}
+
 	// In the case where we don't allow failed results, it can be disorienting to show the broken mesh (especially since sometimes
 	// most of it may be cut away). But user can change this behavior.
-	if (!BasicProperties->bCanAcceptFailedResult && !BasicProperties->bShowIntermediateResultOnFailure)
+	if (!bOperationSucceeded && !BasicProperties->bCanAcceptFailedResult && !BasicProperties->bShowIntermediateResultOnFailure)
 	{
 		// Reset the preview.
 		Preview->PreviewMesh->UpdatePreview(OriginalDynamicMesh.Get());
