PatternTool: Use a gizmo to allow users to modify extents/radius in the viewport visually.

A single gizmo whose position relative to the tool's construction plane mechanic's gizmo is used to define Extent/ExtentX & ExtentY/Radius. The gizmo's movement is restricted depending on the current shape, axis/plane, and the orientation of the tool's construction plane frame. The visual appearance of the gizmo is modified by subclassing the actor builder for the existing combined transform gizmo and registering a 'new' gizmo type with the gizmo manager which is based on the existing combined transform gizmo.

diff --git a/Engine/Plugins/Experimental/MeshModelingToolsetExp/Source/MeshModelingToolsExp/Private/PatternTool.cpp b/Engine/Plugins/Experimental/MeshModelingToolsetExp/Source/MeshModelingToolsExp/Private/PatternTool.cpp
--- a/Engine/Plugins/Experimental/MeshModelingToolsetExp/Source/MeshModelingToolsExp/Private/PatternTool.cpp
+++ b/Engine/Plugins/Experimental/MeshModelingToolsetExp/Source/MeshModelingToolsExp/Private/PatternTool.cpp
@@ -26,7 +26,13 @@
 #include "Engine/StaticMeshActor.h"
 #include "Engine/StaticMesh.h"
 
+#include "ContextObjectStore.h"
+#include "BaseGizmos/AxisPositionGizmo.h"
+#include "BaseGizmos/CombinedTransformGizmo.h"
+#include "BaseGizmos/GizmoActor.h"
 #include "BaseGizmos/GizmoComponents.h"
+#include "BaseGizmos/GizmoViewContext.h"
+#include "BaseGizmos/GizmoBoxComponent.h"
 #include "BaseGizmos/TransformGizmoUtil.h"
 
 #include "Engine/World.h"
@@ -90,6 +96,49 @@ const FToolTargetTypeRequirements& UPatternToolBuilder::GetTargetRequirements()
 	return TypeRequirements;
 }
 
+/*
+ * Custom GizmoActor Factory
+ */
+class FPatternToolGizmoActorFactory : public FCombinedTransformGizmoActorFactory
+{
+public:
+	FPatternToolGizmoActorFactory(UGizmoViewContext* GizmoViewContextIn)
+	: FCombinedTransformGizmoActorFactory(GizmoViewContextIn)
+	{
+		EnableElements = ETransformGizmoSubElements::None;
+	}
+
+	/**
+	 * @param World the UWorld to create the new Actor in
+	 * @return new ACombinedTransformGizmoActor instance with members initialized with Components suitable for a transformation Gizmo
+	 */
+	virtual ACombinedTransformGizmoActor* CreateNewGizmoActor(UWorld* World) const override
+	{
+		FActorSpawnParameters SpawnInfo;
+		ACombinedTransformGizmoActor* NewActor = World->SpawnActor<ACombinedTransformGizmoActor>(FVector::ZeroVector, FRotator::ZeroRotator, SpawnInfo);
+
+		UGizmoBoxComponent* Component = AGizmoActor::AddDefaultBoxComponent(World, NewActor, GizmoViewContext, FLinearColor::Red, FVector::ZeroVector);
+		Component->LineThickness = 5.0f;
+		Component->Dimensions = FVector(Component->LineThickness);
+		Component->NotifyExternalPropertyUpdates();
+		
+		if ((EnableElements & ETransformGizmoSubElements::TranslateAxisX) != ETransformGizmoSubElements::None)
+		{
+			NewActor->TranslateX = Component;
+		}
+		else if ((EnableElements & ETransformGizmoSubElements::TranslatePlaneXY) != ETransformGizmoSubElements::None)
+		{
+			NewActor->TranslateXY = Component;
+		}
+		
+		return NewActor;
+	}
+};
+
+/*
+ * PatternGenerators
+ */
+
 // these pattern generators should be promoted to GeometryProcessing, however they need some
 // work to clean up the API and make them more correct (ie handling of step size seems a bit
 // flaky, particularly for circle patterns)
@@ -532,14 +581,15 @@ void UPatternTool::Setup()
 	Settings = NewObject<UPatternToolSettings>();
 	AddToolPropertySource(Settings);
 	Settings->RestoreProperties(this);
+	Settings->WatchProperty(Settings->SingleAxis, [this](EPatternToolSingleAxis SingleAxis) { OnSingleAxisUpdated(); });
+	Settings->WatchProperty(Settings->SinglePlane, [this](EPatternToolSinglePlane SinglePlane) { OnSinglePlaneUpdated(); });
+	Settings->WatchProperty(Settings->Shape, [this](EPatternToolShape) { OnShapeUpdated(); } );
+	Settings->WatchProperty(Settings->bHideSources, [this](bool bNewValue) { OnSourceVisibilityToggled(!bNewValue); } );
 	Settings->WatchProperty(Settings->Seed, [this](int32 NewSeed) { MarkPatternDirty(); } );
 	Settings->WatchProperty(Settings->bProjectElementsDown, [this](bool bNewValue) { MarkPatternDirty(); } );
 	Settings->WatchProperty(Settings->ProjectionOffset, [this](float NewValue) { MarkPatternDirty(); } );
 	Settings->WatchProperty(Settings->bHideSources, [this](bool bNewValue) { OnSourceVisibilityToggled(!bNewValue); } );
 	Settings->WatchProperty(Settings->bUseRelativeTransforms, [this](bool bNewValue) { MarkPatternDirty(); } );
-	Settings->WatchProperty(Settings->Shape, [this](EPatternToolShape) { OnShapeUpdated(); } );
-	Settings->WatchProperty(Settings->SingleAxis, [this](EPatternToolSingleAxis) { OnParametersUpdated(); } );
-	Settings->WatchProperty(Settings->SinglePlane, [this](EPatternToolSinglePlane) { OnParametersUpdated(); } );
 
 	BoundingBoxSettings = NewObject<UPatternTool_BoundingBoxSettings>();
 	AddToolPropertySource(BoundingBoxSettings);
@@ -552,6 +602,8 @@ void UPatternTool::Setup()
 	AddToolPropertySource(LinearSettings);
 	LinearSettings->RestoreProperties(this);
 	LinearSettings->WatchProperty(LinearSettings->SpacingMode, [this](EPatternToolAxisSpacingMode NewSpacingMode) { OnSpacingModeUpdated(); });
+	LinearSettings->WatchProperty(LinearSettings->bCentered, [this](bool bNewValue) { ResetTransformGizmoPosition(); });
+	LinearExtentWatcherIdx = LinearSettings->WatchProperty(LinearSettings->Extent, [this](double NewValue){ ResetTransformGizmoPosition(); });
 	SetToolPropertySourceEnabled(LinearSettings, false);
 
 	GridSettings = NewObject<UPatternTool_GridSettings>();
@@ -559,12 +611,17 @@ void UPatternTool::Setup()
 	GridSettings->RestoreProperties(this);
 	GridSettings->WatchProperty(GridSettings->SpacingX, [this](EPatternToolAxisSpacingMode NewSpacingMode) { OnSpacingModeUpdated(); });
 	GridSettings->WatchProperty(GridSettings->SpacingY, [this](EPatternToolAxisSpacingMode NewSpacingMode) { OnSpacingModeUpdated(); });
+	GridSettings->WatchProperty(GridSettings->bCenteredX, [this](bool bNewValue) { ResetTransformGizmoPosition(); });
+	GridSettings->WatchProperty(GridSettings->bCenteredY, [this](bool bNewValue) { ResetTransformGizmoPosition(); });
+	GridExtentXWatcherIdx = GridSettings->WatchProperty(GridSettings->ExtentX, [this](double NewValue){ ResetTransformGizmoPosition(); });
+	GridExtentYWatcherIdx = GridSettings->WatchProperty(GridSettings->ExtentY, [this](double NewValue){ ResetTransformGizmoPosition(); });
 	SetToolPropertySourceEnabled(GridSettings, false);
 
 	RadialSettings = NewObject<UPatternTool_RadialSettings>();
 	AddToolPropertySource(RadialSettings);
 	RadialSettings->RestoreProperties(this);
 	RadialSettings->WatchProperty(RadialSettings->SpacingMode, [this](EPatternToolAxisSpacingMode NewSpacingMode) { OnSpacingModeUpdated(); });
+	RadiusWatcherIdx = RadialSettings->WatchProperty(RadialSettings->Radius, [this](double NewValue){ ResetTransformGizmoPosition(); });
 	SetToolPropertySourceEnabled(RadialSettings, false);
 
 	RotationSettings = NewObject<UPatternTool_RotationSettings>();
@@ -661,14 +718,35 @@ void UPatternTool::Setup()
 	}
 
 	CurrentStartFrameWorld = FFrame3d(Elements[0].SourceTransform);
-
-	OnShapeUpdated();
-
+	
 	PlaneMechanic = NewObject<UConstructionPlaneMechanic>(this);
 	PlaneMechanic->Setup(this);
 	PlaneMechanic->Initialize( GetTargetWorld(), CurrentStartFrameWorld );
 	PlaneMechanic->OnPlaneChanged.AddLambda([this]() { OnMainFrameUpdated(); });
-
+	
+	PatternGizmoProxy = NewObject<UTransformProxy>(this);
+	PatternGizmoProxy->OnTransformChanged.AddUObject(this, &UPatternTool::OnTransformGizmoUpdated);
+
+	// The gizmo used to define pattern extents/radius is a bit hacky because it uses a combined transform
+	// gizmo which only ever has a single component of the gizmo at a time (it will either allow movement in
+	// a single axis or a single plane) and is visually just a box. This made the code much simpler than
+	// creating an entirely new gizmo but probably shouldn't be used as a pattern for other custom gizmos.
+	// This tool registers its own builder and unregisters the builder at shutdown, the builder shouldn't be
+	// relied on elsewhere.
+	UCombinedTransformGizmoBuilder* CustomThreeAxisBuilder = NewObject<UCombinedTransformGizmoBuilder>();
+	GizmoActorBuilder = MakeShared<FPatternToolGizmoActorFactory>(GetToolManager()->GetContextObjectStore()->FindContext<UGizmoViewContext>());
+	CustomThreeAxisBuilder->AxisPositionBuilderIdentifier = UInteractiveGizmoManager::DefaultAxisPositionBuilderIdentifier;
+	CustomThreeAxisBuilder->PlanePositionBuilderIdentifier = UInteractiveGizmoManager::DefaultPlanePositionBuilderIdentifier;
+	CustomThreeAxisBuilder->AxisAngleBuilderIdentifier = UInteractiveGizmoManager::DefaultAxisAngleBuilderIdentifier;
+	CustomThreeAxisBuilder->GizmoActorBuilder = GizmoActorBuilder;
+	GetToolManager()->GetPairedGizmoManager()->RegisterGizmoType(PatternToolThreeAxisTransformBuilderIdentifier, CustomThreeAxisBuilder);
+	bPatternToolThreeAxisTransformGizmoRegistered = true;
+
+	// Needs to be called before any of the watchers call ResetTransformGizmoPosition in order to set the
+	// proxy as the target of the gizmo, otherwise the gizmo will attempt to dereference an invalid StateTarget
+	// due to SetActiveTarget having never been called
+	ReconstructTransformGizmos();
+	
 	SetToolDisplayName(LOCTEXT("ToolName", "Pattern"));
 	GetToolManager()->DisplayMessage(
 		LOCTEXT("OnStartPatternTool", "Create Patterns for the selected Objects"),
@@ -722,6 +800,10 @@ void UPatternTool::OnShutdown(EToolShutdownType ShutdownType)
 
 	GetToolManager()->GetPairedGizmoManager()->DestroyAllGizmosByOwner(this);
 
+	ensure(bPatternToolThreeAxisTransformGizmoRegistered);
+	GetToolManager()->GetPairedGizmoManager()->DeregisterGizmoType(PatternToolThreeAxisTransformBuilderIdentifier);
+	bPatternToolThreeAxisTransformGizmoRegistered = false;
+	
 	OnSourceVisibilityToggled(true);
 
 	if (ShutdownType == EToolShutdownType::Accept)
@@ -737,7 +819,6 @@ void UPatternTool::Render(IToolsContextRenderAPI* RenderAPI)
 {
 	DragAlignmentMechanic->Render(RenderAPI);
 	PlaneMechanic->Render(RenderAPI);
-
 	// We should only render here if bVisualize is true and also visible in the details panel.
 	// We don't want to leave the user with no obvious way to turn it off which might happen if they
 	// enable bVisualize and then change the spacing mode to something other than packed
@@ -903,7 +984,9 @@ void UPatternTool::OnPropertyModified(UObject* PropertySet, FProperty* Property)
 void UPatternTool::OnMainFrameUpdated()
 {
 	CurrentStartFrameWorld = PlaneMechanic->Plane;
+
 	MarkPatternDirty();
+	ResetTransformGizmoPosition();
 }
 
 
@@ -917,7 +1000,29 @@ void UPatternTool::OnShapeUpdated()
 	SetToolPropertySourceEnabled(GridSettings, bGridSettings);
 	SetToolPropertySourceEnabled(RadialSettings, bRadialSettings);
 
+	// This keeps bUsingSingleAxis correct, fixes up gizmos, and calls OnParametersUpdated
+	if (bLinearSettings)
+	{
+		OnSingleAxisUpdated();
+	}
+	else if (bGridSettings || bRadialSettings)
+	{
+		OnSinglePlaneUpdated();
+	}
+}
+
+void UPatternTool::OnSingleAxisUpdated()
+{
+	bUsingSingleAxis = true;
+	OnParametersUpdated();
+	ReconstructTransformGizmos();
+}
+
+void UPatternTool::OnSinglePlaneUpdated()
+{
+	bUsingSingleAxis = false;
 	OnParametersUpdated();
+	ReconstructTransformGizmos();
 }
 
 void UPatternTool::OnSpacingModeUpdated()
@@ -957,6 +1062,144 @@ void UPatternTool::OnParametersUpdated()
 }
 
 
+void UPatternTool::OnTransformGizmoUpdated(UTransformProxy* Proxy, FTransform Transform)
+{
+	const FVector3d GizmoToToolOrigin = PatternGizmoProxy->GetTransform().GetLocation() - CurrentStartFrameWorld.Origin;
+	
+	// Recompute extents using current shape, SingleAxis/SinglePlane, and PatternGizmo position
+	switch (Settings->Shape)
+	{
+	case EPatternToolShape::Line:
+		LinearSettings->Extent = GizmoToToolOrigin.ProjectOnTo(CurrentStartFrameWorld.GetAxis((int) Settings->SingleAxis)).Length();
+		if (LinearSettings->bCentered) { LinearSettings->Extent *= 2.0f; }
+
+		LinearSettings->SilentUpdateWatcherAtIndex(LinearExtentWatcherIdx);
+		break;
+	case EPatternToolShape::Grid:
+		{
+		const int32 XIndex = Settings->SinglePlane == EPatternToolSinglePlane::YZPlane ? 1 : 0;
+		const int32 YIndex = Settings->SinglePlane == EPatternToolSinglePlane::XYPlane ? 1 : 2;
+
+		GridSettings->ExtentX = GizmoToToolOrigin.ProjectOnTo(CurrentStartFrameWorld.GetAxis(XIndex)).Length();
+		if (GridSettings->bCenteredX) { GridSettings->ExtentX *= 2.0f; }
+
+		GridSettings->ExtentY = GizmoToToolOrigin.ProjectOnTo(CurrentStartFrameWorld.GetAxis(YIndex)).Length();
+		if (GridSettings->bCenteredY) { GridSettings->ExtentY *= 2.0f; }
+
+		GridSettings->SilentUpdateWatcherAtIndex(GridExtentXWatcherIdx);
+		GridSettings->SilentUpdateWatcherAtIndex(GridExtentYWatcherIdx);
+		break;
+		}
+	case EPatternToolShape::Circle:
+		RadialSettings->Radius = GizmoToToolOrigin.Length();	// This is actually the simplest case because of how conveniently a circle is defined
+		
+		RadialSettings->SilentUpdateWatcherAtIndex(RadiusWatcherIdx);
+		break;
+	}
+	
+	MarkPatternDirty();
+}
+
+void UPatternTool::ResetTransformGizmoPosition()
+{
+	FVector3d OffsetFromOrigin = FVector3d::ZeroVector;
+	double DistanceFromOriginX = 0.0f;
+	double DistanceFromOriginY = 0.0f;
+	bool bDistancesSet = false;
+
+	// The gizmo is always restricted to moving in its local x-axis or xy-plane even when other axes or planes
+	// are used in the tool. This rotation correctly adjusts the gizmo's local frame such that the local x-axis
+	// or xy-plane correspond to the proper directions in the frame of the tool.
+	FRotator RotationInLocalSpace = FRotator::ZeroRotator;
+	
+	switch (Settings->Shape)
+	{
+	case EPatternToolShape::Line:
+		
+		DistanceFromOriginX = LinearSettings->Extent;
+		if (LinearSettings->bCentered) { DistanceFromOriginX /= 2.0f; }
+
+		if (Settings->SingleAxis == EPatternToolSingleAxis::XAxis)
+		{
+			OffsetFromOrigin = FVector3d(DistanceFromOriginX, 0, 0);
+		}
+		else if (Settings->SingleAxis == EPatternToolSingleAxis::YAxis)
+		{
+			OffsetFromOrigin = FVector3d(0, DistanceFromOriginX, 0);
+			RotationInLocalSpace.Add(0, 90, 0);
+		}
+		else if (Settings->SingleAxis == EPatternToolSingleAxis::ZAxis)
+		{
+			OffsetFromOrigin = FVector3d(0, 0, DistanceFromOriginX);
+			RotationInLocalSpace.Add(90, 0, 0);
+		}
+		
+		break;
+
+	case EPatternToolShape::Grid:
+		
+		DistanceFromOriginX = GridSettings->ExtentX;
+		if (GridSettings->bCenteredX) { DistanceFromOriginX /= 2.0f; }
+
+		DistanceFromOriginY = GridSettings->ExtentY;
+		if (GridSettings->bCenteredY) { DistanceFromOriginY /= 2.0f; }
+		
+		bDistancesSet = true;
+		
+		// FALLTHROUGH TO NEXT CASE
+		
+	case EPatternToolShape::Circle:
+		
+		if (!bDistancesSet)
+		{
+			DistanceFromOriginX = RadialSettings->Radius;
+			DistanceFromOriginY = 0.0f;
+		}
+
+		if (Settings->SinglePlane == EPatternToolSinglePlane::XYPlane)
+		{
+			OffsetFromOrigin = FVector3d(DistanceFromOriginX, DistanceFromOriginY, 0);
+		}
+		else if (Settings->SinglePlane == EPatternToolSinglePlane::XZPlane)
+		{
+			OffsetFromOrigin = FVector3d(DistanceFromOriginX, 0, DistanceFromOriginY);
+			RotationInLocalSpace.Add(0, 0, 90);
+		}
+		else if (Settings->SinglePlane == EPatternToolSinglePlane::YZPlane)
+		{
+			OffsetFromOrigin = FVector3d(0, DistanceFromOriginX, DistanceFromOriginY);
+			RotationInLocalSpace.Add(90, 0, 0);
+		}
+		
+		break;
+	}
+
+	const FVector3d ProxyTranslation = CurrentStartFrameWorld.Origin + FRotator(CurrentStartFrameWorld.Rotation).RotateVector(OffsetFromOrigin);
+
+	PatternGizmo->SetNewGizmoTransform(FTransform(FQuat(CurrentStartFrameWorld.Rotation * FQuaterniond(RotationInLocalSpace)), ProxyTranslation));
+}
+
+void UPatternTool::ReconstructTransformGizmos()
+{
+	if (bPatternToolThreeAxisTransformGizmoRegistered)
+	{
+		// Determining which elements of the CombinedTransformGizmo will be used
+		GizmoActorBuilder->EnableElements = bUsingSingleAxis ? ETransformGizmoSubElements::TranslateAxisX : ETransformGizmoSubElements::TranslatePlaneXY;
+
+		// Reconstructing gizmos with proper elements and transform
+		UInteractiveGizmoManager* GizmoManager = GetToolManager()->GetPairedGizmoManager();
+		GizmoManager->DestroyAllGizmosByOwner(this);
+		
+		PatternGizmo = Cast<UCombinedTransformGizmo>(GizmoManager->CreateGizmo(PatternToolThreeAxisTransformBuilderIdentifier, FString(), this));
+
+		// Necessary to force underlying AxisSources to axes in local space to restrict gizmo movement to tool frame axes
+		PatternGizmo->bUseContextCoordinateSystem = false;
+		PatternGizmo->CurrentCoordinateSystem = EToolContextCoordinateSystem::Local;
+
+		PatternGizmo->SetActiveTarget(PatternGizmoProxy);
+		ResetTransformGizmoPosition();
+	}
+}
 
 
 void UPatternTool::ResetPreviews()
@@ -1132,6 +1375,8 @@ void UPatternTool::GetPatternTransforms_Radial(TArray<UE::Geometry::FTransformSR
 	ComputeCombinedPatternBounds();
 	Generator.Dimensions = CombinedPatternBounds;
 	
+	Generator.Radius = RadialSettings->Radius;
+	
 	// Orient the CenterFrame based on the plane setting of the tool & determine which axis is used if bOriented is true
 	const FVector3d Origin(0, 0, 0);
 	switch (Settings->SinglePlane)
@@ -1155,9 +1400,7 @@ void UPatternTool::GetPatternTransforms_Radial(TArray<UE::Geometry::FTransformSR
 	Generator.StartAngleDeg = RadialSettings->StartAngle;
 	Generator.EndAngleDeg = RadialSettings->EndAngle;
 	Generator.AngleShift = RadialSettings->AngleShift;
-
-	Generator.Radius = RadialSettings->Radius;
-
+	
 	Generator.bOriented = RadialSettings->bOriented;
 
 	Generator.SpacingMode = (FRadialPatternGenerator::ESpacingMode)(int)RadialSettings->SpacingMode;
diff --git a/Engine/Plugins/Experimental/MeshModelingToolsetExp/Source/MeshModelingToolsExp/Public/PatternTool.h b/Engine/Plugins/Experimental/MeshModelingToolsetExp/Source/MeshModelingToolsExp/Public/PatternTool.h
--- a/Engine/Plugins/Experimental/MeshModelingToolsetExp/Source/MeshModelingToolsExp/Public/PatternTool.h
+++ b/Engine/Plugins/Experimental/MeshModelingToolsetExp/Source/MeshModelingToolsExp/Public/PatternTool.h
@@ -13,6 +13,7 @@
 #include "TransformTypes.h"
 #include "PatternTool.generated.h"
 
+class FCombinedTransformGizmoActorFactory;
 class UBaseAxisTranslationGizmo;
 class UAxisAngleGizmo;
 class UDragAlignmentMechanic;
@@ -474,18 +475,38 @@ public:
 	TObjectPtr<UPatternTool_OutputSettings> OutputSettings;
 
 protected:
-
+	/**
+	 * Pattern Gizmo:
+	 */
+	FString PatternToolThreeAxisTransformBuilderIdentifier = TEXT("PatternToolThreeAxisTransformBuilderIdentifier");
+	FString PatternToolAxisPositionBuilderIdentifier = TEXT("PatternToolAxisPositionBuilderIdentifier");
+	FString PatternToolPlanePositionBuilderIdentifier = TEXT("PatternToolPlanePositionBuilderIdentifier");
+	TSharedPtr<FCombinedTransformGizmoActorFactory> GizmoActorBuilder;
+	bool bPatternToolThreeAxisTransformGizmoRegistered = false;
+	
 	UPROPERTY()
-	TObjectPtr<UTransformProxy> TransformProxy_End = nullptr;
+	TObjectPtr<UTransformProxy> PatternGizmoProxy = nullptr;
 
 	UPROPERTY()
-	TObjectPtr<UCombinedTransformGizmo> TransformGizmo_End = nullptr;
+	TObjectPtr<UCombinedTransformGizmo> PatternGizmo = nullptr;
 
+	// If true, Settings->SingleAxis is being used. If false, Settings->SinglePlane is being used
+	bool bUsingSingleAxis;
+	
+	int32 LinearExtentWatcherIdx;
+	int32 GridExtentXWatcherIdx;
+	int32 GridExtentYWatcherIdx;
+	int32 RadiusWatcherIdx;
+
+	void OnTransformGizmoUpdated(UTransformProxy* Proxy, FTransform Transform);
+	void ResetTransformGizmoPosition();
+	void ReconstructTransformGizmos();
+	
 
+	
 	UPROPERTY()
 	TObjectPtr<UDragAlignmentMechanic> DragAlignmentMechanic = nullptr;
 
-
 	UPROPERTY()
 	TObjectPtr<UConstructionPlaneMechanic> PlaneMechanic = nullptr;
 
@@ -499,6 +520,8 @@ protected:
 	void OnSourceVisibilityToggled(bool bVisible);
 	void OnMainFrameUpdated();
 	void OnShapeUpdated();
+	void OnSingleAxisUpdated();
+	void OnSinglePlaneUpdated();
 	void OnSpacingModeUpdated();
 	void OnParametersUpdated();
 	void UpdatePattern();
