fixed skin weights painting crashes
    - EditedMesh must be created before calling UpdateBonePositionInfos as it's used to compute the radius when there's only 1 bone
    - MeshDescription to DynamicMesh conversion can result in DynamicMesh having more vertices than the MeshDescription because of non-manifold triangles (FMeshDescriptionToDynamicMesh::Convert for details)

diff --git a/Engine/Plugins/Experimental/MeshModelingToolsetExp/Source/MeshModelingToolsExp/Private/SkinWeightsPaintTool.cpp b/Engine/Plugins/Experimental/MeshModelingToolsetExp/Source/MeshModelingToolsExp/Private/SkinWeightsPaintTool.cpp
--- a/Engine/Plugins/Experimental/MeshModelingToolsetExp/Source/MeshModelingToolsExp/Private/SkinWeightsPaintTool.cpp
+++ b/Engine/Plugins/Experimental/MeshModelingToolsetExp/Source/MeshModelingToolsExp/Private/SkinWeightsPaintTool.cpp
@@ -66,6 +66,9 @@ void USkinWeightsPaintTool::Setup()
 	IPrimitiveComponentBackedTarget* TargetComponent = Cast<IPrimitiveComponentBackedTarget>(Target);
 	USkeletalMeshComponent* Component = Cast<USkeletalMeshComponent>(TargetComponent->GetOwnerComponent());
 
+	EditedMesh = MakeUnique<FMeshDescription>();
+	*EditedMesh = *UE::ToolTarget::GetMeshDescription(Target);
+
 	// hide strength and falloff
 	BrushProperties->RestoreProperties(this);
 
@@ -131,9 +134,6 @@ void USkinWeightsPaintTool::Setup()
 		LOCTEXT("OnStartSkinWeightsPaint", "Paint per-bone skin weights. [ and ] change brush size, Ctrl to Erase/Subtract, Shift to Smooth"),
 		EToolMessageLevel::UserNotification);
 
-	EditedMesh = MakeUnique<FMeshDescription>();
-	*EditedMesh = *UE::ToolTarget::GetMeshDescription(Target);
-
 	InitializeSkinWeights();
 
 	CurrentBoneWatcher.Initialize([this]() { return ToolProps->CurrentBone; },
@@ -322,7 +322,7 @@ void USkinWeightsPaintTool::UpdateBoneVisualization()
 		for (int32 ElementId : ColorOverlay->ElementIndicesItr())
 		{
 			const int32 VertexId = ColorOverlay->GetParentVertex(ElementId);
-			const float Value = SkinWeightsData[VertexId];
+			const float Value = SkinWeightsData.IsValidIndex(VertexId) ? SkinWeightsData[VertexId] : 0.f;
 			const FVector4f Color(WeightToColor(Value));
 			ColorOverlay->SetElement(ElementId, Color);
 		}
