Move the skin binding op to a shared location in ModelingOps. Placed the implementation of the occupancy grid into its own helper header, as FOccupancyGrid3.

#rnx

diff --git a/Engine/Plugins/Experimental/MeshModelingToolsetExp/Source/MeshModelingToolsExp/Private/SkinWeightsBindingTool.cpp b/Engine/Plugins/Experimental/MeshModelingToolsetExp/Source/MeshModelingToolsExp/Private/SkinWeightsBindingTool.cpp
--- a/Engine/Plugins/Experimental/MeshModelingToolsetExp/Source/MeshModelingToolsExp/Private/SkinWeightsBindingTool.cpp
+++ b/Engine/Plugins/Experimental/MeshModelingToolsetExp/Source/MeshModelingToolsExp/Private/SkinWeightsBindingTool.cpp
@@ -14,716 +14,17 @@
 #include "Components/SkeletalMeshComponent.h"
 #include "Engine/SkeletalMesh.h"
 #include "SceneManagement.h"
-#include "Spatial/FastWinding.h"
 #include "Spatial/MeshWindingNumberGrid.h"
 
-#include "TargetInterfaces/PrimitiveComponentBackedTarget.h"
 #include "ModelingToolTargetUtil.h"
+#include "SkinningOps/SkinBindingOp.h"
+#include "Spatial/OccupancyGrid3.h"
 
 #include UE_INLINE_GENERATED_CPP_BY_NAME(SkinWeightsBindingTool)
 
 
-// #pragma optimize( "", off )
-
 #define LOCTEXT_NAMESPACE "USkinWeightsBindingTool"
 
-// TODO: Move to a helper function.
-
-// A simple FIFO queue. Maintains a set of blocks, rather than allocate for each element.
-template<typename T, int32 BlockSize=256>
-class TFIFOQueue
-{
-public:
-	TFIFOQueue() = default;
-
-	void Push(T InElem)
-	{
-		// Container Empty?
-		if (Blocks.IsEmpty())
-		{
-			Blocks.SetNum(1);
-			Blocks[0].bIsFree = false;
-			PushIndex = -1;
-		}
-		// Or are we at the end of current block?
-		else if (PushIndex == (BlockSize - 1))
-		{
-			PushBlock = AllocateNewPushBlock();
-			PushIndex = -1;
-		}
-
-		PushIndex++;
-
-		Blocks[PushBlock].Data[PushIndex] = InElem;
-	}
-
-	T Pop()
-	{
-		check(!IsEmpty());
-
-		if (PopIndex == (BlockSize - 1))
-		{
-			// Reached the end. Free this block and move onto the next one.
-			PopBlock = FreeCurrentPopBlock();
-			PopIndex = -1;
-		}
-
-		PopIndex++;
-		return Blocks[PopBlock].Data[PopIndex];
-	}
-
-	bool TryPop(T &OutValue)
-	{
-		if (IsEmpty())
-		{
-			return false;
-		}
-
-		OutValue = Pop();
-		return true;
-	}
-
-	void Reset()
-	{
-		Blocks.Reset();
-		PushBlock = 0;
-		PushIndex = -1;
-		PopBlock = 0;
-		PopIndex = -1;
-	}
-	
-
-	bool IsEmpty() const
-	{
-		return PushBlock == PopBlock && PushIndex == PopIndex;
-	}
-
-private:
-	int32 AllocateNewPushBlock()
-	{
-		checkSlow(PushBlock != -1);
-
-		int NewBlockIndex = INDEX_NONE;
-		
-		if (FreeCount > 0)
-		{
-			// Find a free block
-			for (int32 Index = 0; Index < Blocks.Num(); Index++)
-			{
-				if (Blocks[Index].bIsFree)
-				{
-					NewBlockIndex = Index;
-					FreeCount--;
-					break;
-				}
-			}
-			checkfSlow(NewBlockIndex != INDEX_NONE, TEXT("We should have found a free block."));
-		}
-		else
-		{
-			NewBlockIndex = Blocks.Num();
-			Blocks.AddDefaulted();
-		}
-
-		Blocks[PushBlock].NextBlock = NewBlockIndex;
-		Blocks[NewBlockIndex].bIsFree = false;
-		Blocks[NewBlockIndex].NextBlock = INDEX_NONE;
-		return NewBlockIndex;
-	}
-
-	int32 FreeCurrentPopBlock()
-	{
-		const int32 NextBlock = Blocks[PopBlock].NextBlock;
-		checkSlow(NextBlock != -1);
-		Blocks[PopBlock].bIsFree = true;
-		FreeCount++;
-
-		return NextBlock;
-	}
-	
-
-	int32 PushBlock = 0;
-	int32 PushIndex = -1;
-
-	int32 PopBlock = 0;
-	int32 PopIndex = -1;
-
-	int32 FreeCount = 0;
-	
-	struct FBlock
-	{
-		FBlock()
-		{
-			Data = new T[BlockSize];
-		}
-		~FBlock()
-		{
-			delete [] Data;
-		}
-
-		FBlock(FBlock &&InOther) noexcept
-		{
-			Data = InOther.Data;
-			InOther.Data = nullptr;
-		}
-		
-		bool bIsFree = true;
-		int32 NextBlock = INDEX_NONE;
-		T *Data;
-	};
-	
-	TArray<FBlock> Blocks;
-};
-
-
-static float DistanceToLineSegment(const FVector& P, const FVector& A, const FVector& B)
-{
-	const FVector M = B - A;
-	const FVector T = P - A;
-
-	const float C1 = FVector::DotProduct(M, T);
-	if (C1 <= 0.0f)
-	{
-		return FVector::Dist(P, A);
-	}
-
-	const float C2 = FVector::DotProduct(M, M);
-	if (C2 <= C1)
-	{
-		return FVector::Dist(P, B);
-	}
-
-	// Project the point onto the line and get the distance between them.
-	const FVector PT = A + M * (C1 / C2);
-	return FVector::Dist(P, PT);
-}
-
-
-// List of bones as used by the binding class. In this case for each bone transform, we want to
-// store a list of line segments going from the bone transform to all the child bone transforms.
-struct FTransformHierarchyQuery
-{
-	explicit FTransformHierarchyQuery(const TArray<TPair<FTransform, int32>>& InTransformHierarchy)
-	{
-		for (int Index = 0; Index < InTransformHierarchy.Num(); Index++)
-		{
-			FTransform Xform = InTransformHierarchy[Index].Key;
-			int32 ParentIndex = InTransformHierarchy[Index].Value;
-
-			while (ParentIndex != INDEX_NONE)
-			{
-				Xform = Xform * InTransformHierarchy[ParentIndex].Key;
-				ParentIndex = InTransformHierarchy[ParentIndex].Value;
-			}
-
-			BoneFans.Add({ Xform.GetLocation() });
-		}
-
-		// Fill in the fan tips, as needed.
-		for (int Index = 0; Index < InTransformHierarchy.Num(); Index++)
-		{
-			const int32 ParentIndex = InTransformHierarchy[Index].Value;
-			if (ParentIndex != INDEX_NONE)
-			{
-				BoneFans[ParentIndex].TipsPos.Add(BoneFans[Index].RootPos);
-			}
-		}		
-	}
-
-	float GetDistanceToBoneFan(const int32 InBoneIndex, const FVector& InPoint) const
-	{
-		return BoneFans[InBoneIndex].GetDistance(InPoint);
-	}
-	
-	FBox GetBoneFanBBox(const int32 InBoneIndex) const
-	{
-		return BoneFans[InBoneIndex].GetBBox();
-	}
-
-	bool GetBoneFanIntersectsBox(const int32 InBoneIndex, const FBox &InBox) const
-	{
-		return BoneFans[InBoneIndex].IntersectsBox(InBox);
-	}
-
-private:
-	struct FBoneFan
-	{
-		FVector RootPos;
-		TArray<FVector> TipsPos;
-
-		float GetDistance(const FVector& InPoint) const
-		{
-			if (TipsPos.IsEmpty())
-			{
-				return FVector::Distance(RootPos, InPoint);
-			}
-			else
-			{
-				float Distance = std::numeric_limits<float>::max();
-				for (const FVector& TipPos: TipsPos)
-				{
-					Distance = FMath::Min(Distance, DistanceToLineSegment(InPoint, RootPos, TipPos));
-				}
-				return Distance;
-			}
-		}
-
-		FBox GetBBox() const
-		{
-			FBox Box(RootPos, RootPos);
-			for (const FVector& TipPos: TipsPos)
-			{
-				Box += TipPos;
-			}
-
-			return Box;
-		}
-
-		bool IntersectsBox(const FBox &InBox) const
-		{
-			if (TipsPos.IsEmpty())
-			{
-				return FMath::PointBoxIntersection(RootPos, InBox);
-			}
-
-			if (GetBBox().Intersect(InBox))
-			{
-				for (const FVector& TipPos: TipsPos)
-				{
-					if (FMath::LineBoxIntersection(InBox, RootPos, TipPos, TipPos - RootPos))
-					{
-						return true;
-					}
-				}
-			}
-			return false;
-		}				
-	};
-
-	 TArray<FBoneFan> BoneFans;
-};
-
-
-struct FOccupancyGrid
-{
-	enum EDomain : int32
-	{
-		Exterior,
-		Boundary,
-		Interior
-	};
-
-	UE::Geometry::FDenseGrid3i Occupancy;
-	float CellSize;
-	FVector3f GridOrigin;
-	FVector3f CellMidPoint;
-
-	FOccupancyGrid(
-		const FDynamicMesh3& InMesh,
-		int32 InVoxelResolution
-		)
-	{
-		using namespace UE::Geometry;
-	
-		// Compute a voxel grid 
-		FDynamicMeshAABBTree3 Spatial(&InMesh);
-		TFastWindingTree FastWinding(&Spatial);
-		FAxisAlignedBox3d Bounds = Spatial.GetBoundingBox();
-		CellSize = Bounds.MaxDim() / InVoxelResolution;
-		CellMidPoint = FVector3f(CellSize / 2.0f, CellSize / 2.0f, CellSize / 2.0f);
-
-		TMeshWindingNumberGrid WindingGrid(&InMesh, &FastWinding, CellSize);
-	
-		WindingGrid.Compute();
-
-		// Our occupancy grid is computed on the winding number grid's cell centers.
-		const FVector3i WindingDims = WindingGrid.Dimensions();
-		Occupancy = FDenseGrid3i(WindingDims.X - 1, WindingDims.Y - 1, WindingDims.Z - 1, Exterior);
-
-		GridOrigin = WindingGrid.GridOrigin + FVector3f(CellSize / 2.0f, CellSize / 2.0f, CellSize / 2.0f);
-
-		static const FVector3i CornerOffsets[] = {
-			FVector3i(0, 0, 0),
-			FVector3i(0, 0, 1),
-			FVector3i(0, 1, 0),
-			FVector3i(0, 1, 1),
-			FVector3i(1, 0, 0),
-			FVector3i(1, 0, 1),
-			FVector3i(1, 1, 0),
-			FVector3i(1, 1, 1),
-		};
-		
-		// TODO: Parallel-for
-		for (int32 OccupancyId = 0; OccupancyId < Occupancy.Size(); OccupancyId++)
-		{
-			const FVector3i OccupancyIndex(Occupancy.ToIndex(OccupancyId));
-			int32 Count = 0;
-			for (int32 CornerId = 0; CornerId < 8; CornerId++)
-			{
-				const FVector3i CornerIndex(OccupancyIndex + CornerOffsets[CornerId]);
-				if (WindingGrid.GetValue(CornerIndex) >= WindingGrid.WindingIsoValue)
-				{
-					Count++;
-				}
-			}
-
-			if (Count == 8)
-			{
-				Occupancy[OccupancyIndex] = Interior;
-			}
-			else if (Count > 0)
-			{
-				Occupancy[OccupancyIndex] = Boundary;
-			}
-		}
-
-		// Make sure we include all the vertices of the mesh as a part of the boundary, if
-		// the vertex areas are marked as being exterior.
-		for (int32 VertexIdx = 0; VertexIdx < InMesh.VertexCount(); VertexIdx++)
-		{
-			const FVector3d& Pos = InMesh.GetVertex(VertexIdx);
-			const FVector3i OccupancyIndex = GetCellIndexFromPoint(FVector(Pos));
-			if (Occupancy[OccupancyIndex] == Exterior)
-			{
-				Occupancy[OccupancyIndex] = Boundary;
-			}
-		}
-	}
-	
-	UE::Geometry::FVector3i GetCellIndexFromPoint(const FVector &InPoint) const
-	{
-		FVector3f PP(InPoint);
-		PP -= GridOrigin;
-		PP += CellMidPoint;
-
-		return { FMath::FloorToInt(PP.X / CellSize),
-				FMath::FloorToInt(PP.Y / CellSize),
-				FMath::FloorToInt(PP.Z / CellSize) };
-	}
-
-	FVector3d GetCellCenterFromIndex(const UE::Geometry::FVector3i &Index) const
-	{
-		const float CS = CellSize;
-		return {Index.X * CS + GridOrigin.X, Index.Y * CS + GridOrigin.Y, Index.Z * CS + GridOrigin.Z};
-	}
-
-	FBox GetCellBoxFromIndex(const UE::Geometry::FVector3i &Index) const
-	{
-		const FVector3f P = (FVector3f)GetCellCenterFromIndex(Index);
-		return {P - CellMidPoint, P + CellMidPoint};
-	}
-};
-
-
-namespace
-{
-	struct FCreateSkinWeights_Closest_WorkData final :
-		TThreadSingleton<FCreateSkinWeights_Closest_WorkData>
-	{
-		TArray<TPair<FBoneIndexType, float>> RawBoneWeights;
-		TArray<UE::AnimationCore::FBoneWeight> BoneWeights;
-	};
-
-	float ComputeWeightStiffness(const float InWeight, const float InStiffness)
-	{
-		return (1.0f - InStiffness) * InWeight + InStiffness * InWeight * InWeight;
-	}
-	
-}
-
-
-class FComputeSkinWeightsBindingOp : public UE::Geometry::FDynamicMeshOperator
-{
-public:
-	virtual ~FComputeSkinWeightsBindingOp() override {}
-
-	// The transform hierarchy to bind to. Listed in the same order as the bones in the
-	// reference skeleton that this skelmesh is tied to.
-	TSharedPtr<FDynamicMesh3, ESPMode::ThreadSafe> OriginalMesh;
-	TArray<TPair<FTransform, int32>> TransformHierarchy;
-
-	FName ProfileName = FSkeletalMeshAttributes::DefaultSkinWeightProfileName;
-	
-	ESkinWeightsBindType BindType = ESkinWeightsBindType::DirectDistance;
-	float Stiffness = 0.2f;
-	int32 MaxInfluences = 5;
-	int32 VoxelResolution = 256;
-	
-	void CalculateResult(FProgressCancel* InProgress) override
-	{
-		if (InProgress && InProgress->Cancelled())
-		{
-			return;
-		}
-
-		ResultMesh->Copy(*OriginalMesh, true, true, true, true);
-		
-		if (InProgress && InProgress->Cancelled())
-		{
-			return;
-		}
-
-		const float ClampedStiffness = FMath::Clamp(Stiffness, 0.0f, 1.0f);	
-
-		UE::AnimationCore::FBoneWeightsSettings Settings;
-		Settings.SetMaxWeightCount(MaxInfluences);
-
-		switch(BindType)
-		{
-		case ESkinWeightsBindType::DirectDistance:
-			CreateSkinWeights_DirectDistance(*ResultMesh, ClampedStiffness, Settings);
-			break;
-			
-		case ESkinWeightsBindType::GeodesicVoxel:
-			CreateSkinWeights_GeodesicVoxel(*ResultMesh, ClampedStiffness, Settings);
-			break;
-		}
-	}
-
-private:
-	static UE::Geometry::FDynamicMeshVertexSkinWeightsAttribute *GetOrCreateSkinWeightsAttribute(
-		FDynamicMesh3& InMesh,
-		FName InProfileName
-		)
-	{
-		using namespace UE::Geometry;
-		
-		FDynamicMeshVertexSkinWeightsAttribute *Attribute = InMesh.Attributes()->GetSkinWeightsAttribute(InProfileName);
-		if (!Attribute)
-		{
-			Attribute = new FDynamicMeshVertexSkinWeightsAttribute(&InMesh);
-			InMesh.Attributes()->AttachSkinWeightsAttribute(InProfileName, Attribute);
-		}
-		return Attribute;
-	}
-	
-	void CreateSkinWeights_DirectDistance(
-		FDynamicMesh3& InMesh,
-		float InStiffness,
-		const UE::AnimationCore::FBoneWeightsSettings& InSettings
-		)
-	{
-		using namespace UE::AnimationCore;
-		using namespace UE::Geometry;
-
-		FDynamicMeshVertexSkinWeightsAttribute *SkinWeights = InMesh.Attributes()->GetSkinWeightsAttribute(ProfileName);
-		
-		const int32 NumVertices = InMesh.VertexCount();
-
-		// Use the diagonal size of the bbox to make the bone distance falloff scale invariant.
-		const float DiagBounds = InMesh.GetBounds(true).DiagonalLength();
-	
-		const FTransformHierarchyQuery Skeleton(TransformHierarchy);
-
-		ParallelFor(NumVertices, [&](const int32 VertexIdx)
-		{
-			const FVector3d& Pos = InMesh.GetVertex(VertexIdx);
-
-			FCreateSkinWeights_Closest_WorkData &WorkData = FCreateSkinWeights_Closest_WorkData::Get();
-		
-			WorkData.RawBoneWeights.Reset(TransformHierarchy.Num());
-
-			float TotalWeight = 0.0f;
-			for (int32 BoneIndex = 0; BoneIndex < TransformHierarchy.Num(); BoneIndex++)
-			{
-				// Normalize the distance by the diagonal size of the bbox to maintain scale invariance.
-				float Weight = Skeleton.GetDistanceToBoneFan(BoneIndex, Pos) / DiagBounds;
-
-				// Avoid div-by-zero but allow for the possibility that multiple bones may
-				// touch this vertex.
-				Weight = FMath::Max(Weight, KINDA_SMALL_NUMBER);
-
-				// Compute the actual weight, factoring in the stiffness value. W = (1/S(D))^2
-				// Where S(x) is the stiffness function.
-				Weight = FMath::Square(1.0f / ComputeWeightStiffness(Weight, InStiffness));
-				TotalWeight += Weight;
-				WorkData.RawBoneWeights.Add(MakeTuple(static_cast<FBoneIndexType>(BoneIndex), Weight));
-			}
-
-			// Normalize
-			for (TPair<FBoneIndexType, float> &BoneWeight: WorkData.RawBoneWeights)
-			{
-				BoneWeight.Value /= TotalWeight;
-			}
-			
-			WorkData.RawBoneWeights.Sort([](const TPair<FBoneIndexType, float> &A, const TPair<FBoneIndexType, float> &B)
-			{
-				return A.Value > B.Value;
-			});
-
-			WorkData.BoneWeights.Reset(InSettings.GetMaxWeightCount());
-			for (int32 BoneIndex = 0; BoneIndex < FMath::Min(InSettings.GetMaxWeightCount(), WorkData.RawBoneWeights.Num()); BoneIndex++)
-			{
-				const TPair<FBoneIndexType, float>& BoneWeight = WorkData.RawBoneWeights[BoneIndex];
-				WorkData.BoneWeights.Add(FBoneWeight(BoneWeight.Key, BoneWeight.Value));
-			}
-
-			SkinWeights->SetValue(VertexIdx, FBoneWeights::Create(WorkData.BoneWeights, InSettings));
-		});
-	}
-	
-	void CreateSkinWeights_GeodesicVoxel(
-		FDynamicMesh3& InMesh,
-		float InStiffness,
-		const UE::AnimationCore::FBoneWeightsSettings& InSettings
-		)
-	{
-		using namespace UE::AnimationCore;
-		using namespace UE::Geometry;
-
-		FDynamicMeshVertexSkinWeightsAttribute *SkinWeights = InMesh.Attributes()->GetSkinWeightsAttribute(ProfileName);
-		
-		const int32 NumVertices = InMesh.VertexCount();
-
-		// Use the diagonal size of the bbox to make the bone distance falloff scale invariant.
-		const float DiagBounds = InMesh.GetBounds(true).DiagonalLength();
-		
-		const FTransformHierarchyQuery Skeleton(TransformHierarchy);
-
-		// This is grossly inefficient but tricky to do otherwise, since each bone distance
-		// computation is done per-thread. We could possibly solve this by chunking instead
-		// and accumulating partial results.
-		TArray<float> Weights;
-		Weights.SetNumUninitialized(NumVertices * TransformHierarchy.Num());
-
-		FOccupancyGrid Occupancy(InMesh, VoxelResolution);
-
-		const FVector3i Dimensions = Occupancy.Occupancy.GetDimensions();
-		
-		ParallelFor(TransformHierarchy.Num(), [&](int32 BoneIndex) {
-			TFIFOQueue<FVector3i> WorkingSet;  
-			FDenseGrid3f BoneDistance(Dimensions.X, Dimensions.Y, Dimensions.Z, DiagBounds);
-
-			// Mark all the cells that the bone intersects with distance of 0 and put them
-			// on the work queue.
-			const FBox BoneBox = Skeleton.GetBoneFanBBox(BoneIndex);
-			const FVector3i BoneMin = Occupancy.GetCellIndexFromPoint(BoneBox.Min);
-			const FVector3i BoneMax = Occupancy.GetCellIndexFromPoint(BoneBox.Max);
-
-			for (int32 I = BoneMin.X; I <= BoneMax.X; I++)
-			{
-				for (int32 J = BoneMin.Y; J <= BoneMax.Y; J++)
-				{
-					for (int32 K = BoneMin.Z; K <= BoneMax.Z; K++)
-					{
-						const FVector3i Candidate(I, J, K);
-						const FBox CellBox = Occupancy.GetCellBoxFromIndex(Candidate);
-						if (Skeleton.GetBoneFanIntersectsBox(BoneIndex, CellBox))
-						{
-							WorkingSet.Push(Candidate);
-							BoneDistance[Candidate] = 0.0f;
-						}
-					}
-				}
-			}
-
-			// Iterate over all the voxels until we have constructed shortest distance paths
-			// throughout the level set.
-			while (!WorkingSet.IsEmpty())
-			{
-				const FVector3i WorkItem = WorkingSet.Pop();
-
-				// Loop through each of the neighbours (6 face neighbours, 12 edge neighbors,
-				// and 8 corner neighbours) and see if any of them are closer to the bone
-				// than their current marked distance.
-				float CurrentDistance = BoneDistance[WorkItem];
-				for (int32 N = 0; N < 26; N++)
-				{
-					FVector3i Offset(IndexUtil::GridOffsets26[N]);
-					FVector3i Candidate(WorkItem + Offset);
-
-					if (!BoneDistance.IsValidIndex(Candidate))
-					{
-						continue;
-					}
-					
-					// Ensure this entry is either a part of the interior or boundary domain.
-					if (Occupancy.Occupancy[Candidate] == FOccupancyGrid::Exterior)
-					{
-						continue;
-					}
-
-					const float CellDistance = (FVector3f(Offset) * Occupancy.CellSize).Length();
-					const float CandidateDistance = CurrentDistance + CellDistance;
-					const float OldDistance = BoneDistance[Candidate]; 
-
-					if (OldDistance > CandidateDistance)
-					{
-						WorkingSet.Push(Candidate);
-						BoneDistance[Candidate] = CandidateDistance;
-					}
-				}
-			}
-
-			// Loop through all the vertices, find the voxel each belongs to, and compute
-			// the distance from the voxel to the vertex (assuming the distance stored in the
-			// voxel is based on traversing from voxel center to voxel center).
-			for (int32 VertexIdx = 0; VertexIdx < NumVertices; VertexIdx++)
-			{
-				const FVector3d& Pos = InMesh.GetVertex(VertexIdx);
-				const FVector3i CellIndex = Occupancy.GetCellIndexFromPoint(Pos);
-				const FVector3d CellCenter = Occupancy.GetCellCenterFromIndex(CellIndex);
-
-				float Distance = BoneDistance[CellIndex];
-				const FOccupancyGrid::EDomain Domain = static_cast<FOccupancyGrid::EDomain>(Occupancy.Occupancy[CellIndex]);
-				// check(Distance != std::numeric_limits<float>::max());
-
-				Distance += FVector3d::Distance(CellCenter, Pos);
-				
-				// Normalize the distance by the diagonal size of the bbox to maintain scale invariance.
-				float Weight = Distance / DiagBounds;
-
-				// Avoid div-by-zero but allow for the possibility that multiple bones may
-				// touch this vertex.
-				Weight = FMath::Max(Weight, KINDA_SMALL_NUMBER);
-
-				// Compute the actual weight, factoring in the stiffness value. W = (1/S(D))^2
-				// Where S(x) is the stiffness function.
-				Weight = FMath::Square(1.0f / ComputeWeightStiffness(Weight, InStiffness));
-
-				Weights[VertexIdx * TransformHierarchy.Num() + BoneIndex] = Weight;
-			}
-		});	
-
-		ParallelFor(NumVertices, [&](const int32 VertexIdx)
-		{
-			FCreateSkinWeights_Closest_WorkData &WorkData = FCreateSkinWeights_Closest_WorkData::Get();
-			
-			WorkData.RawBoneWeights.Reset(TransformHierarchy.Num());
-
-			float TotalWeight = 0.0f;
-			for (int32 BoneIndex = 0; BoneIndex < TransformHierarchy.Num(); BoneIndex++)
-			{
-				const float Weight = Weights[VertexIdx * TransformHierarchy.Num() + BoneIndex];
-				TotalWeight += Weight;
-				WorkData.RawBoneWeights.Add(MakeTuple(static_cast<FBoneIndexType>(BoneIndex), Weight));
-			}
-
-			// Normalize
-			for (TPair<FBoneIndexType, float> &BoneWeight: WorkData.RawBoneWeights)
-			{
-				BoneWeight.Value /= TotalWeight;
-			}
-				
-			WorkData.RawBoneWeights.Sort([](const TPair<FBoneIndexType, float> &A, const TPair<FBoneIndexType, float> &B)
-			{
-				return A.Value > B.Value;
-			});
-
-			WorkData.BoneWeights.Reset(InSettings.GetMaxWeightCount());
-			for (int32 BoneIndex = 0; BoneIndex < FMath::Min(InSettings.GetMaxWeightCount(), WorkData.RawBoneWeights.Num()); BoneIndex++)
-			{
-				const TPair<FBoneIndexType, float>& BoneWeight = WorkData.RawBoneWeights[BoneIndex];
-				WorkData.BoneWeights.Add(FBoneWeight(BoneWeight.Key, BoneWeight.Value));
-			}
-
-			SkinWeights->SetValue(VertexIdx, FBoneWeights::Create(WorkData.BoneWeights, InSettings));
-		});
-	}	
-};
-
-
-
 
 bool USkinWeightsBindingToolBuilder::CanBuildTool(const FToolBuilderState& SceneState) const
 {
@@ -771,41 +72,30 @@ void USkinWeightsBindingTool::Setup()
 		return;
 	}
 	
-	const USkeletalMeshComponent* SkelMeshComponent = Cast<USkeletalMeshComponent>(UE::ToolTarget::GetTargetComponent(Targets[0]));
+	const USkeletalMeshComponent* SkeletalMeshComponent = Cast<USkeletalMeshComponent>(UE::ToolTarget::GetTargetComponent(Targets[0]));
 	
-	if (SkelMeshComponent && SkelMeshComponent->GetSkeletalMeshAsset())
+	if (SkeletalMeshComponent && SkeletalMeshComponent->GetSkeletalMeshAsset())
 	{
-		USkeletalMesh* SkeletalMesh = SkelMeshComponent->GetSkeletalMeshAsset();
+		USkeletalMesh* SkeletalMesh = SkeletalMeshComponent->GetSkeletalMeshAsset();
 
+		ReferenceSkeleton = SkeletalMesh->GetRefSkeleton();
+		
 		// Initialize the bone browser
-		FCurveEvaluationOption CurveEvalOption(
-				SkelMeshComponent->GetAllowedAnimCurveEvaluate(),
-				&SkelMeshComponent->GetDisallowedAnimCurvesEvaluation(),
+		const FCurveEvaluationOption CurveEvalOption(
+				SkeletalMeshComponent->GetAllowedAnimCurveEvaluate(),
+				&SkeletalMeshComponent->GetDisallowedAnimCurvesEvaluation(),
 				0 /* Always use the highest LOD */
 				);
-		BoneContainer.InitializeTo(SkelMeshComponent->RequiredBones, CurveEvalOption, *SkeletalMesh);
-
-		Properties->SkeletalMesh = SkeletalMesh;
-		Properties->CurrentBone.Initialize(BoneContainer);
-
-		const FReferenceSkeleton& RefSkeleton = SkeletalMesh->GetRefSkeleton();
-		for (int32 Index = 0; Index < RefSkeleton.GetRawBoneNum(); Index++)
+		TArray<FBoneIndexType> AllBones;
+		AllBones.Reserve(ReferenceSkeleton.GetRawBoneNum());
+		for (int32 BoneIndex = 0; BoneIndex < ReferenceSkeleton.GetRawBoneNum(); BoneIndex++)
 		{
-			BoneToIndex.Add(RefSkeleton.GetRawRefBoneInfo()[Index].Name, Index);
+			AllBones.Add(static_cast<FBoneIndexType>(BoneIndex));
 		}
+		BoneContainer.InitializeTo(AllBones, CurveEvalOption, *SkeletalMesh);
 
-		// Pick the first root bone as the initial selection.
-		Properties->CurrentBone.BoneName = SkeletalMesh->GetRefSkeleton().GetBoneName(0);
-
-		const TArray<FMeshBoneInfo>& BoneInfo = RefSkeleton.GetRawRefBoneInfo();
-		const TArray<FTransform>& BonePose = RefSkeleton.GetRawRefBonePose();
-		
-		TransformHierarchy.Reserve(BoneInfo.Num());
-		
-		for (int32 Index = 0; Index < BoneInfo.Num(); Index++)
-		{
-			TransformHierarchy.Add(MakeTuple(BonePose[Index], BoneInfo[Index].ParentIndex));
-		}
+		Properties->SkeletalMesh = SkeletalMesh;
+		Properties->CurrentBone.Initialize(BoneContainer);
 	}
 	
 	UE::ToolTarget::HideSourceObject(Targets[0]);
@@ -851,7 +141,7 @@ void USkinWeightsBindingTool::Setup()
 	Preview->PreviewMesh->SetShadowsEnabled(false);
 	Preview->PreviewMesh->UpdatePreview(OriginalMesh.Get());
 
-	Occupancy = MakeShared<FOccupancyGrid>(*OriginalMesh, Properties->VoxelResolution);
+	Occupancy = MakeShared<UE::Geometry::FOccupancyGrid3>(*OriginalMesh, Properties->VoxelResolution);
 	
 	UpdateVisualization(/*bForce=*/true);
 	
@@ -873,7 +163,7 @@ void USkinWeightsBindingTool::OnShutdown(EToolShutdownType ShutdownType)
 
 	UE::ToolTarget::ShowSourceObject(Targets[0]);
 	
-	FDynamicMeshOpResult Result = Preview->Shutdown();
+	const FDynamicMeshOpResult Result = Preview->Shutdown();
 	if (ShutdownType == EToolShutdownType::Accept)
 	{
 		GenerateAsset(Result);
@@ -891,12 +181,12 @@ static void DrawBox(IToolsContextRenderAPI* RenderAPI, const FTransform& Transfo
 	FPrimitiveDrawInterface* PDI = RenderAPI->GetPrimitiveDrawInterface();
 	const float PDIScale = RenderAPI->GetCameraState().GetPDIScalingFactor();
 
-	FVector Corners[2] = {
+	const FVector Corners[2] = {
 		Transform.TransformPosition(Box.Min),
 		Transform.TransformPosition(Box.Max)
 	};
 
-	static const UE::Geometry::FVector3i Offsets[12][2] =
+	static constexpr UE::Geometry::FVector3i Offsets[12][2] =
 	{
 		// Bottom
 		{{0, 0, 0}, {1, 0, 0}},
@@ -930,25 +220,26 @@ static void DrawBox(IToolsContextRenderAPI* RenderAPI, const FTransform& Transfo
 
 void USkinWeightsBindingTool::Render(IToolsContextRenderAPI* RenderAPI)
 {
-	/**/
+	using namespace UE::Geometry;
 
 	if (Occupancy && Properties->bDebugDraw)
 	{
-		bool bShowInterior = false;
-		bool bShowBoundary = true;
+		constexpr bool bShowInterior = false;
+		constexpr bool bShowBoundary = true;
 
 		FPrimitiveDrawInterface* PDI = RenderAPI->GetPrimitiveDrawInterface();
 		const FTransform Transform = (FTransform) UE::ToolTarget::GetLocalToWorldTransform(Targets[0]);
 		float PDIScale = RenderAPI->GetCameraState().GetPDIScalingFactor();
+		const TDenseGrid3<FOccupancyGrid3::EDomain>& OccupancyGrid = Occupancy->GetOccupancyStateGrid(); 
 
-		for (int32 I = 0; I < Occupancy->Occupancy.Size(); I++)
+		for (int32 Index = 0; Index < OccupancyGrid.Size(); Index++)
 		{
-			UE::Geometry::FVector3i G(Occupancy->Occupancy.ToIndex(I));
+			FVector3i OccupancyIndex(OccupancyGrid.ToIndex(Index));
 
-			FOccupancyGrid::EDomain Domain = static_cast<FOccupancyGrid::EDomain>(Occupancy->Occupancy[G]);
-			if (bShowBoundary && Domain == FOccupancyGrid::Boundary)
+			const FOccupancyGrid3::EDomain Domain = OccupancyGrid[OccupancyIndex];
+			if (bShowBoundary && Domain == FOccupancyGrid3::EDomain::Boundary)
 			{
-				FBox Box = Occupancy->GetCellBoxFromIndex(G);
+				FBox Box{Occupancy->GetCellBoxFromIndex(OccupancyIndex)};
 				DrawBox(RenderAPI, Transform, Box, FLinearColor(1.0, 1.0, 0.0, 0.5), 0.5f);
 			}
 			
@@ -973,9 +264,9 @@ void USkinWeightsBindingTool::OnPropertyModified(UObject* PropertySet, FProperty
 		{
 			// Handled by the property watcher.
 		}
-		else
+		else if (Property->GetFName() == GET_MEMBER_NAME_CHECKED(USkinWeightsBindingToolProperties, VoxelResolution))
 		{
-			Occupancy = MakeShared<FOccupancyGrid>(*OriginalMesh, Properties->VoxelResolution);
+			Occupancy = MakeShared<UE::Geometry::FOccupancyGrid3>(*OriginalMesh, Properties->VoxelResolution);
 			
 			Preview->InvalidateResult();
 		}
@@ -985,16 +276,16 @@ void USkinWeightsBindingTool::OnPropertyModified(UObject* PropertySet, FProperty
 
 TUniquePtr<UE::Geometry::FDynamicMeshOperator> USkinWeightsBindingTool::MakeNewOperator()
 {
-	TUniquePtr<FComputeSkinWeightsBindingOp> Op = MakeUnique<FComputeSkinWeightsBindingOp>();
+	TUniquePtr<UE::Geometry::FSkinBindingOp> Op = MakeUnique<UE::Geometry::FSkinBindingOp>();
 
 	Op->ProfileName = FSkeletalMeshAttributes::DefaultSkinWeightProfileName;
-	Op->BindType = Properties->BindingType;
+	Op->BindType = static_cast<UE::Geometry::ESkinBindingType>(Properties->BindingType);
 	Op->Stiffness = Properties->Stiffness;
 	Op->MaxInfluences = Properties->MaxInfluences;
 	Op->VoxelResolution = Properties->VoxelResolution;
 	
 	Op->OriginalMesh = OriginalMesh;
-	Op->TransformHierarchy = TransformHierarchy;
+	Op->SetTransformHierarchyFromReferenceSkeleton(ReferenceSkeleton);
 
 	const FTransform3d LocalToWorld = UE::ToolTarget::GetLocalToWorldTransform(Targets[0]);
 	Op->SetResultTransform(LocalToWorld);
@@ -1029,15 +320,16 @@ void USkinWeightsBindingTool::UpdateVisualization(bool bInForce)
 {
 	using namespace UE::AnimationCore;
 	using namespace UE::Geometry;
-	
-	if ((bInForce || Preview->HaveValidNonEmptyResult()) && BoneToIndex.Contains(Properties->CurrentBone.BoneName))
+
+	const int32 RawBoneIndex = ReferenceSkeleton.FindRawBoneIndex(Properties->CurrentBone.BoneName);
+	if ((bInForce || Preview->HaveValidNonEmptyResult()) && RawBoneIndex != INDEX_NONE)
 	{
-		const FBoneIndexType BoneIndex = BoneToIndex[Properties->CurrentBone.BoneName];
+		const FBoneIndexType BoneIndex = static_cast<FBoneIndexType>(RawBoneIndex);
 
 		// update mesh with new value colors
 		Preview->PreviewMesh->EditMesh([&](FDynamicMesh3& InMesh)
 		{
-			FDynamicMeshVertexSkinWeightsAttribute *SkinWeights = InMesh.Attributes()->GetSkinWeightsAttribute(FSkeletalMeshAttributes::DefaultSkinWeightProfileName);
+			const FDynamicMeshVertexSkinWeightsAttribute* SkinWeights = InMesh.Attributes()->GetSkinWeightsAttribute(FSkeletalMeshAttributes::DefaultSkinWeightProfileName);
 			FDynamicMeshColorOverlay* ColorOverlay = InMesh.Attributes()->PrimaryColors();
 
 			if (!ColorOverlay)
@@ -1046,22 +338,22 @@ void USkinWeightsBindingTool::UpdateVisualization(bool bInForce)
 				InMesh.Attributes()->EnablePrimaryColors();
 				// Create an overlay that has no split elements, init with zero value.
 				ColorOverlay = InMesh.Attributes()->PrimaryColors();
-				ColorOverlay->CreateFromPredicate([](int ParentVID, int TriIDA, int TriIDB){return true;}, 0.f);
+				ColorOverlay->CreateFromPredicate([](int /*ParentVID*/, int /*TriIDA*/, int /*TriIDB*/){return true;}, 0.f);
 			}
 			
 			FBoneWeights BoneWeights;
 			
-			for (int32 ElementId : ColorOverlay->ElementIndicesItr())
+			for (const int32 ElementId : ColorOverlay->ElementIndicesItr())
 			{
 				const int32 VertexId = ColorOverlay->GetParentVertex(ElementId);
 				SkinWeights->GetValue(VertexId, BoneWeights);
 
 				float Weight = 0.0f;
-				for (FBoneWeight BW: BoneWeights)
+				for (const FBoneWeight BoneWeight: BoneWeights)
 				{
-					if (BW.GetBoneIndex() == BoneIndex)
+					if (BoneWeight.GetBoneIndex() == BoneIndex)
 					{
-						Weight = BW.GetWeight();
+						Weight = BoneWeight.GetWeight();
 						break;
 					}
 				}
diff --git a/Engine/Plugins/Experimental/MeshModelingToolsetExp/Source/MeshModelingToolsExp/Public/SkinWeightsBindingTool.h b/Engine/Plugins/Experimental/MeshModelingToolsetExp/Source/MeshModelingToolsExp/Public/SkinWeightsBindingTool.h
--- a/Engine/Plugins/Experimental/MeshModelingToolsetExp/Source/MeshModelingToolsExp/Public/SkinWeightsBindingTool.h
+++ b/Engine/Plugins/Experimental/MeshModelingToolsetExp/Source/MeshModelingToolsExp/Public/SkinWeightsBindingTool.h
@@ -18,6 +18,11 @@ struct FMeshDescription;
 class IMeshDescriptionCommitter;
 struct FOccupancyGrid;
 
+namespace UE::Geometry
+{
+	struct FOccupancyGrid3;
+}
+
 /**
  *
  */
@@ -30,13 +35,12 @@ public:
 	virtual UMultiSelectionMeshEditingTool* CreateNewTool(const FToolBuilderState& SceneState) const override;
 };
 
+// A mirror of UE::Geometry::ESkinBindingType
 UENUM()
 enum class ESkinWeightsBindType : uint8
 {
 	DirectDistance = 0 UMETA(DisplayName = "Direct Distance"),
-	// GeodesicDistance UMETA(DisplayName = "Geodesic Distance"),
-	// Heatmaps UMETA(DisplayName = "Heatmap"),
-	GeodesicVoxel = 3 UMETA(DisplayName = "Geodesic Voxel"),
+	GeodesicVoxel = 1 UMETA(DisplayName = "Geodesic Voxel"),
 };
 
 UCLASS()
@@ -69,7 +73,7 @@ public:
 	bool bDebugDraw = false;
 	
 	// IBoneReferenceSkeletonProvider
-	USkeleton* GetSkeleton(bool& bInvalidSkeletonIsError, const IPropertyHandle* PropertyHandle) override;
+	virtual USkeleton* GetSkeleton(bool& bInvalidSkeletonIsError, const IPropertyHandle* PropertyHandle) override;
 
 	TObjectPtr<USkeletalMesh> SkeletalMesh;
 };
@@ -86,22 +90,22 @@ class MESHMODELINGTOOLSEXP_API USkinWeightsBindingTool :
 
 public:
 	USkinWeightsBindingTool();
-	~USkinWeightsBindingTool();
+	virtual ~USkinWeightsBindingTool() override;
 	
-	void Setup() override;
-	void OnShutdown(EToolShutdownType ShutdownType) override;
+	virtual void Setup() override;
+	virtual void OnShutdown(EToolShutdownType ShutdownType) override;
 	
-	void OnTick(float DeltaTime) override;
-	void Render(IToolsContextRenderAPI* RenderAPI) override;
+	virtual void OnTick(float DeltaTime) override;
+	virtual void Render(IToolsContextRenderAPI* RenderAPI) override;
 
-	bool HasCancel() const override { return true; }
-	bool HasAccept() const override { return true; }
-	bool CanAccept() const override;
+	virtual bool HasCancel() const override { return true; }
+	virtual bool HasAccept() const override { return true; }
+	virtual bool CanAccept() const override;
 
-	void OnPropertyModified(UObject* PropertySet, FProperty* Property) override;
+	virtual void OnPropertyModified(UObject* PropertySet, FProperty* Property) override;
 
 	// IDynamicMeshOperatorFactory API
-	TUniquePtr<UE::Geometry::FDynamicMeshOperator> MakeNewOperator() override;
+	virtual TUniquePtr<UE::Geometry::FDynamicMeshOperator> MakeNewOperator() override;
 
 	UPROPERTY()
 	TObjectPtr<USkinWeightsBindingToolProperties> Properties = nullptr;
@@ -110,14 +114,13 @@ public:
 	TObjectPtr<UMeshOpPreviewWithBackgroundCompute> Preview = nullptr;
 
 protected:
-	TSharedPtr<FOccupancyGrid> Occupancy;
+	TSharedPtr<UE::Geometry::FOccupancyGrid3> Occupancy;
 	
 	TSharedPtr<UE::Geometry::FDynamicMesh3, ESPMode::ThreadSafe> OriginalMesh;
-	
-	FBoneContainer BoneContainer;
-	TMap<FName, FBoneIndexType> BoneToIndex;
-	TArray<TPair<FTransform, int32>> TransformHierarchy;
 
+	FReferenceSkeleton ReferenceSkeleton;
+	FBoneContainer BoneContainer;
+	
 	void GenerateAsset(const FDynamicMeshOpResult& Result);
 
 	static FVector4f WeightToColor(float Value);
diff --git a/Engine/Plugins/Experimental/MeshModelingToolsetExp/Source/MeshModelingToolsExp/Public/SkinWeightsPaintTool.h b/Engine/Plugins/Experimental/MeshModelingToolsetExp/Source/MeshModelingToolsExp/Public/SkinWeightsPaintTool.h
--- a/Engine/Plugins/Experimental/MeshModelingToolsetExp/Source/MeshModelingToolsExp/Public/SkinWeightsPaintTool.h
+++ b/Engine/Plugins/Experimental/MeshModelingToolsetExp/Source/MeshModelingToolsExp/Public/SkinWeightsPaintTool.h
@@ -11,14 +11,12 @@
 #include "Interfaces/Interface_BoneReferenceSkeletonProvider.h"
 #include "Misc/Optional.h"
 #include "Containers/Map.h"
-#include "TargetInterfaces/MeshDescriptionCommitter.h"
 
 #include "SkinWeightsPaintTool.generated.h"
 
 
 struct FMeshDescription;
 class USkinWeightsPaintTool;
-class IMeshDescriptionCommitter;
 
 
 
diff --git a/Engine/Plugins/Runtime/MeshModelingToolset/Source/ModelingOperators/ModelingOperators.Build.cs b/Engine/Plugins/Runtime/MeshModelingToolset/Source/ModelingOperators/ModelingOperators.Build.cs
--- a/Engine/Plugins/Runtime/MeshModelingToolset/Source/ModelingOperators/ModelingOperators.Build.cs
+++ b/Engine/Plugins/Runtime/MeshModelingToolset/Source/ModelingOperators/ModelingOperators.Build.cs
@@ -15,6 +15,7 @@ public class ModelingOperators : ModuleRules
 				"DynamicMesh",
 				"MeshConversion",
 				"GeometryAlgorithms", // required for constrained Delaunay triangulation
+				"SkeletalMeshDescription"
 				// ... add other public dependencies that you statically link with here ...
 			}
 		);
diff --git a/Engine/Plugins/Runtime/MeshModelingToolset/Source/ModelingOperators/Private/SkinningOps/SkinBindingOp.cpp b/Engine/Plugins/Runtime/MeshModelingToolset/Source/ModelingOperators/Private/SkinningOps/SkinBindingOp.cpp
new file mode 100644
--- /dev/null
+++ b/Engine/Plugins/Runtime/MeshModelingToolset/Source/ModelingOperators/Private/SkinningOps/SkinBindingOp.cpp
@@ -0,0 +1,595 @@
+﻿// Copyright Epic Games, Inc. All Rights Reserved.
+
+#include "SkinningOps/SkinBindingOp.h"
+
+#include "BoneWeights.h"
+#include "DynamicMesh/DynamicMesh3.h"
+#include "DynamicMesh/DynamicMeshAttributeSet.h"
+#include "DynamicMesh/DynamicVertexSkinWeightsAttribute.h"
+#include "ReferenceSkeleton.h"
+#include "Spatial/DenseGrid3.h"
+#include "Spatial/MeshWindingNumberGrid.h"
+#include "Spatial/OccupancyGrid3.h"
+
+
+namespace UE::Geometry
+{
+
+// A simple FIFO queue. Maintains a set of blocks, rather than allocate for each element.
+template<typename T, int32 BlockSize=256>
+class TFIFOQueue
+{
+public:
+	TFIFOQueue() = default;
+
+	void Push(T InElem)
+	{
+		// Container Empty?
+		if (Blocks.IsEmpty())
+		{
+			Blocks.SetNum(1);
+			Blocks[0].bIsFree = false;
+			PushIndex = -1;
+		}
+		// Or are we at the end of current block?
+		else if (PushIndex == (BlockSize - 1))
+		{
+			PushBlock = AllocateNewPushBlock();
+			PushIndex = -1;
+		}
+
+		PushIndex++;
+
+		Blocks[PushBlock].Data[PushIndex] = InElem;
+	}
+
+	T Pop()
+	{
+		check(!IsEmpty());
+
+		if (PopIndex == (BlockSize - 1))
+		{
+			// Reached the end. Free this block and move onto the next one.
+			PopBlock = FreeCurrentPopBlock();
+			PopIndex = -1;
+		}
+
+		PopIndex++;
+		return Blocks[PopBlock].Data[PopIndex];
+	}
+
+	bool TryPop(T &OutValue)
+	{
+		if (IsEmpty())
+		{
+			return false;
+		}
+
+		OutValue = Pop();
+		return true;
+	}
+
+	void Reset()
+	{
+		Blocks.Reset();
+		PushBlock = 0;
+		PushIndex = -1;
+		PopBlock = 0;
+		PopIndex = -1;
+	}
+	
+
+	bool IsEmpty() const
+	{
+		return PushBlock == PopBlock && PushIndex == PopIndex;
+	}
+
+private:
+	int32 AllocateNewPushBlock()
+	{
+		checkSlow(PushBlock != -1);
+
+		int NewBlockIndex = INDEX_NONE;
+		
+		if (FreeCount > 0)
+		{
+			// Find a free block
+			for (int32 Index = 0; Index < Blocks.Num(); Index++)
+			{
+				if (Blocks[Index].bIsFree)
+				{
+					NewBlockIndex = Index;
+					FreeCount--;
+					break;
+				}
+			}
+			checkfSlow(NewBlockIndex != INDEX_NONE, TEXT("We should have found a free block."));
+		}
+		else
+		{
+			NewBlockIndex = Blocks.Num();
+			Blocks.AddDefaulted();
+		}
+
+		Blocks[PushBlock].NextBlock = NewBlockIndex;
+		Blocks[NewBlockIndex].bIsFree = false;
+		Blocks[NewBlockIndex].NextBlock = INDEX_NONE;
+		return NewBlockIndex;
+	}
+
+	int32 FreeCurrentPopBlock()
+	{
+		const int32 NextBlock = Blocks[PopBlock].NextBlock;
+		checkSlow(NextBlock != -1);
+		Blocks[PopBlock].bIsFree = true;
+		FreeCount++;
+
+		return NextBlock;
+	}
+	
+
+	int32 PushBlock = 0;
+	int32 PushIndex = -1;
+
+	int32 PopBlock = 0;
+	int32 PopIndex = -1;
+
+	int32 FreeCount = 0;
+	
+	struct FBlock
+	{
+		FBlock()
+		{
+			Data = new T[BlockSize];
+		}
+		~FBlock()
+		{
+			delete [] Data;
+		}
+
+		FBlock(FBlock &&InOther) noexcept
+		{
+			Data = InOther.Data;
+			InOther.Data = nullptr;
+		}
+		
+		bool bIsFree = true;
+		int32 NextBlock = INDEX_NONE;
+		T *Data;
+	};
+	
+	TArray<FBlock> Blocks;
+};
+
+
+static float DistanceToLineSegment(const FVector& P, const FVector& A, const FVector& B)
+{
+	const FVector M = B - A;
+	const FVector T = P - A;
+
+	const float C1 = FVector::DotProduct(M, T);
+	if (C1 <= 0.0f)
+	{
+		return FVector::Dist(P, A);
+	}
+
+	const float C2 = FVector::DotProduct(M, M);
+	if (C2 <= C1)
+	{
+		return FVector::Dist(P, B);
+	}
+
+	// Project the point onto the line and get the distance between them.
+	const FVector PT = A + M * (C1 / C2);
+	return FVector::Dist(P, PT);
+}
+
+
+// List of bones as used by the binding class. In this case for each bone transform, we want to
+// store a list of line segments going from the bone transform to all the child bone transforms.
+struct FTransformHierarchyQuery
+{
+	explicit FTransformHierarchyQuery(const TArray<TPair<FTransform, int32>>& InTransformHierarchy)
+	{
+		for (int Index = 0; Index < InTransformHierarchy.Num(); Index++)
+		{
+			FTransform Xform = InTransformHierarchy[Index].Key;
+			int32 ParentIndex = InTransformHierarchy[Index].Value;
+
+			while (ParentIndex != INDEX_NONE)
+			{
+				Xform = Xform * InTransformHierarchy[ParentIndex].Key;
+				ParentIndex = InTransformHierarchy[ParentIndex].Value;
+			}
+
+			BoneFans.Add({ Xform.GetLocation() });
+		}
+
+		// Fill in the fan tips, as needed.
+		for (int Index = 0; Index < InTransformHierarchy.Num(); Index++)
+		{
+			const int32 ParentIndex = InTransformHierarchy[Index].Value;
+			if (ParentIndex != INDEX_NONE)
+			{
+				BoneFans[ParentIndex].TipsPos.Add(BoneFans[Index].RootPos);
+			}
+		}		
+	}
+
+	float GetDistanceToBoneFan(const int32 InBoneIndex, const FVector& InPoint) const
+	{
+		return BoneFans[InBoneIndex].GetDistance(InPoint);
+	}
+	
+	FBox GetBoneFanBBox(const int32 InBoneIndex) const
+	{
+		return BoneFans[InBoneIndex].GetBBox();
+	}
+
+	bool GetBoneFanIntersectsBox(const int32 InBoneIndex, const FBox &InBox) const
+	{
+		return BoneFans[InBoneIndex].IntersectsBox(InBox);
+	}
+
+private:
+	struct FBoneFan
+	{
+		FVector RootPos;
+		TArray<FVector> TipsPos;
+
+		float GetDistance(const FVector& InPoint) const
+		{
+			if (TipsPos.IsEmpty())
+			{
+				return FVector::Distance(RootPos, InPoint);
+			}
+			else
+			{
+				float Distance = std::numeric_limits<float>::max();
+				for (const FVector& TipPos: TipsPos)
+				{
+					Distance = FMath::Min(Distance, DistanceToLineSegment(InPoint, RootPos, TipPos));
+				}
+				return Distance;
+			}
+		}
+
+		FBox GetBBox() const
+		{
+			FBox Box(RootPos, RootPos);
+			for (const FVector& TipPos: TipsPos)
+			{
+				Box += TipPos;
+			}
+
+			return Box;
+		}
+
+		bool IntersectsBox(const FBox &InBox) const
+		{
+			if (TipsPos.IsEmpty())
+			{
+				return FMath::PointBoxIntersection(RootPos, InBox);
+			}
+
+			if (GetBBox().Intersect(InBox))
+			{
+				for (const FVector& TipPos: TipsPos)
+				{
+					if (FMath::LineBoxIntersection(InBox, RootPos, TipPos, TipPos - RootPos))
+					{
+						return true;
+					}
+				}
+			}
+			return false;
+		}				
+	};
+
+	 TArray<FBoneFan> BoneFans;
+};
+
+
+namespace
+{
+	struct FCreateSkinWeights_Closest_WorkData final :
+		TThreadSingleton<FCreateSkinWeights_Closest_WorkData>
+	{
+		TArray<TPair<FBoneIndexType, float>> RawBoneWeights;
+		TArray<AnimationCore::FBoneWeight> BoneWeights;
+	};
+
+	float ComputeWeightStiffness(const float InWeight, const float InStiffness)
+	{
+		return (1.0f - InStiffness) * InWeight + InStiffness * InWeight * InWeight;
+	}
+	
+}
+
+
+void FSkinBindingOp::SetTransformHierarchyFromReferenceSkeleton(
+	const FReferenceSkeleton& InRefSkeleton
+	)
+{
+	// Only use non-virtual bones, since virtual bones cannot be bound to the skin.
+	const TArray<FMeshBoneInfo>& BoneInfo = InRefSkeleton.GetRawRefBoneInfo();
+	const TArray<FTransform>& BonePose = InRefSkeleton.GetRawRefBonePose();
+		
+	TransformHierarchy.Reset(BoneInfo.Num());
+		
+	for (int32 Index = 0; Index < BoneInfo.Num(); Index++)
+	{
+		TransformHierarchy.Add(MakeTuple(BonePose[Index], BoneInfo[Index].ParentIndex));
+	}
+}
+
+
+void FSkinBindingOp::CalculateResult(FProgressCancel* InProgress)
+{
+	if (InProgress && InProgress->Cancelled())
+	{
+		return;
+	}
+
+	ResultMesh->Copy(*OriginalMesh, true, true, true, true);
+		
+	if (InProgress && InProgress->Cancelled())
+	{
+		return;
+	}
+
+	const float ClampedStiffness = FMath::Clamp(Stiffness, 0.0f, 1.0f);	
+
+	UE::AnimationCore::FBoneWeightsSettings Settings;
+	Settings.SetMaxWeightCount(MaxInfluences);
+
+	switch(BindType)
+	{
+	case ESkinBindingType::DirectDistance:
+		CreateSkinWeights_DirectDistance(*ResultMesh, ClampedStiffness, Settings);
+		break;
+			
+	case ESkinBindingType::GeodesicVoxel:
+		CreateSkinWeights_GeodesicVoxel(*ResultMesh, ClampedStiffness, Settings);
+		break;
+	}
+}
+
+
+FDynamicMeshVertexSkinWeightsAttribute* FSkinBindingOp::GetOrCreateSkinWeightsAttribute(
+	FDynamicMesh3& InMesh,
+	FName InProfileName
+	)
+{
+	FDynamicMeshVertexSkinWeightsAttribute *Attribute = InMesh.Attributes()->GetSkinWeightsAttribute(InProfileName);
+	if (!Attribute)
+	{
+		Attribute = new FDynamicMeshVertexSkinWeightsAttribute(&InMesh);
+		InMesh.Attributes()->AttachSkinWeightsAttribute(InProfileName, Attribute);
+	}
+	
+	return Attribute;
+}
+
+
+void FSkinBindingOp::CreateSkinWeights_DirectDistance(
+	FDynamicMesh3& InMesh,
+	float InStiffness,
+	const AnimationCore::FBoneWeightsSettings& InSettings
+	)
+{
+	using namespace AnimationCore;
+
+	FDynamicMeshVertexSkinWeightsAttribute *SkinWeights = InMesh.Attributes()->GetSkinWeightsAttribute(ProfileName);
+		
+	const int32 NumVertices = InMesh.VertexCount();
+
+	// Use the diagonal size of the bbox to make the bone distance falloff scale invariant.
+	const float DiagBounds = InMesh.GetBounds(true).DiagonalLength();
+	
+	const FTransformHierarchyQuery Skeleton(TransformHierarchy);
+
+	ParallelFor(NumVertices, [&](const int32 VertexIdx)
+	{
+		const FVector3d& Pos = InMesh.GetVertex(VertexIdx);
+
+		FCreateSkinWeights_Closest_WorkData &WorkData = FCreateSkinWeights_Closest_WorkData::Get();
+		
+		WorkData.RawBoneWeights.Reset(TransformHierarchy.Num());
+
+		float TotalWeight = 0.0f;
+		for (int32 BoneIndex = 0; BoneIndex < TransformHierarchy.Num(); BoneIndex++)
+		{
+			// Normalize the distance by the diagonal size of the bbox to maintain scale invariance.
+			float Weight = Skeleton.GetDistanceToBoneFan(BoneIndex, Pos) / DiagBounds;
+
+			// Avoid div-by-zero but allow for the possibility that multiple bones may
+			// touch this vertex.
+			Weight = FMath::Max(Weight, KINDA_SMALL_NUMBER);
+
+			// Compute the actual weight, factoring in the stiffness value. W = (1/S(D))^2
+			// Where S(x) is the stiffness function.
+			Weight = FMath::Square(1.0f / ComputeWeightStiffness(Weight, InStiffness));
+			TotalWeight += Weight;
+			WorkData.RawBoneWeights.Add(MakeTuple(static_cast<FBoneIndexType>(BoneIndex), Weight));
+		}
+
+		// Normalize
+		for (TPair<FBoneIndexType, float> &BoneWeight: WorkData.RawBoneWeights)
+		{
+			BoneWeight.Value /= TotalWeight;
+		}
+			
+		WorkData.RawBoneWeights.Sort([](const TPair<FBoneIndexType, float> &A, const TPair<FBoneIndexType, float> &B)
+		{
+			return A.Value > B.Value;
+		});
+
+		WorkData.BoneWeights.Reset(InSettings.GetMaxWeightCount());
+		for (int32 BoneIndex = 0; BoneIndex < FMath::Min(InSettings.GetMaxWeightCount(), WorkData.RawBoneWeights.Num()); BoneIndex++)
+		{
+			const TPair<FBoneIndexType, float>& BoneWeight = WorkData.RawBoneWeights[BoneIndex];
+			WorkData.BoneWeights.Add(FBoneWeight(BoneWeight.Key, BoneWeight.Value));
+		}
+
+		SkinWeights->SetValue(VertexIdx, FBoneWeights::Create(WorkData.BoneWeights, InSettings));
+	});
+}
+
+
+void FSkinBindingOp::CreateSkinWeights_GeodesicVoxel(
+	FDynamicMesh3& InMesh,
+	float InStiffness,
+	const AnimationCore::FBoneWeightsSettings& InSettings
+	)
+{
+	using namespace AnimationCore;
+
+	FDynamicMeshVertexSkinWeightsAttribute *SkinWeights = InMesh.Attributes()->GetSkinWeightsAttribute(ProfileName);
+		
+	const int32 NumVertices = InMesh.VertexCount();
+
+	// Use the diagonal size of the bbox to make the bone distance falloff scale invariant.
+	const float DiagonalBounds = InMesh.GetBounds(true).DiagonalLength();
+		
+	const FTransformHierarchyQuery Skeleton(TransformHierarchy);
+
+	// This is grossly inefficient but tricky to do otherwise, since each bone distance
+	// computation is done per-thread. We could possibly solve this by chunking instead
+	// and accumulating partial results.
+	TArray<float> Weights;
+	Weights.SetNumUninitialized(NumVertices * TransformHierarchy.Num());
+
+	const FOccupancyGrid3 Occupancy(InMesh, VoxelResolution);
+
+	const FVector3i Dimensions = Occupancy.GetOccupancyStateGrid().GetDimensions();
+		
+	ParallelFor(TransformHierarchy.Num(), [&](int32 BoneIndex) {
+		TFIFOQueue<FVector3i> WorkingSet;  
+		FDenseGrid3f BoneDistance(Dimensions.X, Dimensions.Y, Dimensions.Z, DiagonalBounds);
+
+		// Mark all the cells that the bone intersects with distance of 0 and put them
+		// on the work queue.
+		const FBox BoneBox = Skeleton.GetBoneFanBBox(BoneIndex);
+		const FVector3i BoneMin = Occupancy.GetCellIndexFromPoint(BoneBox.Min);
+		const FVector3i BoneMax = Occupancy.GetCellIndexFromPoint(BoneBox.Max);
+
+		for (int32 I = BoneMin.X; I <= BoneMax.X; I++)
+		{
+			for (int32 J = BoneMin.Y; J <= BoneMax.Y; J++)
+			{
+				for (int32 K = BoneMin.Z; K <= BoneMax.Z; K++)
+				{
+					const FVector3i Candidate(I, J, K);
+					const FBox3d CellBox{Occupancy.GetCellBoxFromIndex(Candidate)};
+					if (Skeleton.GetBoneFanIntersectsBox(BoneIndex, CellBox))
+					{
+						WorkingSet.Push(Candidate);
+						BoneDistance[Candidate] = 0.0f;
+					}
+				}
+			}
+		}
+
+		// Iterate over all the voxels until we have constructed shortest distance paths
+		// throughout the level set.
+		while (!WorkingSet.IsEmpty())
+		{
+			const FVector3i WorkItem = WorkingSet.Pop();
+
+			// Loop through each of the neighbours (6 face neighbours, 12 edge neighbors,
+			// and 8 corner neighbours) and see if any of them are closer to the bone
+			// than their current marked distance.
+			float CurrentDistance = BoneDistance[WorkItem];
+			for (int32 N = 0; N < 26; N++)
+			{
+				FVector3i Offset(IndexUtil::GridOffsets26[N]);
+				FVector3i Candidate(WorkItem + Offset);
+
+				if (!BoneDistance.IsValidIndex(Candidate))
+				{
+					continue;
+				}
+					
+				// Ensure this entry is either a part of the interior or boundary domain.
+				if (Occupancy.GetOccupancyStateGrid()[Candidate] == FOccupancyGrid3::EDomain::Exterior)
+				{
+					continue;
+				}
+
+				const float CellDistance = (FVector3f(Offset) * Occupancy.GetCellSize()).Length();
+				const float CandidateDistance = CurrentDistance + CellDistance;
+				const float OldDistance = BoneDistance[Candidate]; 
+
+				if (OldDistance > CandidateDistance)
+				{
+					WorkingSet.Push(Candidate);
+					BoneDistance[Candidate] = CandidateDistance;
+				}
+			}
+		}
+
+		// Loop through all the vertices, find the voxel each belongs to, and compute
+		// the distance from the voxel to the vertex (assuming the distance stored in the
+		// voxel is based on traversing from voxel center to voxel center).
+		for (int32 VertexIdx = 0; VertexIdx < NumVertices; VertexIdx++)
+		{
+			const FVector3d& Pos = InMesh.GetVertex(VertexIdx);
+			const FVector3i CellIndex = Occupancy.GetCellIndexFromPoint(Pos);
+			const FVector3d CellCenter{Occupancy.GetCellCenterFromIndex(CellIndex)};
+
+			float Distance = BoneDistance[CellIndex];
+			Distance += FVector3d::Distance(CellCenter, Pos);
+				
+			// Normalize the distance by the diagonal size of the bbox to maintain scale invariance.
+			float Weight = Distance / DiagonalBounds;
+
+			// Avoid div-by-zero but allow for the possibility that multiple bones may
+			// touch this vertex.
+			Weight = FMath::Max(Weight, UE_KINDA_SMALL_NUMBER);
+
+			// Compute the actual weight, factoring in the stiffness value. W = (1/S(D))^2
+			// Where S(x) is the stiffness function.
+			Weight = FMath::Square(1.0f / ComputeWeightStiffness(Weight, InStiffness));
+
+			Weights[VertexIdx * TransformHierarchy.Num() + BoneIndex] = Weight;
+		}
+	});	
+
+	ParallelFor(NumVertices, [&](const int32 VertexIdx)
+	{
+		FCreateSkinWeights_Closest_WorkData &WorkData = FCreateSkinWeights_Closest_WorkData::Get();
+			
+		WorkData.RawBoneWeights.Reset(TransformHierarchy.Num());
+
+		float TotalWeight = 0.0f;
+		for (int32 BoneIndex = 0; BoneIndex < TransformHierarchy.Num(); BoneIndex++)
+		{
+			const float Weight = Weights[VertexIdx * TransformHierarchy.Num() + BoneIndex];
+			TotalWeight += Weight;
+			WorkData.RawBoneWeights.Add(MakeTuple(static_cast<FBoneIndexType>(BoneIndex), Weight));
+		}
+
+		// Normalize
+		for (TPair<FBoneIndexType, float> &BoneWeight: WorkData.RawBoneWeights)
+		{
+			BoneWeight.Value /= TotalWeight;
+		}
+				
+		WorkData.RawBoneWeights.Sort([](const TPair<FBoneIndexType, float> &A, const TPair<FBoneIndexType, float> &B)
+		{
+			return A.Value > B.Value;
+		});
+
+		WorkData.BoneWeights.Reset(InSettings.GetMaxWeightCount());
+		for (int32 BoneIndex = 0; BoneIndex < FMath::Min(InSettings.GetMaxWeightCount(), WorkData.RawBoneWeights.Num()); BoneIndex++)
+		{
+			const TPair<FBoneIndexType, float>& BoneWeight = WorkData.RawBoneWeights[BoneIndex];
+			WorkData.BoneWeights.Add(FBoneWeight(BoneWeight.Key, BoneWeight.Value));
+		}
+
+		SkinWeights->SetValue(VertexIdx, FBoneWeights::Create(WorkData.BoneWeights, InSettings));
+	});
+}
+
+} // namespace UE::Geometry
diff --git a/Engine/Plugins/Runtime/MeshModelingToolset/Source/ModelingOperators/Public/SkinningOps/SkinBindingOp.h b/Engine/Plugins/Runtime/MeshModelingToolset/Source/ModelingOperators/Public/SkinningOps/SkinBindingOp.h
new file mode 100644
--- /dev/null
+++ b/Engine/Plugins/Runtime/MeshModelingToolset/Source/ModelingOperators/Public/SkinningOps/SkinBindingOp.h
@@ -0,0 +1,77 @@
+﻿// Copyright Epic Games, Inc. All Rights Reserved.
+
+#pragma once
+
+#include "ModelingOperators.h"
+#include "SkeletalMeshAttributes.h"
+
+
+struct FReferenceSkeleton;
+
+
+namespace UE::Geometry
+{
+
+template<typename ParentType> class TDynamicVertexSkinWeightsAttribute;
+class FDynamicMesh3;
+using FDynamicMeshVertexSkinWeightsAttribute = TDynamicVertexSkinWeightsAttribute<FDynamicMesh3>;
+
+
+enum class ESkinBindingType : uint8
+{
+	// Computes the binding strength by computing the Euclidean distance to the closest set of bones,
+	// where the strength of binding is proportional to the inverse distance. May cause bones to affects
+	// parts of geometry that, although close in space, may be topologically distant.
+	DirectDistance = 0,
+
+	// Computes the binding by computing the geodesic distance from each set of bones. This is slower than the
+	// direct distance.
+	GeodesicVoxel = 1,
+};
+
+
+class MODELINGOPERATORS_API FSkinBindingOp : public FDynamicMeshOperator
+{
+public:
+	virtual ~FSkinBindingOp() override {}
+	
+	// The transform hierarchy to bind to. Listed in the same order as the bones in the
+	// reference skeleton that this skeletal mesh is tied to.
+	TSharedPtr<FDynamicMesh3, ESPMode::ThreadSafe> OriginalMesh;
+	TArray<TPair<FTransform, int32>> TransformHierarchy;
+
+	FName ProfileName = FSkeletalMeshAttributes::DefaultSkinWeightProfileName;
+	
+	ESkinBindingType BindType = ESkinBindingType::DirectDistance;
+	float Stiffness = 0.2f;
+	int32 MaxInfluences = 5;
+	int32 VoxelResolution = 256;
+
+	void SetTransformHierarchyFromReferenceSkeleton(const FReferenceSkeleton& InRefSkeleton);
+	
+	virtual void CalculateResult(FProgressCancel* InProgress) override;
+
+private:
+	static FDynamicMeshVertexSkinWeightsAttribute* GetOrCreateSkinWeightsAttribute(
+		FDynamicMesh3& InMesh,
+		FName InProfileName
+		);
+
+
+	void CreateSkinWeights_DirectDistance(
+		FDynamicMesh3& InMesh,
+		float InStiffness,
+		const AnimationCore::FBoneWeightsSettings& InSettings
+		);
+
+
+	void CreateSkinWeights_GeodesicVoxel(
+		FDynamicMesh3& InMesh,
+		float InStiffness,
+		const AnimationCore::FBoneWeightsSettings& InSettings
+		);
+};
+
+
+} // namespace UE::Geometry
+
diff --git a/Engine/Source/Runtime/GeometryCore/Public/Spatial/OccupancyGrid3.h b/Engine/Source/Runtime/GeometryCore/Public/Spatial/OccupancyGrid3.h
new file mode 100644
--- /dev/null
+++ b/Engine/Source/Runtime/GeometryCore/Public/Spatial/OccupancyGrid3.h
@@ -0,0 +1,154 @@
+﻿// Copyright Epic Games, Inc. All Rights Reserved.
+
+#pragma once
+
+#include "DenseGrid3.h"
+#include "FastWinding.h"
+#include "MeshAABBTree3.h"
+#include "MeshWindingNumberGrid.h"
+#include "Algo/Count.h"
+
+namespace UE::Geometry
+{
+
+/**
+ * Sample mesh occupancy in a voxel grid by using the mesh's winding number, on voxel corners, to determine whether the
+ * voxel is fully contained within the interior of the mesh, touches the mesh's boundary, or is fully contained on the 
+ * exterior of the mesh.
+ */
+struct FOccupancyGrid3
+{
+	/// An enum representing a voxel's classification.
+	enum class EDomain : int32
+	{
+		Exterior,		/// The point in the occupancy grid is completely on the exterior of the given mesh. 
+		Boundary,		/// The point in the occupancy grid includes a mesh boundary.
+		Interior		/// The point in the occupancy grid is completely on the interior of the given mesh.
+	};
+
+	template<typename MeshType>
+	FOccupancyGrid3(
+		const MeshType& InMesh,
+		const int32 InVoxelResolution
+		)
+	{
+		// Compute a voxel grid 
+		TMeshAABBTree3<MeshType> Spatial(&InMesh);
+		TFastWindingTree FastWinding(&Spatial);
+		const FAxisAlignedBox3d Bounds = Spatial.GetBoundingBox();
+		CellSize = Bounds.MaxDim() / InVoxelResolution;
+		CellMidPoint = {CellSize / 2.0f, CellSize / 2.0f, CellSize / 2.0f};
+
+		TMeshWindingNumberGrid WindingGrid(&InMesh, &FastWinding, CellSize);
+	
+		WindingGrid.Compute();
+
+		// Our occupancy grid is computed on the winding number grid's cell centers.
+		const FVector3i WindingDims = WindingGrid.Dimensions();
+		Occupancy = {WindingDims.X - 1, WindingDims.Y - 1, WindingDims.Z - 1, EDomain::Exterior};
+
+		GridOrigin = WindingGrid.GridOrigin + CellMidPoint;
+
+		static constexpr FVector3i CornerOffsets[] = {
+			FVector3i(0, 0, 0),
+			FVector3i(0, 0, 1),
+			FVector3i(0, 1, 0),
+			FVector3i(0, 1, 1),
+			FVector3i(1, 0, 0),
+			FVector3i(1, 0, 1),
+			FVector3i(1, 1, 0),
+			FVector3i(1, 1, 1),
+		};
+		
+		for (int32 OccupancyId = 0; OccupancyId < Occupancy.Size(); OccupancyId++)
+		{
+			const FVector3i OccupancyIndex(Occupancy.ToIndex(OccupancyId));
+			const int32 Count = Algo::CountIf(CornerOffsets, [&WindingGrid, OccupancyIndex](FVector3i CornerOffset)
+			{
+				const FVector3i CornerIndex(OccupancyIndex + CornerOffset);
+				return WindingGrid.GetValue(CornerIndex) >= WindingGrid.WindingIsoValue;
+			});
+
+			if (Count == 8)
+			{
+				Occupancy[OccupancyIndex] = EDomain::Interior;
+			}
+			else if (Count > 0)
+			{
+				Occupancy[OccupancyIndex] = EDomain::Boundary;
+			}
+		}
+
+		// Make sure we include all the vertices of the mesh as a part of the boundary, if
+		// the vertex areas are marked as being exterior.
+		for (int32 VertexIdx = 0; VertexIdx < InMesh.VertexCount(); VertexIdx++)
+		{
+			const FVector3d& Pos = InMesh.GetVertex(VertexIdx);
+			const FVector3i OccupancyIndex = GetCellIndexFromPoint(FVector(Pos));
+			if (Occupancy[OccupancyIndex] == EDomain::Exterior)
+			{
+				Occupancy[OccupancyIndex] = EDomain::Boundary;
+			}
+		}
+	}
+
+	/// Given a point in space, return a computed index into the occupancy grid. If the point is outside of the mesh's
+	/// bounding box, the given index will be outside the range of the occupancy grid. To check for validity, test
+	/// against the occupancy grid's extents as returned by GetOccupancy.
+	FVector3i GetCellIndexFromPoint(const FVector &InPoint) const
+	{
+		FVector3f PP(InPoint);
+		PP -= GridOrigin;
+		PP += CellMidPoint;
+
+		return { FMath::FloorToInt(PP.X / CellSize),
+				FMath::FloorToInt(PP.Y / CellSize),
+				FMath::FloorToInt(PP.Z / CellSize) };
+	}
+
+	/// Given an index into the occupancy grid, returns the midpoint of the bbox that represents the matching cell 
+	/// in the winding number grid, from which the occupancy got computed.
+	FVector3f GetCellCenterFromIndex(const FVector3i &Index) const
+	{
+		return {Index.X * CellSize + GridOrigin.X, Index.Y * CellSize + GridOrigin.Y, Index.Z * CellSize + GridOrigin.Z};
+	}
+
+	/// Given an index into the occupancy grid, returns the bbox that represents the cell in the winding number
+	/// grid, from which the occupancy got computed.
+	FBox3f GetCellBoxFromIndex(const FVector3i &Index) const
+	{
+		const FVector3f P = GetCellCenterFromIndex(Index);
+		return {P - CellMidPoint, P + CellMidPoint};
+	}
+
+	/// Returns a reference to the occupancy grid.  
+	const TDenseGrid3<EDomain>& GetOccupancyStateGrid() const
+	{
+		return Occupancy;
+	}
+
+	// Returns the size of the cell. The cell size is uniform on each axis. 
+	float GetCellSize() const
+	{
+		return CellSize;
+	}
+	
+private:
+	// The dense grid that contains the occupancy state of the mesh within each voxel. 
+	TDenseGrid3<EDomain> Occupancy;
+
+	// The size of each cell, applies to all dimensions.
+	float CellSize;
+
+	// The origin of the occupancy grid. Since the occupancy of a cell is defined at the winding number grid's cell
+	// centers, this is offset by half the cell size, compared to the winding number grid's origin (same as the
+	// mesh's origin, which defined by the lowest coordinate value).
+	FVector3f GridOrigin;
+
+	// This value is the offset from the lowest coordinate of a winding number grid's cell's box, to its mid-point.    
+	FVector3f CellMidPoint;
+};
+
+
+
+}
