Mesh Modeling and Dynamic Mesh.  Provide ability to store and read a vertex mapping  with a DynamicMesh inthe case that the source data is non-manifold.

ModelingTools - MeshAttributePaintTool and SkinWeightsPaintTool.  Use non-maniforld vertex ids when provided.  fixes a crash
SplitMeshesTool - removed an unsed mesh conversion.
ModelingComponents - ModelingToolTargetUtil  - request non-manifold vertex Ids on dynamic meshes when converted for tools.
DynamicMeshViaMeshDescription - request non-manifold vertex Ids on dynamic meshes when converted for tools.
ModelingObjectCreation - make sure to strip the any src vertex IDs from a dynamic mesh during object creation.

MeshConversions - MeshDescriptionToDynamicMesh - Add ability to capture non-manifold source veretx IDs in the form of a vertex attribute.

diff --git a/Engine/Plugins/Experimental/MeshModelingToolsetExp/Source/MeshModelingToolsExp/Private/MeshAttributePaintTool.cpp b/Engine/Plugins/Experimental/MeshModelingToolsetExp/Source/MeshModelingToolsExp/Private/MeshAttributePaintTool.cpp
--- a/Engine/Plugins/Experimental/MeshModelingToolsetExp/Source/MeshModelingToolsExp/Private/MeshAttributePaintTool.cpp
+++ b/Engine/Plugins/Experimental/MeshModelingToolsetExp/Source/MeshModelingToolsExp/Private/MeshAttributePaintTool.cpp
@@ -11,6 +11,7 @@
 
 #include "TargetInterfaces/PrimitiveComponentBackedTarget.h"
 #include "ModelingToolTargetUtil.h"
+#include "DynamicMesh/NonManifoldMappingSupport.h"
 
 #include UE_INLINE_GENERATED_CPP_BY_NAME(MeshAttributePaintTool)
 
@@ -453,10 +454,12 @@ void UMeshAttributePaintTool::UpdateVisibleAttribute()
 		PreviewMesh->EditMesh([&](FDynamicMesh3& Mesh)
 		{
 			FDynamicMeshColorOverlay* ColorOverlay = Mesh.Attributes()->PrimaryColors();
+			FNonManifoldMappingSupport NonManifoldMappingSupport(Mesh);
 			for (int32 elid : ColorOverlay->ElementIndicesItr())
 			{
 				const int32 vid = ColorOverlay->GetParentVertex(elid);
-				const float Value = AttribData.CurrentValues[vid];
+				const int32 srcvid = NonManifoldMappingSupport.GetOriginalNonManifoldVertexID(vid);
+				const float Value = AttribData.CurrentValues[srcvid];
 				const FVector4f Color4f = ToVector4<float>(ColorMapper->ToColor(Value));
 				ColorOverlay->SetElement(elid,  Color4f);
 			}
@@ -525,11 +528,13 @@ void UMeshAttributePaintTool::ApplyStamp(const FBrushStampData& Stamp)
 	{
 		TArray<int> ElIDs;
 		FDynamicMeshColorOverlay* ColorOverlay = Mesh.Attributes()->PrimaryColors();
+		FNonManifoldMappingSupport NonManifoldMappingSupport(Mesh);
 		int32 NumVertices = ActionData.ROIVertices.Num();
 		for (int32 k = 0; k < NumVertices; ++k)
 		{
 			int32 vid = ActionData.ROIVertices[k];
-			AttribData.CurrentValues[vid] = ActionData.ROIAfter[k];
+			int32 srcvid = NonManifoldMappingSupport.GetOriginalNonManifoldVertexID(vid);
+			AttribData.CurrentValues[srcvid] = ActionData.ROIAfter[k];
 			FVector4f NewColor( ToVector4<float>(ColorMapper->ToColor(ActionData.ROIAfter[k])) );
 			ColorOverlay->GetVertexElements(vid, ElIDs);
 			for (int elid : ElIDs)
@@ -555,6 +560,7 @@ void UMeshAttributePaintTool::ApplyStamp_Paint(const FBrushStampData& Stamp, FSt
 	FAttributeData& AttribData = Attributes[CurrentAttributeIndex];
 
 	const FDynamicMesh3* CurrentMesh = PreviewMesh->GetMesh();
+	FNonManifoldMappingSupport NonManifoldMappingSupport(*CurrentMesh);
 	if (bInSmoothStroke)
 	{
 		float SmoothSpeed = 0.25f;
@@ -566,17 +572,19 @@ void UMeshAttributePaintTool::ApplyStamp_Paint(const FBrushStampData& Stamp, FSt
 			float ValueSum = 0, WeightSum = 0;
 			for (int32 NbrVID : CurrentMesh->VtxVerticesItr(vid))
 			{
+				int32 srcNbrVID =  NonManifoldMappingSupport.GetOriginalNonManifoldVertexID(NbrVID);
 				FVector3d NbrPos = CurrentMesh->GetVertex(NbrVID);
 				float Weight = FMathf::Clamp(1.0f / DistanceSquared(NbrPos, Position), 0.0001f, 1000.0f);
-				ValueSum += Weight * AttribData.CurrentValues[NbrVID];
+				ValueSum += Weight * AttribData.CurrentValues[srcNbrVID];
 				WeightSum += Weight;
 			}
 			ValueSum /= WeightSum;
 
 			float Falloff = (float)CalculateBrushFalloff(Distance(Position, StampPosLocal));
-			float NewValue = FMathf::Lerp(AttribData.CurrentValues[vid], ValueSum, SmoothSpeed*Falloff);
+			const int32 srcvid =  NonManifoldMappingSupport.GetOriginalNonManifoldVertexID(vid);
+			float NewValue = FMathf::Lerp(AttribData.CurrentValues[srcvid], ValueSum, SmoothSpeed*Falloff);
 
-			ActionData.ROIBefore[k] = AttribData.CurrentValues[vid];
+			ActionData.ROIBefore[k] = AttribData.CurrentValues[srcvid];
 			ActionData.ROIAfter[k] = CurrentValueRange.Clamp(NewValue);
 		}
 	}
@@ -588,9 +596,10 @@ void UMeshAttributePaintTool::ApplyStamp_Paint(const FBrushStampData& Stamp, FSt
 		for (int32 k = 0; k < NumVertices; ++k)
 		{
 			int32 vid = ActionData.ROIVertices[k];
+			int32 srcvid =  NonManifoldMappingSupport.GetOriginalNonManifoldVertexID(vid);
 			FVector3d Position = CurrentMesh->GetVertex(vid);
 			float Falloff = (float)CalculateBrushFalloff(Distance(Position, StampPosLocal));
-			ActionData.ROIBefore[k] = AttribData.CurrentValues[vid];
+			ActionData.ROIBefore[k] = AttribData.CurrentValues[srcvid];
 			ActionData.ROIAfter[k] = CurrentValueRange.Clamp(ActionData.ROIBefore[k] + UseStrength*Falloff);
 		}
 	}
@@ -615,6 +624,7 @@ void UMeshAttributePaintTool::ApplyStamp_FloodFill(const FBrushStampData& Stamp,
 		SetValue = CurrentValueRange.Min;
 	}
 
+	FNonManifoldMappingSupport NonManifoldMappingSupport(*CurrentMesh);
 	ActionData.ROIVertices.Reset();
 	TArray<int32> InputTriROI, OutputTriROI, QueueTempBuffer;
 	TSet<int32> DoneTempBuffer, DoneVertices;
@@ -638,10 +648,12 @@ void UMeshAttributePaintTool::ApplyStamp_FloodFill(const FBrushStampData& Stamp,
 			{
 				if (DoneVertices.Contains(TriVertices[j]) == false)
 				{
-					ActionData.ROIVertices.Add(TriVertices[j]);
-					ActionData.ROIBefore.Add(AttribData.CurrentValues[TriVertices[j]]);
+					int32 vid = TriVertices[j];
+					ActionData.ROIVertices.Add(vid);
+					int32 srcvid = NonManifoldMappingSupport.GetOriginalNonManifoldVertexID(vid);
+					ActionData.ROIBefore.Add(AttribData.CurrentValues[srcvid]);
 					ActionData.ROIAfter.Add(SetValue);
-					DoneVertices.Add(TriVertices[j]);
+					DoneVertices.Add(vid);
 				}
 			}
 		}
diff --git a/Engine/Plugins/Experimental/MeshModelingToolsetExp/Source/MeshModelingToolsExp/Private/SkinWeightsPaintTool.cpp b/Engine/Plugins/Experimental/MeshModelingToolsetExp/Source/MeshModelingToolsExp/Private/SkinWeightsPaintTool.cpp
--- a/Engine/Plugins/Experimental/MeshModelingToolsetExp/Source/MeshModelingToolsExp/Private/SkinWeightsPaintTool.cpp
+++ b/Engine/Plugins/Experimental/MeshModelingToolsetExp/Source/MeshModelingToolsExp/Private/SkinWeightsPaintTool.cpp
@@ -14,6 +14,7 @@
 #include "ModelingToolTargetUtil.h"
 
 #include "MeshDescription.h"
+#include "DynamicMesh/NonManifoldMappingSupport.h"
 
 #include UE_INLINE_GENERATED_CPP_BY_NAME(SkinWeightsPaintTool)
 
@@ -318,11 +319,13 @@ void USkinWeightsPaintTool::UpdateBoneVisualization()
 	// update mesh with new value colors
 	PreviewMesh->EditMesh([&](FDynamicMesh3& Mesh)
 	{
+		UE::Geometry::FNonManifoldMappingSupport NonManifoldMappingSupport(Mesh);
 		UE::Geometry::FDynamicMeshColorOverlay* ColorOverlay = Mesh.Attributes()->PrimaryColors();
 		for (int32 ElementId : ColorOverlay->ElementIndicesItr())
 		{
-			const int32 VertexId = ColorOverlay->GetParentVertex(ElementId);
-			const float Value = SkinWeightsData.IsValidIndex(VertexId) ? SkinWeightsData[VertexId] : 0.f;
+			const int32 VertexId = ColorOverlay->GetParentVertex(ElementId);	
+			const int32 SrcVertexId = NonManifoldMappingSupport.GetOriginalNonManifoldVertexID(VertexId);
+               const float Value = SkinWeightsData.IsValidIndex(SrcVertexId) ? SkinWeightsData[SrcVertexId] : 0.f;
 			const FVector4f Color(WeightToColor(Value));
 			ColorOverlay->SetElement(ElementId, Color);
 		}
@@ -369,6 +372,7 @@ void USkinWeightsPaintTool::ApplyStamp(const FBrushStampData& Stamp)
 	TArray<float>& SkinWeightsData = *SkinWeightsMap.Find(CurrentBone);
 
 	const FDynamicMesh3* CurrentMesh = PreviewMesh->GetMesh();
+	UE::Geometry::FNonManifoldMappingSupport NonManifoldMappingSupport(*CurrentMesh);
 	if (bInSmoothStroke)
 	{
 		const float SmoothSpeed = 0.25f;
@@ -376,21 +380,23 @@ void USkinWeightsPaintTool::ApplyStamp(const FBrushStampData& Stamp)
 		for (int32 Index = 0; Index < NumROIVertices; ++Index)
 		{
 			const int32 VertexId = ROIVertices[Index];
+			const int32 SrcVertexId = NonManifoldMappingSupport.GetOriginalNonManifoldVertexID(VertexId);
 			FVector3d Position = CurrentMesh->GetVertex(VertexId);
 			float ValueSum = 0, WeightSum = 0;
 			for (int32 NeighborVertexId : CurrentMesh->VtxVerticesItr(VertexId))
 			{
 				FVector3d NbrPos = CurrentMesh->GetVertex(NeighborVertexId);
 				const float Weight = FMathf::Clamp(1.0f / FVector3d::DistSquared(NbrPos, Position), 0.0001f, 1000.0f);
-				ValueSum += Weight * SkinWeightsData[NeighborVertexId];
+				const int32 SrcNeighborVertexId = NonManifoldMappingSupport.GetOriginalNonManifoldVertexID(NeighborVertexId);
+				ValueSum += Weight * SkinWeightsData[SrcNeighborVertexId];
 				WeightSum += Weight;
 			}
 			ValueSum /= WeightSum;
 
 			const float Falloff = float(CalculateBrushFalloff(FVector3d::Dist(Position, StampPosLocal)));
-			const float NewValue = FMathf::Lerp(SkinWeightsData[VertexId], ValueSum, SmoothSpeed*Falloff);
+			const float NewValue = FMathf::Lerp(SkinWeightsData[SrcVertexId], ValueSum, SmoothSpeed*Falloff);
 
-			ROIBefore[Index] = SkinWeightsData[VertexId];
+			ROIBefore[Index] = SkinWeightsData[SrcVertexId];
 			ROIAfter[Index] = FMath::Clamp(NewValue, 0.0f, 1.0f);
 		}
 	}
@@ -402,9 +408,10 @@ void USkinWeightsPaintTool::ApplyStamp(const FBrushStampData& Stamp)
 		for (int32 Index = 0; Index < NumROIVertices; ++Index)
 		{
 			const int32 VertexId = ROIVertices[Index];
+			const int32 SrcVertexId = NonManifoldMappingSupport.GetOriginalNonManifoldVertexID(VertexId);
 			const FVector3d Position = CurrentMesh->GetVertex(VertexId);
 			const float Falloff = (float)CalculateBrushFalloff(FVector3d::Dist(Position, StampPosLocal));
-			ROIBefore[Index] = SkinWeightsData[VertexId];
+			ROIBefore[Index] = SkinWeightsData[SrcVertexId];
 			ROIAfter[Index] = FMath::Clamp(ROIBefore[Index] + UseStrength * Falloff, 0.0f, 1.0f);
 		}
 	}
@@ -423,7 +430,8 @@ void USkinWeightsPaintTool::ApplyStamp(const FBrushStampData& Stamp)
 		for (int32 Index = 0; Index < NumROIVertices; ++Index)
 		{
 			const int32 VertexId = ROIVertices[Index];
-			SkinWeightsData[VertexId] = ROIAfter[Index];
+			const int32 SrcVertexId = NonManifoldMappingSupport.GetOriginalNonManifoldVertexID(VertexId);
+			SkinWeightsData[SrcVertexId] = ROIAfter[Index];
 			FVector4f NewColor(WeightToColor(ROIAfter[Index]));
 			ColorOverlay->GetVertexElements(VertexId, ElementIds);
 			for (int ElementId : ElementIds)
diff --git a/Engine/Plugins/Experimental/MeshModelingToolsetExp/Source/MeshModelingToolsExp/Private/SplitMeshesTool.cpp b/Engine/Plugins/Experimental/MeshModelingToolsetExp/Source/MeshModelingToolsExp/Private/SplitMeshesTool.cpp
--- a/Engine/Plugins/Experimental/MeshModelingToolsetExp/Source/MeshModelingToolsExp/Private/SplitMeshesTool.cpp
+++ b/Engine/Plugins/Experimental/MeshModelingToolsetExp/Source/MeshModelingToolsExp/Private/SplitMeshesTool.cpp
@@ -106,7 +106,6 @@ void USplitMeshesTool::OnShutdown(EToolShutdownType ShutdownType)
 			DeleteActors.Add(TargetActor);
 
 			FTransform3d SourceTransform = UE::ToolTarget::GetLocalToWorldTransform(Targets[ti]);
-			FDynamicMesh3 SourceMesh = UE::ToolTarget::GetDynamicMeshCopy(Targets[ti], true);
 			FString AssetName = TargetActor->GetActorNameOrLabel();
 
 			FCreateMeshObjectParams BaseMeshObjectParams;
diff --git a/Engine/Plugins/Runtime/MeshModelingToolset/Source/ModelingComponents/Private/ConversionUtils/DynamicMeshViaMeshDescriptionUtil.cpp b/Engine/Plugins/Runtime/MeshModelingToolset/Source/ModelingComponents/Private/ConversionUtils/DynamicMeshViaMeshDescriptionUtil.cpp
--- a/Engine/Plugins/Runtime/MeshModelingToolset/Source/ModelingComponents/Private/ConversionUtils/DynamicMeshViaMeshDescriptionUtil.cpp
+++ b/Engine/Plugins/Runtime/MeshModelingToolset/Source/ModelingComponents/Private/ConversionUtils/DynamicMeshViaMeshDescriptionUtil.cpp
@@ -13,6 +13,7 @@ FDynamicMesh3 UE::Geometry::GetDynamicMeshViaMeshDescription(
 {
 	FDynamicMesh3 DynamicMesh;
 	FMeshDescriptionToDynamicMesh Converter;
+	Converter.bVIDsFromNonManifoldMeshDescriptionAttr = true;
 	Converter.Convert(MeshDescriptionProvider.GetMeshDescription(), DynamicMesh);
 	return DynamicMesh;
 }
diff --git a/Engine/Plugins/Runtime/MeshModelingToolset/Source/ModelingComponents/Private/ModelingObjectsCreationAPI.cpp b/Engine/Plugins/Runtime/MeshModelingToolset/Source/ModelingComponents/Private/ModelingObjectsCreationAPI.cpp
--- a/Engine/Plugins/Runtime/MeshModelingToolset/Source/ModelingComponents/Private/ModelingObjectsCreationAPI.cpp
+++ b/Engine/Plugins/Runtime/MeshModelingToolset/Source/ModelingComponents/Private/ModelingObjectsCreationAPI.cpp
@@ -7,6 +7,7 @@
 #include "Misc/Paths.h"
 #include "Components/StaticMeshComponent.h"
 #include "Engine/StaticMesh.h"
+#include "DynamicMesh/NonManifoldMappingSupport.h"
 
 #include UE_INLINE_GENERATED_CPP_BY_NAME(ModelingObjectsCreationAPI)
 
@@ -22,12 +23,20 @@ void FCreateMeshObjectParams::SetMesh(const UE::Geometry::FDynamicMesh3* Dynamic
 {
 	DynamicMesh = *DynamicMeshIn;
 	MeshType = ECreateMeshObjectSourceMeshType::DynamicMesh;
+	if(DynamicMesh)
+	{ 
+		UE::Geometry::FNonManifoldMappingSupport::RemoveAllNonManifoldMappingData(*DynamicMesh);
+	}
 }
 
 void FCreateMeshObjectParams::SetMesh(UE::Geometry::FDynamicMesh3&& DynamicMeshIn)
 {
 	DynamicMesh = MoveTemp(DynamicMeshIn);
 	MeshType = ECreateMeshObjectSourceMeshType::DynamicMesh;
+	if (DynamicMesh)
+	{
+		UE::Geometry::FNonManifoldMappingSupport::RemoveAllNonManifoldMappingData(*DynamicMesh);
+	}
 }
 
 
diff --git a/Engine/Plugins/Runtime/MeshModelingToolset/Source/ModelingComponents/Private/ModelingToolTargetUtil.cpp b/Engine/Plugins/Runtime/MeshModelingToolset/Source/ModelingComponents/Private/ModelingToolTargetUtil.cpp
--- a/Engine/Plugins/Runtime/MeshModelingToolset/Source/ModelingComponents/Private/ModelingToolTargetUtil.cpp
+++ b/Engine/Plugins/Runtime/MeshModelingToolset/Source/ModelingComponents/Private/ModelingToolTargetUtil.cpp
@@ -25,6 +25,8 @@
 #include "DynamicMeshToMeshDescription.h"
 #include "StaticMeshAttributes.h"
 
+#include "DynamicMesh/NonManifoldMappingSupport.h"
+
 #define LOCTEXT_NAMESPACE "ModelingToolTargetUtil"
 
 using namespace UE::Geometry;
@@ -206,6 +208,7 @@ FDynamicMesh3 UE::ToolTarget::GetDynamicMeshCopy(UToolTarget* Target, bool bWant
 	if (MeshDescriptionProvider)
 	{
 		FMeshDescriptionToDynamicMesh Converter;
+		Converter.bVIDsFromNonManifoldMeshDescriptionAttr= true;
 		if (bWantMeshTangents)
 		{
 			FGetMeshParameters GetMeshParams;
@@ -307,6 +310,7 @@ void UE::ToolTarget::Internal::CommitDynamicMeshViaIPersistentDynamicMeshSource(
 	DynamicMesh->EditMesh([&](FDynamicMesh3& EditMesh)
 		{
 			EditMesh.CompactCopy(UpdatedMesh);
+			FNonManifoldMappingSupport::RemoveAllNonManifoldMappingData(EditMesh);
 		});
 
 	TSharedPtr<FDynamicMesh3> NewMeshShared = MakeShared<FDynamicMesh3>();
diff --git a/Engine/Source/Runtime/GeometryCore/Private/DynamicMesh/NonManifoldMappingSupport.cpp b/Engine/Source/Runtime/GeometryCore/Private/DynamicMesh/NonManifoldMappingSupport.cpp
new file mode 100644
--- /dev/null
+++ b/Engine/Source/Runtime/GeometryCore/Private/DynamicMesh/NonManifoldMappingSupport.cpp
@@ -0,0 +1,93 @@
+// Copyright Epic Games, Inc. All Rights Reserved.
+
+#include "DynamicMesh/NonManifoldMappingSupport.h"
+
+#include "CoreMinimal.h"
+#include "DynamicMesh/DynamicMesh3.h"
+#include "DynamicMesh/DynamicMeshAttributeSet.h"
+
+
+FName UE::Geometry::FNonManifoldMappingSupport::NonManifoldMeshVIDsAttrName = FName(TEXT("NonManifoldVIDAttr"));
+
+
+UE::Geometry::FNonManifoldMappingSupport::FNonManifoldMappingSupport(const FDynamicMesh3& MeshIn)
+{
+	Reset(MeshIn);
+}
+
+void UE::Geometry::FNonManifoldMappingSupport::Reset(const FDynamicMesh3& MeshIn)
+{
+	
+	NonManifoldSrcVIDsAttribute = nullptr;
+	DynamicMesh = &MeshIn;
+
+	const FDynamicMeshAttributeSet* Attributes = DynamicMesh->Attributes();
+
+	if (Attributes && Attributes->HasAttachedAttribute(NonManifoldMeshVIDsAttrName))
+	{
+		NonManifoldSrcVIDsAttribute = static_cast<const FDynamicMeshVertexInt32Attribute* >(Attributes->GetAttachedAttribute(NonManifoldMeshVIDsAttrName));
+	}
+	
+}
+
+bool  UE::Geometry::FNonManifoldMappingSupport::IsNonManifoldVertexInSource() const
+{
+	return (NonManifoldSrcVIDsAttribute != nullptr);
+}
+
+int32 UE::Geometry::FNonManifoldMappingSupport::GetOriginalNonManifoldVertexID(const int32 vid) const
+{
+	checkSlow(DynamicMesh->IsVertex(vid));
+
+	int32 Result = vid;
+	if (NonManifoldSrcVIDsAttribute)
+	{
+		NonManifoldSrcVIDsAttribute->GetValue(vid, &Result);
+	}
+	return Result;
+}
+
+bool UE::Geometry::FNonManifoldMappingSupport::AttachNonManifoldVertexMappingData(const TArray<int32>& VertexToNonManifoldVertexIDMap, FDynamicMesh3& MeshInOut)
+{
+	
+	const int32 MaxVID = MeshInOut.MaxVertexID();
+	FDynamicMeshAttributeSet* Attributes = MeshInOut.Attributes();
+
+	if (VertexToNonManifoldVertexIDMap.Num() < MaxVID || !Attributes)
+	{
+		return false;
+	}
+
+	// create and attach a vertex ID buffer, removing any pre-exising one.
+	FDynamicMeshVertexInt32Attribute* VIDAttr = [&] {
+													// replace existing buffer.  we *should* be able to re-use it
+													if (Attributes->HasAttachedAttribute(NonManifoldMeshVIDsAttrName))
+													{
+														Attributes->RemoveAttribute(NonManifoldMeshVIDsAttrName);
+													}
+													FDynamicMeshVertexInt32Attribute* SrcMeshVIDAttr = new FDynamicMeshVertexInt32Attribute(&MeshInOut);
+													SrcMeshVIDAttr->SetName(NonManifoldMeshVIDsAttrName);
+													MeshInOut.Attributes()->AttachAttribute(NonManifoldMeshVIDsAttrName, SrcMeshVIDAttr);
+													return SrcMeshVIDAttr;
+												}();
+	
+	
+	for (int32 vid : MeshInOut.VertexIndicesItr())
+	{
+		const int32 SrcVID = VertexToNonManifoldVertexIDMap[vid];
+		VIDAttr->SetValue(vid, &SrcVID);
+	}
+	return true;
+}
+
+void UE::Geometry::FNonManifoldMappingSupport::RemoveNonManifoldVertexMappingData(FDynamicMesh3& MeshInOut)
+{
+
+	FDynamicMeshAttributeSet* Attributes = MeshInOut.Attributes();
+
+	if (Attributes && Attributes->HasAttachedAttribute(NonManifoldMeshVIDsAttrName))
+	{
+		 Attributes->RemoveAttribute(NonManifoldMeshVIDsAttrName);
+	}
+
+}
\ No newline at end of file
diff --git a/Engine/Source/Runtime/GeometryCore/Public/DynamicMesh/DynamicMeshAttributeSet.h b/Engine/Source/Runtime/GeometryCore/Public/DynamicMesh/DynamicMeshAttributeSet.h
--- a/Engine/Source/Runtime/GeometryCore/Public/DynamicMesh/DynamicMeshAttributeSet.h
+++ b/Engine/Source/Runtime/GeometryCore/Public/DynamicMesh/DynamicMeshAttributeSet.h
@@ -397,11 +397,15 @@ public:
 		}
 	}
 
-	FDynamicMeshAttributeBase* GetAttachedAttribute(FName AttribName) const
+	FDynamicMeshAttributeBase* GetAttachedAttribute(FName AttribName)
 	{
 		return GenericAttributes[AttribName].Get();
 	}
 
+	const FDynamicMeshAttributeBase* GetAttachedAttribute(FName AttribName) const
+	{
+		return GenericAttributes[AttribName].Get();
+	}
 	int NumAttachedAttributes() const
 	{
 		return GenericAttributes.Num();
diff --git a/Engine/Source/Runtime/GeometryCore/Public/DynamicMesh/NonManifoldMappingSupport.h b/Engine/Source/Runtime/GeometryCore/Public/DynamicMesh/NonManifoldMappingSupport.h
new file mode 100644
--- /dev/null
+++ b/Engine/Source/Runtime/GeometryCore/Public/DynamicMesh/NonManifoldMappingSupport.h
@@ -0,0 +1,122 @@
+// Copyright Epic Games, Inc. All Rights Reserved.
+
+
+#pragma once
+
+#include "CoreMinimal.h"
+
+namespace UE
+{
+namespace Geometry
+{
+
+// forward
+class FDynamicMesh3;
+
+template<typename AttribValueType, int AttribDimension, typename ParentType>
+class TDynamicVertexAttribute;
+
+typedef  TDynamicVertexAttribute<int32, 1, FDynamicMesh3> FDynamicMeshVertexInt32Attribute;
+
+
+/*
+* The FNonManifoldMappingSupport is a light-weight class that exists to interpret and manage additional DynamicMesh data stored 
+* when the mesh was created from a potentially non-manifold source mesh.  
+*
+* For context the FDynamicMesh3 is less accepting of non-manifold mesh features than some source mesh data-structures
+* and conversion to a DynamicMesh may require changes to the topology ( e.g. vertex splits)  and in such cases
+* this class can be use to access the mapping from DynamicMesh vertices to the source mesh vertices that generated them.
+* 
+* Note:
+* Since this class holds pointers to named attributes owned by the DynamicMesh, unintended behavior will result if said attributes
+* are removed from the mesh while an associated NomManifoldSupport object is in use. 
+*
+* Also, any mapping information between the DynamicMesh and source data will be corrupted by topological changes 
+* to the DynamicMesh (e.g. edge splits, collapses etc).
+*  
+*/
+class GEOMETRYCORE_API FNonManifoldMappingSupport
+{
+public:
+
+	FNonManifoldMappingSupport(const FDynamicMesh3& Mesh);
+
+
+	/**
+	* Update the support for a new DynamicMesh.
+	*/ 
+	void Reset(const FDynamicMesh3& Mesh);
+
+	/*
+	* Return true if attribute data indicates that the source data that was converted to this DynamicMesh contained non-manifold vertices.
+	*/
+	bool IsNonManifoldVertexInSource() const;
+
+	/*
+	* Return true if the provided DynamicMesh vertex id resulted from a non-manifold vertex in the source data.
+	* @param vid - the id of a vertex in the DynamicMesh.
+	* 
+	* Note: the code assumes but does not check that vid is a valid vertex id
+	*/
+	bool IsNonManifoldVertexID(const int32 vid) const
+	{
+		return !( vid == GetOriginalNonManifoldVertexID(vid) ); 
+	}
+
+	/*
+	* Return the vertex ID in the potentially non-manifold data used to generate this DynamicMesh associated with the provided vertex id.   
+	* In the case that the source data was actually manifold the returned vertex id will be identical to the DynamicMesh vertex id.
+	* @param vid - the id of a vertex in the DynamicMesh. 
+	* 
+	* Note: the code assumes but does not check that vid is a valid vertex vid.
+	*/
+	int32 GetOriginalNonManifoldVertexID(const int32 vid) const;
+
+
+	// --- helper functions.
+	
+	/*
+	* Attaches or replaces non-manifold vertex mapping data to the provided mesh.  
+	* @param VertexToNonManifoldVertexIDMap - an array that maps each DynamicMesh vertex id to the associated non-manifold vertex id.
+	* @return false on failure (no attribute will be attached to the DynamicMesh in this case) 
+	* 
+	* Note: Failure occurs if the DynamicMesh does not have attributes enabled or if the provided array is not long enough to provide a mapping value for each DynamicMesh vertex id.
+	*/
+	static bool AttachNonManifoldVertexMappingData(const TArray<int32>& VertexToNonManifoldVertexIDMap, FDynamicMesh3& MeshInOut);
+
+	/*
+	*  Removes vertex mapping data related to the non-manifold nature of the source data the produced this DynamicMesh.
+	* 
+	*  Note: this will invalidate any NonManifoldMappingSupport object associated with this DynamicMesh, 
+	*  and subsequent use of such object will produce unexpected results.
+	*/
+	static void RemoveNonManifoldVertexMappingData(FDynamicMesh3& MeshInOut);
+
+
+
+	/*
+	*  Removes all  mapping data related to the non-manifold nature of the source data the produced this DynamicMesh.
+	* 
+	*  Note, this will invalidate any NonManifoldMappingSupport object associated with this DynamicMesh, 
+	*  and subsequent use of such object will produce unexpected results.
+	*/
+	static void RemoveAllNonManifoldMappingData(FDynamicMesh3& MeshInOut) 
+	{
+		RemoveNonManifoldVertexMappingData(MeshInOut);
+	}
+
+	/*
+	* Name used to identify vertex attribute data generated during conversion to a DynamicMesh in the case that the source was non-manifold.
+	*/
+	static FName NonManifoldMeshVIDsAttrName;
+
+
+protected:
+
+	const FDynamicMeshVertexInt32Attribute*  NonManifoldSrcVIDsAttribute = nullptr;
+	const FDynamicMesh3* DynamicMesh = nullptr;
+};
+
+
+} // namespace Geometry
+} // namespace UE
\ No newline at end of file
diff --git a/Engine/Source/Runtime/MeshConversion/Private/MeshDescriptionToDynamicMesh.cpp b/Engine/Source/Runtime/MeshConversion/Private/MeshDescriptionToDynamicMesh.cpp
--- a/Engine/Source/Runtime/MeshConversion/Private/MeshDescriptionToDynamicMesh.cpp
+++ b/Engine/Source/Runtime/MeshConversion/Private/MeshDescriptionToDynamicMesh.cpp
@@ -7,6 +7,7 @@
 #include "DynamicMesh/DynamicVertexAttribute.h"
 #include "DynamicMesh/DynamicVertexSkinWeightsAttribute.h"
 #include "DynamicMesh/MeshTangents.h"
+#include "DynamicMesh/NonManifoldMappingSupport.h"
 #include "MeshDescriptionBuilder.h"
 #include "StaticMeshAttributes.h"
 #include "SkeletalMeshAttributes.h"
@@ -241,6 +242,7 @@ void FMeshDescriptionToDynamicMesh::Convert(const FMeshDescription* MeshIn, FDyn
 	// allocate the TriIDMap (maps from DynamicMesh triangle ID to MeshDescription FTriangleID)
 	TriIDMap.AddUninitialized(MeshIn->Triangles().Num());
 	
+	bool bSrcIsManifold = true;
 
 	// Iterate over triangles in the Mesh Description
 	// NOTE: If you change the iteration order here, please update the corresponding iteration in FDynamicMeshToMeshDescription::UpdateAttributes, 
@@ -310,6 +312,8 @@ void FMeshDescriptionToDynamicMesh::Convert(const FMeshDescription* MeshIn, FDyn
 		// if append failed due to non-manifold, duplicate verts
 		if (NewTriangleID == FDynamicMesh3::NonManifoldID)
 		{
+			bSrcIsManifold = false;
+
 			int e0 = MeshOut.FindEdge(VertexIDs[0], VertexIDs[1]);
 			int e1 = MeshOut.FindEdge(VertexIDs[1], VertexIDs[2]);
 			int e2 = MeshOut.FindEdge(VertexIDs[2], VertexIDs[0]);
@@ -851,6 +855,19 @@ void FMeshDescriptionToDynamicMesh::Convert(const FMeshDescription* MeshIn, FDyn
 		// wait for all work to be done
 		UE::Tasks::Wait(Pending);
 
+		// add non-manifold mapping information if required.
+		if (!bSrcIsManifold && bVIDsFromNonManifoldMeshDescriptionAttr)
+		{
+			TArray<int32> TmpVIDMap;
+			TmpVIDMap.SetNumUninitialized(VertIDMap.Num());
+			for (int32 i = 0, I = VertIDMap.Num(); i < I; ++i)
+			{
+				TmpVIDMap[i] = VertIDMap[i].GetValue();
+			}
+			FNonManifoldMappingSupport::AttachNonManifoldVertexMappingData(TmpVIDMap, MeshOut);
+		}
+
+
 		if (!bFoundNonDefaultVertexInstanceColor)
 		{
 			MeshOut.Attributes()->DisablePrimaryColors();
diff --git a/Engine/Source/Runtime/MeshConversion/Public/MeshDescriptionToDynamicMesh.h b/Engine/Source/Runtime/MeshConversion/Public/MeshDescriptionToDynamicMesh.h
--- a/Engine/Source/Runtime/MeshConversion/Public/MeshDescriptionToDynamicMesh.h
+++ b/Engine/Source/Runtime/MeshConversion/Public/MeshDescriptionToDynamicMesh.h
@@ -53,6 +53,16 @@ public:
 	*/
 	TArray<FVertexID> VertIDMap;
 
+
+	/**
+	* If the source MeshDescription is non-manifold, store mesh vertex id as a vertex attribute on the result mesh 
+	* This data can be accessed using FNonManifoldMappingSupport ( see NonManifoldMappingSupport.h)
+	* 
+	* Note: the vertex attribute will not be created if the source mesh is manifold (as this data would be redundant) or if bDisableAttributes is true.
+	*/ 
+	bool bVIDsFromNonManifoldMeshDescriptionAttr = false;
+
+
 	/**
 	 * Various modes can be used to create output triangle groups
 	 */
