BakeRS Tool: Make settings save between tool invocations.

The OnShutdown method was not saving BasicProperties and Setup was not restoring BasicProperties.

diff --git a/Engine/Plugins/Experimental/MeshModelingToolsetExp/Source/MeshModelingToolsExp/Private/BakeTransformTool.cpp b/Engine/Plugins/Experimental/MeshModelingToolsetExp/Source/MeshModelingToolsExp/Private/BakeTransformTool.cpp
--- a/Engine/Plugins/Experimental/MeshModelingToolsetExp/Source/MeshModelingToolsExp/Private/BakeTransformTool.cpp
+++ b/Engine/Plugins/Experimental/MeshModelingToolsetExp/Source/MeshModelingToolsExp/Private/BakeTransformTool.cpp
@@ -63,6 +63,7 @@ void UBakeTransformTool::Setup()
 
 	BasicProperties = NewObject<UBakeTransformToolProperties>(this);
 	AddToolPropertySource(BasicProperties);
+	BasicProperties->RestoreProperties(this);
 
 	FText AllTheWarnings = LOCTEXT("BakeTransformWarning", "WARNING: This Tool will Modify the selected StaticMesh Assets! If you do not wish to modify the original Assets, please make copies in the Content Browser first!");
 
@@ -108,6 +109,8 @@ void UBakeTransformTool::OnShutdown(EToolShutdownType ShutdownType)
 	{
 		UpdateAssets();
 	}
+
+	BasicProperties->SaveProperties(this);
 }
 
 
