ModelingTools: add new modeling-mode Tool/Operator for extruding Mesh Selections

Add UE::ToolTarget::SupportsIncrementalMeshChanges() and ::ApplyIncrementalMeshEditChange(), these allow a Tool to directly edit a DynamicMeshComponent in the level (and emit corresponding MeshChange) so that a Tool can emit multiple undoable edits on a DynamicMesh
Add new USingleTargetWithSelectionTool base tool class, this is a USingleSelectionTool that can receive a FGeometrySelection. USingleTargetWithSelectionToolBuilder configures this type of Tool as necessary.
Add new LinearExtrusionOp dynamic mesh operator, provides linear-extrusion API to FOffsetMeshRegion using new Operator design (mesh shared via FSharedConstDynamicMesh3, CalculateResultInPlace() function to simplify direct usage in geometry script/etc)
Add new ExtrudeMeshSelectionTool, subclass of USingleTargetWithSelectionTool, exposes new LinearExtrusionOp, uses above ToolTarget functions to directly edit DynamicMeshComponent where possible
Enable new Tool in Modeling Mode

Add a few new functions to UMeshOpPreviewWithBackgroundCompute to simplify configuring InProgress and Secondary materials
Small additions to USingleSelectionTool

diff --git a/Engine/Plugins/Editor/ModelingToolsEditorMode/Source/ModelingToolsEditorMode/Private/ModelingToolsEditorMode.cpp b/Engine/Plugins/Editor/ModelingToolsEditorMode/Source/ModelingToolsEditorMode/Private/ModelingToolsEditorMode.cpp
--- a/Engine/Plugins/Editor/ModelingToolsEditorMode/Source/ModelingToolsEditorMode/Private/ModelingToolsEditorMode.cpp
+++ b/Engine/Plugins/Editor/ModelingToolsEditorMode/Source/ModelingToolsEditorMode/Private/ModelingToolsEditorMode.cpp
@@ -94,7 +94,7 @@
 #include "SplitMeshesTool.h"
 #include "PatternTool.h"
 
-//#include "Polymodeling/ExtrudeMeshSelectionTool.h"
+#include "Polymodeling/ExtrudeMeshSelectionTool.h"
 
 #include "Physics/PhysicsInspectorTool.h"
 #include "Physics/SetCollisionGeometryTool.h"
@@ -735,8 +735,8 @@ void UModelingToolsEditorMode::Enter()
 	RegisterPolyModelActionTool(EEditMeshPolygonsToolActions::Outset, ToolManagerCommands.BeginPolyModelTool_Outset, TEXT("PolyEdit_Outset"));
 	RegisterPolyModelActionTool(EEditMeshPolygonsToolActions::CutFaces, ToolManagerCommands.BeginPolyModelTool_CutFaces, TEXT("PolyEdit_CutFaces"));
 
-
-	//RegisterTool(ToolManagerCommands.BeginSelectionAction_Extrude, TEXT("BeginSelectionExtrudeTool"), NewObject<UExtrudeMeshSelectionToolBuilder>());
+	// register mesh-selection-driven tools
+	RegisterTool(ToolManagerCommands.BeginSelectionAction_Extrude, TEXT("BeginSelectionExtrudeTool"), NewObject<UExtrudeMeshSelectionToolBuilder>());
 
 
 	// set up selection type toggles
diff --git a/Engine/Plugins/Editor/ModelingToolsEditorMode/Source/ModelingToolsEditorMode/Private/ModelingToolsEditorModeToolkit.cpp b/Engine/Plugins/Editor/ModelingToolsEditorMode/Source/ModelingToolsEditorMode/Private/ModelingToolsEditorModeToolkit.cpp
--- a/Engine/Plugins/Editor/ModelingToolsEditorMode/Source/ModelingToolsEditorMode/Private/ModelingToolsEditorModeToolkit.cpp
+++ b/Engine/Plugins/Editor/ModelingToolsEditorMode/Source/ModelingToolsEditorMode/Private/ModelingToolsEditorModeToolkit.cpp
@@ -679,7 +679,7 @@ void FModelingToolsEditorModeToolkit::BuildToolPalette(FName PaletteIndex, class
 	else if (PaletteIndex == SelectionActionsTabName)
 	{
 		ToolbarBuilder.AddToolBarButton(Commands.BeginSelectionAction_Delete);
-		//ToolbarBuilder.AddToolBarButton(Commands.BeginSelectionAction_Extrude);
+		ToolbarBuilder.AddToolBarButton(Commands.BeginSelectionAction_Extrude);
 	}
 	else if (PaletteIndex == PrimitiveTabName)
 	{
diff --git a/Engine/Plugins/Experimental/MeshModelingToolsetExp/Source/MeshModelingToolsExp/Private/Polymodeling/ExtrudeMeshSelectionTool.cpp b/Engine/Plugins/Experimental/MeshModelingToolsetExp/Source/MeshModelingToolsExp/Private/Polymodeling/ExtrudeMeshSelectionTool.cpp
new file mode 100644
--- /dev/null
+++ b/Engine/Plugins/Experimental/MeshModelingToolsetExp/Source/MeshModelingToolsExp/Private/Polymodeling/ExtrudeMeshSelectionTool.cpp
@@ -0,0 +1,420 @@
+// Copyright Epic Games, Inc. All Rights Reserved.
+
+#include "Polymodeling/ExtrudeMeshSelectionTool.h"
+#include "InteractiveToolManager.h"
+#include "MeshOpPreviewHelpers.h" // UMeshOpPreviewWithBackgroundCompute
+#include "ToolSetupUtil.h"
+#include "ModelingToolTargetUtil.h"
+#include "Selections/GeometrySelectionUtil.h"
+#include "ToolSceneQueriesUtil.h"
+#include "MeshQueries.h"
+
+#include "Selections/MeshFaceSelection.h"
+#include "DynamicSubmesh3.h"
+#include "Operations/MeshRegionOperator.h"
+#include "Operations/OffsetMeshRegion.h"
+
+#include "ModelingOperators.h"
+#include "PolyModelingOps/LinearExtrusionOp.h"
+#include "Operations/PolyModeling/PolyModelingMaterialUtil.h"
+#include "Operations/PolyModeling/PolyModelingFaceUtil.h"
+
+#include "BaseGizmos/CombinedTransformGizmo.h"
+#include "BaseGizmos/TransformGizmoUtil.h"
+#include "BaseGizmos/TransformProxy.h"
+
+
+using namespace UE::Geometry;
+
+#define LOCTEXT_NAMESPACE "UExtrudeMeshSelectionTool"
+
+
+
+USingleTargetWithSelectionTool* UExtrudeMeshSelectionToolBuilder::CreateNewTool(const FToolBuilderState& SceneState) const
+{
+	return NewObject<UExtrudeMeshSelectionTool>(SceneState.ToolManager);
+}
+
+
+
+class FExtrudeMeshSelectionOpFactory : public IDynamicMeshOperatorFactory
+{
+public:
+	UExtrudeMeshSelectionTool* SourceTool;
+
+	FExtrudeMeshSelectionOpFactory(UExtrudeMeshSelectionTool* Tool)
+	{
+		SourceTool = Tool;
+	}
+
+	TUniquePtr<FLinearExtrusionOp> MakeNewExtrudeOp(
+		TSharedPtr<FSharedConstDynamicMesh3> SourceMesh,
+		const TArray<int32>& ExtrudeROI	)
+	{
+		TUniquePtr<FLinearExtrusionOp> Op = MakeUnique<FLinearExtrusionOp>();
+		Op->OriginalMeshShared = SourceTool->EditRegionSharedMesh;
+		Op->TriangleSelection = ExtrudeROI;
+
+		Op->bShellsToSolids = SourceTool->ExtrudeProperties->bShellsToSolids;
+
+		Op->bInferGroupsFromNeighbours = SourceTool->ExtrudeProperties->bInferGroupsFromNbrs;
+		Op->bNewGroupPerSubdivision = SourceTool->ExtrudeProperties->bGroupPerSubdivision;
+		Op->bUseColinearityForSettingBorderGroups = false;		// what is this??
+		Op->bRemapExtrudeGroups = ! SourceTool->ExtrudeProperties->bReplaceSelectionGroups;
+
+		Op->CreaseAngleThresholdDeg = SourceTool->ExtrudeProperties->CreaseAngle;
+
+		Op->UVScaleFactor = SourceTool->ExtrudeProperties->UVScale;
+		Op->bUVIslandPerGroup = SourceTool->ExtrudeProperties->bUVIslandPerGroup;
+
+		Op->bInferMaterialID = SourceTool->ExtrudeProperties->bInferMaterialID;
+		Op->SetMaterialID = FMath::Min(SourceTool->ExtrudeProperties->SetMaterialID, SourceTool->MaxMaterialID);
+
+		Op->RegionModifierMode = (FLinearExtrusionOp::ESelectionShapeModifierMode)(int32)SourceTool->ExtrudeProperties->RegionMode;
+		Op->NumSubdivisions = SourceTool->ExtrudeProperties->NumSubdivisions;
+
+		Op->RaycastMaxDistance = SourceTool->ExtrudeProperties->RaycastMaxDistance;
+
+		if (SourceTool->ExtrudeProperties->InputMode == EExtrudeMeshSelectionInteractionMode::Interactive)
+		{
+			Op->StartFrame = SourceTool->InitialFrameLocal;
+			Op->ToFrame = SourceTool->ExtrudeFrameLocal;
+			Op->Scale = SourceTool->LocalScale;
+		}
+		else
+		{
+			Op->StartFrame = SourceTool->InitialFrameLocal;
+			Op->ToFrame = FFrame3d(Op->StartFrame.Origin + SourceTool->ExtrudeProperties->ExtrudeDistance * Op->StartFrame.Z(), Op->StartFrame.Rotation);
+			Op->Scale = FVector3d::One();
+		}
+
+		Op->SetResultTransform(SourceTool->WorldTransform);
+
+		return Op;
+	}
+
+	virtual TUniquePtr<FDynamicMeshOperator> MakeNewOperator() override
+	{
+		return MakeNewExtrudeOp(SourceTool->EditRegionSharedMesh, SourceTool->RegionExtrudeROI.Array());
+	}
+};
+
+
+
+void UExtrudeMeshSelectionTool::Setup()
+{
+	UToolTarget* UseTarget = Super::GetTarget();
+	CurrentMesh = UE::ToolTarget::GetDynamicMeshCopy(UseTarget);
+
+	// need valid MaterialIDs for MaterialID parameter
+	FInterval1i MatIDRange;
+	UE::Geometry::ComputeMaterialIDRange(CurrentMesh, MatIDRange);
+	MaxMaterialID = MatIDRange.Max;
+
+	WorldTransform = UE::ToolTarget::GetLocalToWorldTransform(UseTarget);
+	FComponentMaterialSet MaterialSet = UE::ToolTarget::GetMaterialSet(UseTarget);
+
+	ExtrudeProperties = NewObject<UExtrudeMeshSelectionToolProperties>(this);
+	ExtrudeProperties->RestoreProperties(this);
+	ExtrudeProperties->WatchProperty(ExtrudeProperties->InputMode,
+		[this](EExtrudeMeshSelectionInteractionMode NewMode) { UpdateInteractionMode(NewMode); });
+	ExtrudeProperties->WatchProperty(ExtrudeProperties->ExtrudeDistance,
+		[this](double) { EditCompute->InvalidateResult(); });
+	ExtrudeProperties->WatchProperty(ExtrudeProperties->RegionMode,
+		[this](EExtrudeMeshSelectionRegionModifierMode) { EditCompute->InvalidateResult();});
+	ExtrudeProperties->WatchProperty(ExtrudeProperties->NumSubdivisions,
+		[this](int) { EditCompute->InvalidateResult();});
+	ExtrudeProperties->WatchProperty(ExtrudeProperties->bShellsToSolids,
+		[this](bool) { EditCompute->InvalidateResult(); });
+	ExtrudeProperties->WatchProperty(ExtrudeProperties->UVScale,
+		[this](double) { EditCompute->InvalidateResult(); });
+	ExtrudeProperties->WatchProperty(ExtrudeProperties->bInferGroupsFromNbrs,
+		[this](bool) { EditCompute->InvalidateResult(); });
+	ExtrudeProperties->WatchProperty(ExtrudeProperties->bGroupPerSubdivision,
+		[this](bool) { EditCompute->InvalidateResult(); });
+	ExtrudeProperties->WatchProperty(ExtrudeProperties->bUVIslandPerGroup,
+		[this](bool) { EditCompute->InvalidateResult(); });
+	ExtrudeProperties->WatchProperty(ExtrudeProperties->bReplaceSelectionGroups,
+		[this](bool) { EditCompute->InvalidateResult(); });
+	ExtrudeProperties->WatchProperty(ExtrudeProperties->CreaseAngle,
+		[this](double) { EditCompute->InvalidateResult(); });
+	ExtrudeProperties->WatchProperty(ExtrudeProperties->RaycastMaxDistance,
+		[this](double) { EditCompute->InvalidateResult(); });
+	ExtrudeProperties->WatchProperty(ExtrudeProperties->bInferMaterialID,
+		[this](bool) { EditCompute->InvalidateResult(); });
+	ExtrudeProperties->WatchProperty(ExtrudeProperties->SetMaterialID,
+		[this](int) { EditCompute->InvalidateResult(); });
+
+	ExtrudeProperties->WatchProperty(ExtrudeProperties->bShowInputMaterials,
+		[this](bool) { UpdateVisualizationSettings(); });
+
+	AddToolPropertySource(ExtrudeProperties);
+
+	// extract selection
+	FMeshFaceSelection TriSelection(&CurrentMesh);
+	const FGeometrySelection& Selection = GetGeometrySelection();
+	UE::Geometry::EnumerateSelectionTriangles(Selection, CurrentMesh,
+		[&](int32 tid) { TriSelection.Select(tid); },
+		nullptr  /* no support for group layers yet */);
+	ExtrudeROI = TriSelection.AsArray();
+
+	// if we have an empty selection, just abort here
+	if (ExtrudeROI.Num() == 0)
+	{
+		GetToolManager()->PostActiveToolShutdownRequest(this, EToolShutdownType::Cancel,
+			true, LOCTEXT("InvalidSelectionMessage", "ExtrudeMeshSelectionTool: input face selection was empty, cannot Extrude"));
+		return;
+	}
+
+	TriSelection.ExpandToOneRingNeighbours(2);
+	ModifiedROI = TriSelection.AsSet();		// could steal here if it was possible
+
+	SelectionBoundsWorld = (FBox)TMeshQueries<FDynamicMesh3>::GetTrianglesBounds(CurrentMesh, ExtrudeROI, WorldTransform);
+
+	// create the preview object for the unmodified area
+	// (if input was a UDynamicMesh we could do this w/o making a copy...)
+	SourcePreview = NewObject<UPreviewMesh>();
+	SourcePreview->CreateInWorld(GetTargetWorld(), WorldTransform);
+	ToolSetupUtil::ApplyRenderingConfigurationToPreview(SourcePreview, UseTarget);
+	SourcePreview->SetMaterials(MaterialSet.Materials);
+	SourcePreview->EnableSecondaryTriangleBuffers(
+		[this](const FDynamicMesh3* Mesh, int32 TriangleID)
+		{
+			return ExtrudeROI.Contains(TriangleID);		// hide triangles in extrude area in base mesh
+		});
+	SourcePreview->SetSecondaryBuffersVisibility(false);
+	SourcePreview->SetTangentsMode(EDynamicMeshComponentTangentsMode::AutoCalculated);
+	SourcePreview->UpdatePreview(&CurrentMesh);
+
+	// initialize a region operator for the modified area
+	RegionOperator = MakePimpl<FMeshRegionOperator>(&CurrentMesh, ModifiedROI.Array());
+	for (int32 tid : ExtrudeROI)
+	{
+		RegionExtrudeROI.Add( RegionOperator->Region.MapTriangleToSubmesh(tid) );
+	}
+	EditRegionMesh = RegionOperator->Region.GetSubmesh();
+	for (int32 tid : EditRegionMesh.TriangleIndicesItr())
+	{
+		if (RegionExtrudeROI.Contains(tid) == false)
+		{
+			RegionBorderTris.Add(tid);
+		}
+	}
+	EditRegionSharedMesh = MakeShared<FSharedConstDynamicMesh3>(&EditRegionMesh);
+
+	// try to guess selection frame...
+	SelectionFrameLocal = UE::Geometry::ComputeFaceSelectionFrame(EditRegionMesh, RegionExtrudeROI.Array(), false);
+	InitialFrameLocal = SelectionFrameLocal;
+	InitialFrameWorld = InitialFrameLocal;
+	InitialFrameWorld.Transform(WorldTransform);
+
+	this->OperatorFactory = MakePimpl<FExtrudeMeshSelectionOpFactory>(this);
+
+	// Create the preview compute for the extrusion operation
+	EditCompute = NewObject<UMeshOpPreviewWithBackgroundCompute>(this);
+	EditCompute->Setup(GetTargetWorld(), this->OperatorFactory.Get());
+	ToolSetupUtil::ApplyRenderingConfigurationToPreview(EditCompute->PreviewMesh, UseTarget);
+	EditCompute->PreviewMesh->SetTransform((FTransform)WorldTransform);
+	EditCompute->ConfigureMaterials(MaterialSet.Materials, nullptr, nullptr);
+
+	// hide the triangles in the 'border' region outside the extrude area, although they are included in
+	// the extrude region operator submesh, they are still visible in the base mesh
+	EditCompute->PreviewMesh->EnableSecondaryTriangleBuffers(
+		[this](const FDynamicMesh3* Mesh, int32 TriangleID) { return RegionBorderTris.Contains(TriangleID); });
+	EditCompute->PreviewMesh->SetSecondaryBuffersVisibility(false);
+
+	UpdateVisualizationSettings();
+
+	// is this the right tangents behavior?
+	EditCompute->PreviewMesh->SetTangentsMode(EDynamicMeshComponentTangentsMode::AutoCalculated);
+	EditCompute->PreviewMesh->UpdatePreview( & RegionOperator->Region.GetSubmesh() );
+	EditCompute->SetVisibility(true);
+
+	// hide input Component
+	UE::ToolTarget::HideSourceObject(Target);
+
+	// initialize UI stuff
+	Initialize_GizmoMechanic();
+	UpdateInteractionMode(ExtrudeProperties->InputMode);
+}
+
+
+void UExtrudeMeshSelectionTool::OnShutdown(EToolShutdownType ShutdownType)
+{
+	if ( ExtrudeProperties )
+	{
+		ExtrudeProperties->SaveProperties(this);
+	}
+
+	if (SourcePreview != nullptr )
+	{
+		SourcePreview->Disconnect();
+		SourcePreview = nullptr;
+	}
+
+	if (EditCompute != nullptr)
+	{
+		// shut down background computation
+		// (we don't actually care about compute result here because it was only for the live preview...)
+		FDynamicMeshOpResult ComputeResult = EditCompute->Shutdown();
+		EditCompute->ClearOpFactory();
+		EditCompute->OnOpCompleted.RemoveAll(this);
+		EditCompute = nullptr;
+
+		UToolTarget* UseTarget = GetTarget();
+		UE::ToolTarget::ShowSourceObject(UseTarget);
+
+		if (ShutdownType == EToolShutdownType::Accept)
+		{
+			GetToolManager()->BeginUndoTransaction(LOCTEXT("ExtrudeSelection", "Extrude Selection"));
+
+			// this function computes the final Extrude on full Mesh and cleans it up
+			auto ComputeFinalResult = [this](FDynamicMesh3& EditMesh) -> bool
+			{
+				TSharedPtr<FSharedConstDynamicMesh3> FullMeshShared = MakeShared<FSharedConstDynamicMesh3>(&EditMesh);
+				TUniquePtr<FLinearExtrusionOp> FinalOp = OperatorFactory->MakeNewExtrudeOp(FullMeshShared, this->ExtrudeROI);
+				bool bResult = FinalOp->CalculateResultInPlace(EditMesh, nullptr);
+				if (EditMesh.IsCompact() == false)
+				{
+					EditMesh.CompactInPlace();
+				}
+				FullMeshShared->ReleaseSharedObject();
+				return bResult;
+			};
+
+			// try to emit an incremental FMeshChange, to avoid storing full mesh before/after copies in undo history
+			bool bChangeApplied = false;
+			if (UE::ToolTarget::SupportsIncrementalMeshChanges(UseTarget))
+			{
+				bChangeApplied = UE::ToolTarget::ApplyIncrementalMeshEditChange(UseTarget,
+					[&](FDynamicMesh3& EditMesh, UObject* TransactionTarget)
+					{
+						FDynamicMeshChangeTracker ChangeTracker(&EditMesh);
+						ChangeTracker.BeginChange();
+						ChangeTracker.SaveTriangles(ExtrudeROI, true);
+						if (ComputeFinalResult(EditMesh))
+						{
+							TUniquePtr<FDynamicMeshChange> MeshEditChange = ChangeTracker.EndChange();
+							GetToolManager()->GetContextTransactionsAPI()->AppendChange( TransactionTarget,
+								MakeUnique<FMeshChange>(MoveTemp(MeshEditChange)), LOCTEXT("ExtrudeSelection", "Extrude Selection") );
+							return true;
+						}
+						return false;
+					});
+
+				ensureMsgf(bChangeApplied, TEXT("UExtrudeMeshSelectionTool::OnShutdown : incremental mesh edit failed!"));
+			}
+
+			// if we could not apply incremental change, apply a full-mesh update
+			if ( bChangeApplied == false )
+			{
+				ComputeFinalResult(CurrentMesh);
+				const bool bModifiedTopology = true;
+				UE::ToolTarget::CommitDynamicMeshUpdate(UseTarget, CurrentMesh, bModifiedTopology);
+			}
+
+			GetToolManager()->EndUndoTransaction();
+		}
+	}
+
+	if ( TransformGizmo != nullptr )
+	{
+		GetToolManager()->GetPairedGizmoManager()->DestroyAllGizmosByOwner(this);
+		TransformProxy = nullptr;
+		TransformGizmo = nullptr;
+	}
+
+	if ( EditRegionSharedMesh.IsValid() )
+	{
+		EditRegionSharedMesh->ReleaseSharedObject();
+	}
+
+	CurrentMesh.Clear();
+}
+
+
+
+void UExtrudeMeshSelectionTool::OnTick(float DeltaTime)
+{
+	EditCompute->Tick(DeltaTime);
+}
+
+
+void UExtrudeMeshSelectionTool::Render(IToolsContextRenderAPI* RenderAPI)
+{
+}
+
+FBox UExtrudeMeshSelectionTool::GetWorldSpaceFocusBox() 
+{ 
+	FAxisAlignedBox3d CurBox = SelectionBoundsWorld;
+	FVector3d Translation = ExtrudeFrameWorld.Origin - InitialFrameWorld.Origin;
+	CurBox.Min += Translation; CurBox.Max += Translation;
+	CurBox.Contain(SelectionBoundsWorld);
+	return (FBox)CurBox;
+}
+
+void UExtrudeMeshSelectionTool::Initialize_GizmoMechanic()
+{
+	// Set up the gizmo.
+	TransformProxy = NewObject<UTransformProxy>(this);
+	TransformProxy->SetTransform( InitialFrameWorld.ToFTransform() );
+
+	//TransformGizmo = UE::TransformGizmoUtil::CreateCustomTransformGizmo(
+	TransformGizmo = UE::TransformGizmoUtil::CreateCustomRepositionableTransformGizmo(
+		GetToolManager()->GetPairedGizmoManager(),
+		ETransformGizmoSubElements::FullTranslateRotateScale, this);
+
+	TransformGizmo->SetActiveTarget(TransformProxy, GetToolManager());
+	TransformProxy->OnTransformChanged.AddUObject(this, &UExtrudeMeshSelectionTool::GizmoTransformChanged);
+
+	ExtrudeFrameWorld = InitialFrameWorld;
+	ExtrudeFrameLocal = ExtrudeFrameWorld;
+	ExtrudeFrameLocal.Transform( WorldTransform.InverseUnsafe() );
+	LocalScale = FVector3d::One();
+}
+
+void UExtrudeMeshSelectionTool::GizmoTransformChanged(UTransformProxy* Proxy, FTransform Transform)
+{
+	ExtrudeFrameWorld = (FFrame3d)Transform;
+	ExtrudeFrameLocal = ExtrudeFrameWorld;
+	ExtrudeFrameLocal.Transform( WorldTransform.InverseUnsafe() );
+
+	LocalScale = Transform.GetScale3D();
+
+	EditCompute->InvalidateResult();
+}
+
+
+void UExtrudeMeshSelectionTool::UpdateVisualizationSettings()
+{
+	if (ExtrudeProperties->bShowInputMaterials)
+	{
+		EditCompute->DisablePreviewMaterials();
+		Cast<UDynamicMeshComponent>(EditCompute->PreviewMesh->GetRootComponent())->SetColorOverrideMode(EDynamicMeshComponentColorOverrideMode::None);
+	}
+	else
+	{
+		EditCompute->ConfigurePreviewMaterials(
+			ToolSetupUtil::GetDefaultWorkingMaterial(GetToolManager()),
+			ToolSetupUtil::GetSelectionMaterial(FLinearColor::Yellow, GetToolManager())  );
+		Cast<UDynamicMeshComponent>(EditCompute->PreviewMesh->GetRootComponent())->SetColorOverrideMode(EDynamicMeshComponentColorOverrideMode::Polygroups);
+	}
+}
+
+
+void UExtrudeMeshSelectionTool::UpdateInteractionMode(EExtrudeMeshSelectionInteractionMode InteractionMode)
+{
+	if ( InteractionMode == EExtrudeMeshSelectionInteractionMode::Interactive )
+	{
+		TransformGizmo->SetVisibility(true);
+	}
+	else
+	{
+		TransformGizmo->SetVisibility(false);
+	}
+
+	EditCompute->InvalidateResult();
+}
+
+
+#undef LOCTEXT_NAMESPACE
\ No newline at end of file
diff --git a/Engine/Plugins/Experimental/MeshModelingToolsetExp/Source/MeshModelingToolsExp/Public/Polymodeling/ExtrudeMeshSelectionTool.h b/Engine/Plugins/Experimental/MeshModelingToolsetExp/Source/MeshModelingToolsExp/Public/Polymodeling/ExtrudeMeshSelectionTool.h
new file mode 100644
--- /dev/null
+++ b/Engine/Plugins/Experimental/MeshModelingToolsetExp/Source/MeshModelingToolsExp/Public/Polymodeling/ExtrudeMeshSelectionTool.h
@@ -0,0 +1,190 @@
+// Copyright Epic Games, Inc. All Rights Reserved.
+
+#pragma once
+
+#include "BaseTools/SingleTargetWithSelectionTool.h"
+#include "DynamicMesh/DynamicMesh3.h"
+#include "DynamicMesh/MeshSharingUtil.h"
+#include "Templates/PimplPtr.h"
+#include "ExtrudeMeshSelectionTool.generated.h"
+
+class UPreviewMesh;
+class UMeshOpPreviewWithBackgroundCompute; 
+class FExtrudeMeshSelectionOpFactory;
+class UTransformProxy;
+class UCombinedTransformGizmo;
+namespace UE::Geometry { class FMeshRegionOperator; }
+
+/**
+ * ToolBuilder
+ */
+UCLASS()
+class MESHMODELINGTOOLSEXP_API UExtrudeMeshSelectionToolBuilder : public USingleTargetWithSelectionToolBuilder
+{
+	GENERATED_BODY()
+public:
+	virtual USingleTargetWithSelectionTool* CreateNewTool(const FToolBuilderState& SceneState) const override;
+	virtual bool RequiresInputSelection() const override { return true; }
+};
+
+
+UENUM()
+enum class EExtrudeMeshSelectionInteractionMode : uint8
+{
+	/** Define the extrusion distance using a 3D gizmo */
+	Interactive = 0,
+	/** Define the extrusion distance using a slider in the Settings */
+	Fixed = 1
+};
+
+UENUM()
+enum class EExtrudeMeshSelectionRegionModifierMode : uint8
+{
+	/** Transform the original selected area */
+	OriginalShape = 0,
+	/** Flatten the extrusion area to the X/Y plane of the extrusion frame */
+	FlattenToPlane = 1,
+	/** Flatten the extrusion area by raycasting against the X/Y plane of the extrusion frame */
+	RaycastToPlane = 2
+};
+
+UCLASS()
+class MESHMODELINGTOOLSEXP_API UExtrudeMeshSelectionToolProperties : public UInteractiveToolPropertySet
+{
+	GENERATED_BODY()
+
+public:
+	/** Control how the Extruded Area should be Transformed */
+	UPROPERTY(EditAnywhere, Category = "Extrude")
+	EExtrudeMeshSelectionInteractionMode InputMode = EExtrudeMeshSelectionInteractionMode::Interactive;
+
+	/** The Extrusion Distance used in Fixed Input Mode*/
+	UPROPERTY(EditAnywhere, Category = "Extrude", meta = (DisplayName="Fixed Distance", UIMin = -250, UIMax = 250, EditCondition = "InputMode == EExtrudeMeshSelectionInteractionMode::Fixed"))
+	double ExtrudeDistance = 10.0;
+
+	/** Control how the Extruded Area should be deformed as part of the Extrusion */
+	UPROPERTY(EditAnywhere, Category = "Extrude")
+	EExtrudeMeshSelectionRegionModifierMode RegionMode = EExtrudeMeshSelectionRegionModifierMode::OriginalShape;
+
+	/** Specify the number of subdivisions along the sides of the Extrusion */
+	UPROPERTY(EditAnywhere, Category = "Extrude", meta = (DisplayName="Subdivisions", UIMin = 0, UIMax = 10, ClampMin = 0, ClampMax = 1000))
+	int NumSubdivisions = 0;
+
+	/** Specify the Crease Angle used to split the sides of the Extrusion into separate Groups */
+	UPROPERTY(EditAnywhere, Category = "Extrude", meta = (UIMin = 0, UIMax = 180, Min = 0, Max = 180))
+	double CreaseAngle = 60.0;
+
+	/** Control the maximum distance each vertex may be moved in Raycast To Plane Mode */
+	UPROPERTY(EditAnywhere, Category = "Extrude", meta = (DisplayName="Max Distance", UIMin = 0, UIMax = 1000, ClampMin = 0, EditCondition = "RegionMode == EExtrudeMeshSelectionRegionModifierMode::RaycastToPlane"))
+	double RaycastMaxDistance = 1000.0;
+
+	/** If the Extruded Area has a fully open border, this option determines if Extrusion will create a Solid mesh or leave the base "open" */
+	UPROPERTY(EditAnywhere, Category = "Extrude")
+	bool bShellsToSolids = true;			// todo: should only be available if input has shells that could become solid...
+
+	/** Control whether a single Group should be generated along the sides of the Extrusion, or multiple Groups based on the adjacent Groups around the Extruded Area border */
+	UPROPERTY(EditAnywhere, Category = "Groups", meta = (DisplayName="Propagate Groups"))
+	bool bInferGroupsFromNbrs = true;
+
+	/** Control whether a new Group is generated for each Subdivision */
+	UPROPERTY(EditAnywhere, Category = "Groups", meta=(DisplayName="Per Subdivision", EditCondition="NumSubdivisions > 0") )
+	bool bGroupPerSubdivision = true;
+
+	/** Control whether groups in the Extruded Area are mapped to new Groups, or replaced with a single new Group */
+	UPROPERTY(EditAnywhere, Category = "Groups", meta=(DisplayName="Replace Cap Groups"))
+	bool bReplaceSelectionGroups = false;
+
+	/** The automatically-generated UVs on the sides of the Extrusion are scaled by this value */
+	UPROPERTY(EditAnywhere, Category = "UVs", meta = (DisplayName="UV Scale", UIMin = 0.001, UIMax = 10.0, ClampMin = 0.000001))
+	double UVScale = 1.0f;
+
+	/** Control whether a separate UV island should be generated for each output Group on the sides of the Extrusion, or a single UV island wrapping around the entire "tube" */
+	UPROPERTY(EditAnywhere, Category = "UVs", meta = (DisplayName="Island Per Group"))
+	bool bUVIslandPerGroup = true;
+
+	/** Control whether SetMaterialID is assigned to all triangles along the sides of the Extrusion, or if MaterialIDs should be inferred from the Extruded Area */
+	UPROPERTY(EditAnywhere, Category = "Material", meta=(DisplayName="Infer Materials"))
+	bool bInferMaterialID = true;
+
+	/** Constant Material ID used when MaterialIDs are not being inferred, or no adjacent MaterialID exists */
+	UPROPERTY(EditAnywhere, Category = "Material", meta=(DisplayName="Material ID", EditCondition="bInferMaterialID == false"))
+	int SetMaterialID = 0;
+
+	/** Control whether the original Mesh Materials should be shown, or a visualization of the extruded Groups */
+	UPROPERTY(EditAnywhere, Category = "Visualization", meta=(DisplayName="Show Materials"))
+	bool bShowInputMaterials = false;
+};
+
+
+
+UCLASS()
+class MESHMODELINGTOOLSEXP_API UExtrudeMeshSelectionTool : public USingleTargetWithSelectionTool
+{
+	GENERATED_BODY()
+	using FFrame3d = UE::Geometry::FFrame3d;
+public:
+
+	virtual void Setup() override;
+	virtual void OnShutdown(EToolShutdownType ShutdownType) override;
+
+	virtual void OnTick(float DeltaTime) override;
+	virtual void Render(IToolsContextRenderAPI* RenderAPI) override;
+
+	virtual bool HasCancel() const override { return true; }
+	virtual bool HasAccept() const override { return true; }
+
+	// IInteractiveToolCameraFocusAPI implementation
+	virtual FBox GetWorldSpaceFocusBox() override;
+
+
+protected:
+	UPROPERTY()
+	TObjectPtr<UExtrudeMeshSelectionToolProperties> ExtrudeProperties = nullptr;
+
+
+protected:
+	FBox SelectionBoundsWorld;
+	UE::Geometry::FTransformSRT3d WorldTransform;
+
+	UE::Geometry::FFrame3d SelectionFrameLocal;
+	UE::Geometry::FFrame3d InitialFrameLocal;
+	UE::Geometry::FFrame3d InitialFrameWorld;
+	UE::Geometry::FFrame3d ExtrudeFrameWorld;
+	UE::Geometry::FFrame3d ExtrudeFrameLocal;
+	FVector3d LocalScale;
+
+	UE::Geometry::FDynamicMesh3 CurrentMesh;
+	TArray<int32> ExtrudeROI;
+	TSet<int32> ModifiedROI;
+	int32 MaxMaterialID = 0;		// maximum material ID on mesh (inclusive)
+
+	TPimplPtr<UE::Geometry::FMeshRegionOperator> RegionOperator;
+	UE::Geometry::FDynamicMesh3 EditRegionMesh;
+	TSharedPtr<UE::Geometry::FSharedConstDynamicMesh3> EditRegionSharedMesh;
+
+	TSet<int32> RegionExtrudeROI;
+	TSet<int32> RegionBorderTris;
+
+	TPimplPtr<FExtrudeMeshSelectionOpFactory> OperatorFactory;
+	friend class FExtrudeMeshSelectionOpFactory;
+
+	UPROPERTY()
+	TObjectPtr<UPreviewMesh> SourcePreview = nullptr;
+
+	UPROPERTY()
+	TObjectPtr<UMeshOpPreviewWithBackgroundCompute> EditCompute = nullptr;
+
+	void UpdateVisualizationSettings();
+
+	void Initialize_GizmoMechanic();
+
+	UPROPERTY()
+	TObjectPtr<UCombinedTransformGizmo> TransformGizmo = nullptr;
+
+	UPROPERTY()
+	TObjectPtr<UTransformProxy> TransformProxy = nullptr;
+
+	void GizmoTransformChanged(UTransformProxy* Proxy, FTransform Transform);
+
+	void UpdateInteractionMode(EExtrudeMeshSelectionInteractionMode InteractionMode);
+};
\ No newline at end of file
diff --git a/Engine/Plugins/Runtime/MeshModelingToolset/Source/ModelingComponents/Private/BaseTools/SingleTargetWithSelectionTool.cpp b/Engine/Plugins/Runtime/MeshModelingToolset/Source/ModelingComponents/Private/BaseTools/SingleTargetWithSelectionTool.cpp
new file mode 100644
--- /dev/null
+++ b/Engine/Plugins/Runtime/MeshModelingToolset/Source/ModelingComponents/Private/BaseTools/SingleTargetWithSelectionTool.cpp
@@ -0,0 +1,102 @@
+// Copyright Epic Games, Inc. All Rights Reserved.
+
+#include "BaseTools/SingleTargetWithSelectionTool.h"
+
+#include "Engine/World.h"
+
+#include "TargetInterfaces/MaterialProvider.h"
+#include "TargetInterfaces/MeshDescriptionCommitter.h"
+#include "TargetInterfaces/MeshDescriptionProvider.h"
+#include "TargetInterfaces/PrimitiveComponentBackedTarget.h"
+#include "ToolTargetManager.h"
+#include "Selection/StoredMeshSelectionUtil.h"
+#include "Selections/GeometrySelection.h"
+
+#include UE_INLINE_GENERATED_CPP_BY_NAME(SingleTargetWithSelectionTool)
+
+
+/*
+ * ToolBuilder
+ */
+const FToolTargetTypeRequirements& USingleTargetWithSelectionToolBuilder::GetTargetRequirements() const
+{
+	static FToolTargetTypeRequirements TypeRequirements({
+		UMaterialProvider::StaticClass(),
+		UMeshDescriptionCommitter::StaticClass(),
+		UMeshDescriptionProvider::StaticClass(),
+		UPrimitiveComponentBackedTarget::StaticClass()
+		});
+	return TypeRequirements;
+}
+
+bool USingleTargetWithSelectionToolBuilder::CanBuildTool(const FToolBuilderState& SceneState) const
+{
+	return SceneState.TargetManager->CountSelectedAndTargetable(SceneState, GetTargetRequirements()) == 1;
+}
+
+UInteractiveTool* USingleTargetWithSelectionToolBuilder::BuildTool(const FToolBuilderState& SceneState) const
+{
+	USingleTargetWithSelectionTool* NewTool = CreateNewTool(SceneState);
+	InitializeNewTool(NewTool, SceneState);
+	return NewTool;
+}
+
+void USingleTargetWithSelectionToolBuilder::InitializeNewTool(USingleTargetWithSelectionTool* NewTool, const FToolBuilderState& SceneState) const
+{
+	UToolTarget* Target = SceneState.TargetManager->BuildFirstSelectedTargetable(SceneState, GetTargetRequirements());
+	check(Target);
+	NewTool->SetTarget(Target);
+	NewTool->SetTargetWorld(SceneState.World);
+
+	UE::Geometry::FGeometrySelection Selection;
+	bool bHaveSelection = UE::Geometry::GetCurrentGeometrySelectionForTarget(SceneState, Target, Selection);
+	if (bHaveSelection)
+	{
+		NewTool->SetGeometrySelection(MoveTemp(Selection));
+	}
+
+}
+
+
+void USingleTargetWithSelectionTool::Shutdown(EToolShutdownType ShutdownType)
+{
+	OnShutdown(ShutdownType);
+	TargetWorld = nullptr;
+}
+
+void USingleTargetWithSelectionTool::OnShutdown(EToolShutdownType ShutdownType)
+{
+}
+
+void USingleTargetWithSelectionTool::SetTargetWorld(UWorld* World)
+{
+	TargetWorld = World;
+}
+
+UWorld* USingleTargetWithSelectionTool::GetTargetWorld()
+{
+	return TargetWorld.Get();
+}
+
+
+void USingleTargetWithSelectionTool::SetGeometrySelection(const UE::Geometry::FGeometrySelection& SelectionIn)
+{
+	GeometrySelection = SelectionIn;
+	bGeometrySelectionInitialized = true;
+}
+
+void USingleTargetWithSelectionTool::SetGeometrySelection(UE::Geometry::FGeometrySelection&& SelectionIn)
+{
+	GeometrySelection = MoveTemp(SelectionIn);
+	bGeometrySelectionInitialized = true;
+}
+
+bool USingleTargetWithSelectionTool::HasGeometrySelection() const
+{
+	return bGeometrySelectionInitialized;
+}
+
+const UE::Geometry::FGeometrySelection& USingleTargetWithSelectionTool::GetGeometrySelection() const
+{
+	return GeometrySelection;
+}
\ No newline at end of file
diff --git a/Engine/Plugins/Runtime/MeshModelingToolset/Source/ModelingComponents/Private/MeshOpPreviewHelpers.cpp b/Engine/Plugins/Runtime/MeshModelingToolset/Source/ModelingComponents/Private/MeshOpPreviewHelpers.cpp
--- a/Engine/Plugins/Runtime/MeshModelingToolset/Source/ModelingComponents/Private/MeshOpPreviewHelpers.cpp
+++ b/Engine/Plugins/Runtime/MeshModelingToolset/Source/ModelingComponents/Private/MeshOpPreviewHelpers.cpp
@@ -189,10 +189,14 @@ void UMeshOpPreviewWithBackgroundCompute::ConfigureMaterials(UMaterialInterface*
 }
 
 
-void UMeshOpPreviewWithBackgroundCompute::ConfigureMaterials(TArray<UMaterialInterface*> StandardMaterialsIn, UMaterialInterface* WorkingMaterialIn)
+void UMeshOpPreviewWithBackgroundCompute::ConfigureMaterials(
+	TArray<UMaterialInterface*> StandardMaterialsIn, 
+	UMaterialInterface* WorkingMaterialIn,
+	UMaterialInterface* SecondaryMaterialIn)
 {
 	this->StandardMaterials = StandardMaterialsIn;
 	this->WorkingMaterial = WorkingMaterialIn;
+	this->SecondaryMaterial = SecondaryMaterialIn;
 
 	if (PreviewMesh != nullptr)
 	{
@@ -200,6 +204,19 @@ void UMeshOpPreviewWithBackgroundCompute::ConfigureMaterials(TArray<UMaterialInt
 	}
 }
 
+void UMeshOpPreviewWithBackgroundCompute::ConfigurePreviewMaterials(
+	UMaterialInterface* InProgressMaterialIn,
+	UMaterialInterface* SecondaryMaterialIn)
+{
+	this->WorkingMaterial = InProgressMaterialIn;
+	this->SecondaryMaterial = SecondaryMaterialIn;
+}
+
+void UMeshOpPreviewWithBackgroundCompute::DisablePreviewMaterials()
+{
+	this->WorkingMaterial = nullptr;
+	this->SecondaryMaterial = nullptr;
+}
 
 void UMeshOpPreviewWithBackgroundCompute::SetVisibility(bool bVisibleIn)
 {
diff --git a/Engine/Plugins/Runtime/MeshModelingToolset/Source/ModelingComponents/Private/ModelingToolTargetUtil.cpp b/Engine/Plugins/Runtime/MeshModelingToolset/Source/ModelingComponents/Private/ModelingToolTargetUtil.cpp
--- a/Engine/Plugins/Runtime/MeshModelingToolset/Source/ModelingComponents/Private/ModelingToolTargetUtil.cpp
+++ b/Engine/Plugins/Runtime/MeshModelingToolset/Source/ModelingComponents/Private/ModelingToolTargetUtil.cpp
@@ -486,6 +486,34 @@ UE::ToolTarget::EDynamicMeshUpdateResult UE::ToolTarget::CommitDynamicMeshNormal
 
 
 
+bool UE::ToolTarget::SupportsIncrementalMeshChanges(UToolTarget* Target)
+{
+	IPersistentDynamicMeshSource* DynamicMeshSource = Cast<IPersistentDynamicMeshSource>(Target);
+	if (DynamicMeshSource)
+	{
+		return DynamicMeshSource->HasDynamicMeshComponent();
+	}
+	return false;
+}
+
+
+bool UE::ToolTarget::ApplyIncrementalMeshEditChange(
+	UToolTarget* Target,
+	TFunctionRef<bool(FDynamicMesh3& EditMesh, UObject* TransactionTarget)> MeshEditChangeFunc )
+{
+	IPersistentDynamicMeshSource* DynamicMeshSource = Cast<IPersistentDynamicMeshSource>(Target);
+	if (DynamicMeshSource && DynamicMeshSource->HasDynamicMeshComponent())
+	{
+		bool bOK;
+		UDynamicMeshComponent* Component = DynamicMeshSource->GetDynamicMeshComponent();
+		Component->EditMesh([&](FDynamicMesh3& EditMesh)
+		{
+			bOK = MeshEditChangeFunc(EditMesh, Component);
+		});
+		return bOK;
+	}
+	return false;
+}
 
 
 bool UE::ToolTarget::ConfigureCreateMeshObjectParams(UToolTarget* SourceTarget, FCreateMeshObjectParams& DerivedParamsOut)
diff --git a/Engine/Plugins/Runtime/MeshModelingToolset/Source/ModelingComponents/Public/BaseTools/SingleTargetWithSelectionTool.h b/Engine/Plugins/Runtime/MeshModelingToolset/Source/ModelingComponents/Public/BaseTools/SingleTargetWithSelectionTool.h
new file mode 100644
--- /dev/null
+++ b/Engine/Plugins/Runtime/MeshModelingToolset/Source/ModelingComponents/Public/BaseTools/SingleTargetWithSelectionTool.h
@@ -0,0 +1,66 @@
+// Copyright Epic Games, Inc. All Rights Reserved.
+
+#pragma once
+
+#include "SingleSelectionTool.h"
+#include "InteractiveToolBuilder.h"
+#include "Selections/GeometrySelection.h"
+#include "SingleTargetWithSelectionTool.generated.h"
+
+
+UCLASS(Transient, Abstract)
+class MODELINGCOMPONENTS_API USingleTargetWithSelectionToolBuilder : public UInteractiveToolWithToolTargetsBuilder
+{
+	GENERATED_BODY()
+public:
+	/** @return true if a single mesh source can be found in the active selection */
+	virtual bool CanBuildTool(const FToolBuilderState& SceneState) const override;
+
+	/** @return new Tool instance initialized with selected mesh source */
+	virtual UInteractiveTool* BuildTool(const FToolBuilderState& SceneState) const override;
+
+	/** @return new Tool instance. Override this in subclasses to build a different Tool class type */
+	virtual USingleTargetWithSelectionTool* CreateNewTool(const FToolBuilderState& SceneState) const PURE_VIRTUAL(USingleTargetWithSelectionToolBuilder::CreateNewTool, return nullptr; );
+
+	/** Called by BuildTool to configure the Tool with the input MeshSource based on the SceneState */
+	virtual void InitializeNewTool(USingleTargetWithSelectionTool* Tool, const FToolBuilderState& SceneState) const;
+
+	/** @return true if this Tool requires an input selection */
+	virtual bool RequiresInputSelection() const { return false; }
+
+protected:
+	virtual const FToolTargetTypeRequirements& GetTargetRequirements() const override;
+};
+
+
+
+UCLASS()
+class MODELINGCOMPONENTS_API USingleTargetWithSelectionTool : public USingleSelectionTool
+{
+	GENERATED_BODY()
+public:
+	virtual void Shutdown(EToolShutdownType ShutdownType) override;
+	virtual void OnShutdown(EToolShutdownType ShutdownType);
+
+	virtual void SetTargetWorld(UWorld* World);
+	virtual UWorld* GetTargetWorld();
+
+protected:
+	UPROPERTY()
+	TWeakObjectPtr<UWorld> TargetWorld = nullptr;
+
+
+public:
+	virtual void SetGeometrySelection(const UE::Geometry::FGeometrySelection& SelectionIn);
+	virtual void SetGeometrySelection(UE::Geometry::FGeometrySelection&& SelectionIn);
+
+	/** @return true if a Selection is available */
+	virtual bool HasGeometrySelection() const;
+
+	/** @return the input Selection */
+	virtual const UE::Geometry::FGeometrySelection& GetGeometrySelection() const;
+
+protected:
+	UE::Geometry::FGeometrySelection GeometrySelection;
+	bool bGeometrySelectionInitialized = false;
+};
\ No newline at end of file
diff --git a/Engine/Plugins/Runtime/MeshModelingToolset/Source/ModelingComponents/Public/MeshOpPreviewHelpers.h b/Engine/Plugins/Runtime/MeshModelingToolset/Source/ModelingComponents/Public/MeshOpPreviewHelpers.h
--- a/Engine/Plugins/Runtime/MeshModelingToolset/Source/ModelingComponents/Public/MeshOpPreviewHelpers.h
+++ b/Engine/Plugins/Runtime/MeshModelingToolset/Source/ModelingComponents/Public/MeshOpPreviewHelpers.h
@@ -161,9 +161,26 @@ public:
 
 	/**
 	 * Configure the Standard and In-Progress materials
+	 * @param StandardMaterials list of materials corresponding to MaterialIDs on the Mesh
+	 * @param InProgressMaterial material displayed if the background compute takes longer than WorkingMaterialDelay
+	 * @param SecondaryMaterial optional secondary material passed to mesh component, call PreviewMesh->EnableSecondaryTriangleBuffers to enable
 	 */
-	void ConfigureMaterials(TArray<UMaterialInterface*> StandardMaterials, UMaterialInterface* InProgressMaterial);
+	void ConfigureMaterials(
+		TArray<UMaterialInterface*> StandardMaterials, 
+		UMaterialInterface* InProgressMaterial, 
+		UMaterialInterface* SecondaryMaterial = nullptr);
+
+	/** 
+	 * Configure the In-Progress and Secondary Materials
+	 * @param InProgressMaterial material displayed if the background compute takes longer than WorkingMaterialDelay
+	 * @param SecondaryMaterial optional secondary material passed to mesh component, call PreviewMesh->EnableSecondaryTriangleBuffers to enable
+	 */
+	void ConfigurePreviewMaterials(
+		UMaterialInterface* InProgressMaterial, 
+		UMaterialInterface* SecondaryMaterial);
 
+	/** Disable the In Progress and Secondary Materials */
+	void DisablePreviewMaterials();
 
 	/**
 	 * Set the visibility of the Preview mesh
diff --git a/Engine/Plugins/Runtime/MeshModelingToolset/Source/ModelingComponents/Public/ModelingToolTargetUtil.h b/Engine/Plugins/Runtime/MeshModelingToolset/Source/ModelingComponents/Public/ModelingToolTargetUtil.h
--- a/Engine/Plugins/Runtime/MeshModelingToolset/Source/ModelingComponents/Public/ModelingToolTargetUtil.h
+++ b/Engine/Plugins/Runtime/MeshModelingToolset/Source/ModelingComponents/Public/ModelingToolTargetUtil.h
@@ -186,6 +186,24 @@ MODELINGCOMPONENTS_API EDynamicMeshUpdateResult CommitDynamicMeshUVUpdate(UToolT
 */
 MODELINGCOMPONENTS_API EDynamicMeshUpdateResult CommitDynamicMeshNormalsUpdate(UToolTarget* Target, const UE::Geometry::FDynamicMesh3* UpdatedMesh, bool bUpdateTangents = false);
 
+/**
+ * @return true if ToolTarget can be directly updated via an incremental edits, ie if it is acceptable to call ApplyIncrementalMeshEditChange() on the target
+ */
+MODELINGCOMPONENTS_API bool SupportsIncrementalMeshChanges(UToolTarget* Target);
+
+
+/**
+ * Apply an incremental mesh edit (MeshEditingFunc) to the ToolTarget. Tools/etc can use this to directly
+ * modify the ToolTarget's internal DynamicMesh, while (presumably) also emitting a FChange transaction that
+ * represents that Edit. This allows for undoable edits directly to a DynamicMeshComponent/etc in the level. 
+ * @return true if the incremental edit succeeds
+ */
+MODELINGCOMPONENTS_API bool ApplyIncrementalMeshEditChange(
+	UToolTarget* Target,
+	TFunctionRef<bool(UE::Geometry::FDynamicMesh3& EditMesh, UObject* TransactionTarget)> MeshEditingFunc );
+
+
+
 /**
  * FCreateMeshObjectParams::TypeHint is used by the ModelingObjectsCreationAPI to suggest what type of mesh object to create
  * inside various Tools. This should often be derived from the input mesh object type (eg if you plane-cut a Volume, the output
diff --git a/Engine/Plugins/Runtime/MeshModelingToolset/Source/ModelingOperators/Private/PolyModelingOps/LinearExtrusionOp.cpp b/Engine/Plugins/Runtime/MeshModelingToolset/Source/ModelingOperators/Private/PolyModelingOps/LinearExtrusionOp.cpp
new file mode 100644
--- /dev/null
+++ b/Engine/Plugins/Runtime/MeshModelingToolset/Source/ModelingOperators/Private/PolyModelingOps/LinearExtrusionOp.cpp
@@ -0,0 +1,111 @@
+// Copyright Epic Games, Inc. All Rights Reserved.
+
+#include "PolyModelingOps/LinearExtrusionOp.h"
+#include "Operations/OffsetMeshRegion.h"
+#include "Selections/MeshConnectedComponents.h"
+#include "DynamicMesh/MeshNormals.h"
+#include "Util/ProgressCancel.h"
+
+using namespace UE::Geometry;
+
+
+void FLinearExtrusionOp::CalculateResult(FProgressCancel* Progress)
+{
+	bool bMeshValid = OriginalMeshShared->AccessSharedObject([&](const FDynamicMesh3& Mesh)
+	{
+		*ResultMesh = Mesh;
+	});
+
+	if (bMeshValid)
+	{
+		CalculateResultInPlace(*ResultMesh, Progress);
+	}
+}
+
+
+
+bool FLinearExtrusionOp::CalculateResultInPlace(FDynamicMesh3& EditMesh, FProgressCancel* Progress)
+{
+	if (Progress && Progress->Cancelled())
+	{
+		return false;
+	}
+
+	FOffsetMeshRegion Extruder(&EditMesh);
+
+	Extruder.Triangles = TriangleSelection;
+	Extruder.NumSubdivisions = NumSubdivisions;
+
+	if ( bInferGroupsFromNeighbours )
+	{
+		Extruder.LoopEdgesShouldHaveSameGroup = [this, &EditMesh](int32 EdgeID1, int32 EdgeID2) 
+		{
+			return FOffsetMeshRegion::EdgesSeparateSameGroupsAndAreColinearAtBorder(
+				&EditMesh, EdgeID1, EdgeID2, bUseColinearityForSettingBorderGroups);
+		};
+	}
+	else
+	{
+		Extruder.LoopEdgesShouldHaveSameGroup = [this, &EditMesh](int32 EdgeID1, int32 EdgeID2) { return true; };
+	}
+	Extruder.bGroupPerSubdivision = bNewGroupPerSubdivision;
+	Extruder.bSingleGroupPerArea = !bRemapExtrudeGroups;
+	Extruder.CreaseAngleThresholdDeg = CreaseAngleThresholdDeg;
+
+	Extruder.UVScaleFactor = UVScaleFactor;
+	Extruder.bUVIslandPerGroup = bUVIslandPerGroup;
+
+	Extruder.SetMaterialID = SetMaterialID;
+	Extruder.bInferMaterialID = bInferMaterialID;
+
+	Extruder.ExtrusionVectorType = FOffsetMeshRegion::EVertexExtrusionVectorType::Zero;
+	Extruder.OffsetPositionFunc = [this](const FVector3d& Pos, const FVector3d& VertexVector, int32 VertexID)
+	{
+		FVector3d RelativePos = StartFrame.ToFramePoint(Pos);
+		if (this->RegionModifierMode == ESelectionShapeModifierMode::FlattenToPlane)
+		{
+			RelativePos.Z = 0;
+		}
+		RelativePos *= Scale;
+		FVector3d NewPos = ToFrame.FromFramePoint(RelativePos);
+
+		if (this->RegionModifierMode == ESelectionShapeModifierMode::RaycastToPlane)
+		{
+			FVector3d Direction = Normalized(ToFrame.Origin - StartFrame.Origin);
+			if (ToFrame.RayPlaneIntersection(Pos, Direction, 2, NewPos) == false)
+			{
+				if (ToFrame.RayPlaneIntersection(Pos, -Direction, 2, NewPos) == false)
+				{
+					NewPos = Pos;
+				}
+			}
+			FRay3d Ray(Pos, Direction);
+			double RayT = Ray.GetParameter(NewPos);
+			if (FMathd::Abs(RayT) > RaycastMaxDistance)
+			{
+				NewPos = Ray.PointAt(RaycastMaxDistance * FMathd::Sign(RayT));
+			}
+		}
+
+		return NewPos;
+	};
+
+	Extruder.bIsPositiveOffset = (ToFrame.Origin - StartFrame.Origin).Dot(StartFrame.Z()) > 0;
+	Extruder.bOffsetFullComponentsAsSolids = bShellsToSolids;
+
+	Extruder.Apply();
+
+	if (Progress && Progress->Cancelled())
+	{
+		return false;
+	}
+
+	if (bRecomputeNormals)
+	{
+		FMeshNormals::RecomputeOverlayTriNormals(EditMesh, Extruder.AllModifiedAndNewTriangles);
+	}
+
+
+	return true;
+}
+
diff --git a/Engine/Plugins/Runtime/MeshModelingToolset/Source/ModelingOperators/Public/PolyModelingOps/LinearExtrusionOp.h b/Engine/Plugins/Runtime/MeshModelingToolset/Source/ModelingOperators/Public/PolyModelingOps/LinearExtrusionOp.h
new file mode 100644
--- /dev/null
+++ b/Engine/Plugins/Runtime/MeshModelingToolset/Source/ModelingOperators/Public/PolyModelingOps/LinearExtrusionOp.h
@@ -0,0 +1,126 @@
+// Copyright Epic Games, Inc. All Rights Reserved.
+
+#pragma once
+
+#include "CoreMinimal.h"
+#include "DynamicMesh/MeshSharingUtil.h"
+#include "ModelingOperators.h"
+
+
+namespace UE {
+namespace Geometry {
+
+class FDynamicMesh3;
+
+/**
+ * LinearExtrusionOp extrudes a set of triangles of a mesh in a linear direction,
+ * with various options for how to process that extrusion. 
+ */
+class MODELINGOPERATORS_API FLinearExtrusionOp : public FDynamicMeshOperator
+{
+public:
+	virtual ~FLinearExtrusionOp() {}
+
+	//
+	// Inputs
+	//
+
+	/** Handle to the source DynamicMesh */
+	TSharedPtr<FSharedConstDynamicMesh3> OriginalMeshShared;
+
+	/** Triangles of the OriginalMesh that should be extruded */
+	TArray<int32> TriangleSelection;
+
+	/** Initial frame of the Extrusion */
+	FFrame3d StartFrame = FFrame3d();
+	/** End frame of the Extrusion */
+	FFrame3d ToFrame = FFrame3d();
+	/** End scale of the Extrusion, relative to initial scale (1,1,1) */
+	FVector3d Scale = FVector3d::One();
+
+	//
+	// Parameters
+	//
+
+	/** Modifier applied to the selected triangles during extrusion */
+	enum class ESelectionShapeModifierMode
+	{
+		/** No modification, keep the same shape */
+		None = 0,
+		/** Project the selection vertices to the XY plane of StartFrame before extruding */
+		FlattenToPlane = 1,
+		/** Cast a ray for each selection vertex against plane defined by ToFrame's Orgin/Z-Axis and use the hit position */
+		RaycastToPlane = 2
+	};
+	/** Current modifier applied to selection region */
+	ESelectionShapeModifierMode RegionModifierMode = ESelectionShapeModifierMode::None;
+
+	/** In RaycastToPlane mode, if raycast distance is > this value, fallback to no shape change */
+	double RaycastMaxDistance = 9999.0;
+
+	/** If bShellsToSolids is true, patches surrounded by mesh boundaries will be extruded into closed solids instead of open shells */
+	bool bShellsToSolids = true;
+
+	/** Number of subdivisions along the extrusion "tubes" */
+	int32 NumSubdivisions = 0;
+
+	/** 
+	 * Assign new groups to regions of the extrusion "tube" based on the existing adjacent groups before extrusion. Each (disconnected) 
+	 * unique pair of group IDs gets a new unique group ID on the extrusion. If disabled, the entire extrusion tube will be assigned
+	 * a single new group ID
+	 */
+	bool bInferGroupsFromNeighbours = true;
+	
+	/** If true, assign a separate group for each subdivision, otherwise all subdivisions have the same group */
+	bool bNewGroupPerSubdivision = true;
+	
+	/** assign new GroupIDs to each input GroupID in the extrude area. If false, disconnected extrude areas get single new group IDs */
+	bool bRemapExtrudeGroups = true;
+
+	/** if the opening angle at an edge of the extrusion tube is > this threshold, "cut" any group that crosses the edge into two groups */
+	double CreaseAngleThresholdDeg = 180.0;
+
+	/** Scaling applied to the default UV values */
+	double UVScaleFactor = 1.0f;
+
+	/** If true, assign a new UV island for each new group, otherwise each disconnected extrusion tube is unwrapped to a single UV island*/
+	bool bUVIslandPerGroup = true;
+
+	/** Constant Material ID to set along the extrusion tube */
+	int SetMaterialID = 0;
+
+	/** If true, Material IDs around the border of the extrude area are propagated "down" the extrusion tube */
+	bool bInferMaterialID = true;
+
+	/** Determine whether normals should be recomputed */
+	bool bRecomputeNormals = true;
+
+	// (this parameter may not be relevant and perhaps should be removed...)
+	// Used when setting groups for the sides when the extrusion includes a mesh border. When true,
+	// different groups are added based on colinearity of the border edges.
+	bool bUseColinearityForSettingBorderGroups = true;
+
+
+	//
+	// Outputs
+	//
+
+
+
+
+	/**
+	 * Compute result applied directly to EditMesh
+	 * @return true on success, false on error or if Progress is cancelled
+	 */
+	virtual bool CalculateResultInPlace(FDynamicMesh3& EditMesh, FProgressCancel* Progress);
+
+
+	//
+	// FDynamicMeshOperator interface 
+	//
+	virtual void CalculateResult(FProgressCancel* Progress) override;
+
+
+};
+
+}} // end UE::Geometry
diff --git a/Engine/Source/Runtime/InteractiveToolsFramework/Public/SingleSelectionTool.h b/Engine/Source/Runtime/InteractiveToolsFramework/Public/SingleSelectionTool.h
--- a/Engine/Source/Runtime/InteractiveToolsFramework/Public/SingleSelectionTool.h
+++ b/Engine/Source/Runtime/InteractiveToolsFramework/Public/SingleSelectionTool.h
@@ -14,11 +14,22 @@ class INTERACTIVETOOLSFRAMEWORK_API USingleSelectionTool : public UInteractiveTo
 {
 GENERATED_BODY()
 public:
-	void SetTarget(UToolTarget* TargetIn)
+	/**
+	 * Set the ToolTarget for the Tool
+	 */
+	virtual void SetTarget(UToolTarget* TargetIn)
 	{
 		Target = TargetIn;
 	}
 
+	/**
+	 * @return the ToolTarget for the Tool
+	 */
+	virtual UToolTarget* GetTarget()
+	{
+		return Target;
+	}
+
 	/**
 	 * @return true if all ToolTargets of this tool are still valid
 	 */
