InteractiveToolsFramework: Adds new helper methods to determine count ActorComponents that match TargetRequirements and an additional test predicate for specialized scenarios. Additionally adds a couple of helper methods to test for Volume components and component classes that support UVs. This is to fix the issues surrounding the UV modeling tools that shouldn't be viable for Volume type targets.

diff --git a/Engine/Plugins/Experimental/MeshModelingToolsetExp/Source/MeshModelingToolsExp/Private/EditUVIslandsTool.cpp b/Engine/Plugins/Experimental/MeshModelingToolsetExp/Source/MeshModelingToolsExp/Private/EditUVIslandsTool.cpp
--- a/Engine/Plugins/Experimental/MeshModelingToolsetExp/Source/MeshModelingToolsExp/Private/EditUVIslandsTool.cpp
+++ b/Engine/Plugins/Experimental/MeshModelingToolsetExp/Source/MeshModelingToolsExp/Private/EditUVIslandsTool.cpp
@@ -13,6 +13,7 @@
 #include "BaseBehaviors/SingleClickBehavior.h"
 #include "ToolSetupUtil.h"
 #include "ModelingToolTargetUtil.h"
+#include "ToolTargetManager.h"
 
 #include "Materials/MaterialInstanceDynamic.h"
 
@@ -32,6 +33,12 @@ UMeshSurfacePointTool* UEditUVIslandsToolBuilder::CreateNewTool(const FToolBuild
 	return DeformTool;
 }
 
+bool UEditUVIslandsToolBuilder::CanBuildTool(const FToolBuilderState& SceneState) const
+{
+	return UMeshSurfacePointToolBuilder::CanBuildTool(SceneState) && 		
+		SceneState.TargetManager->CountSelectedAndTargetableWithPredicate(SceneState, GetTargetRequirements(),
+		[](UActorComponent& Component) { return ToolBuilderUtil::ComponentTypeCouldHaveUVs(Component); }) > 0;
+}
 
 /*
 * Tool methods
diff --git a/Engine/Plugins/Experimental/MeshModelingToolsetExp/Source/MeshModelingToolsExp/Private/SeamSculptTool.cpp b/Engine/Plugins/Experimental/MeshModelingToolsetExp/Source/MeshModelingToolsExp/Private/SeamSculptTool.cpp
--- a/Engine/Plugins/Experimental/MeshModelingToolsetExp/Source/MeshModelingToolsExp/Private/SeamSculptTool.cpp
+++ b/Engine/Plugins/Experimental/MeshModelingToolsetExp/Source/MeshModelingToolsExp/Private/SeamSculptTool.cpp
@@ -21,9 +21,10 @@
 #include "DynamicMesh/MeshIndexUtil.h"
 #include "DynamicMesh/DynamicMeshChangeTracker.h"
 #include "DynamicMeshToMeshDescription.h"
+#include "ToolTargetManager.h"
 
-#include "TargetInterfaces/MeshDescriptionCommitter.h"
 #include "TargetInterfaces/PrimitiveComponentBackedTarget.h"
+
 #include "ModelingToolTargetUtil.h"
 
 #include UE_INLINE_GENERATED_CPP_BY_NAME(SeamSculptTool)
@@ -44,6 +45,12 @@ UMeshSurfacePointTool* USeamSculptToolBuilder::CreateNewTool(const FToolBuilderS
 }
 
 
+bool USeamSculptToolBuilder::CanBuildTool(const FToolBuilderState& SceneState) const
+{
+	return UMeshSurfacePointMeshEditingToolBuilder::CanBuildTool(SceneState) && 
+		SceneState.TargetManager->CountSelectedAndTargetableWithPredicate(SceneState, GetTargetRequirements(),
+			[](UActorComponent& Component) { return ToolBuilderUtil::ComponentTypeCouldHaveUVs(Component); }) > 0;
+}
 
 
 
diff --git a/Engine/Plugins/Experimental/MeshModelingToolsetExp/Source/MeshModelingToolsExp/Public/EditUVIslandsTool.h b/Engine/Plugins/Experimental/MeshModelingToolsetExp/Source/MeshModelingToolsExp/Public/EditUVIslandsTool.h
--- a/Engine/Plugins/Experimental/MeshModelingToolsetExp/Source/MeshModelingToolsExp/Public/EditUVIslandsTool.h
+++ b/Engine/Plugins/Experimental/MeshModelingToolsetExp/Source/MeshModelingToolsExp/Public/EditUVIslandsTool.h
@@ -28,6 +28,7 @@ class MESHMODELINGTOOLSEXP_API UEditUVIslandsToolBuilder : public UMeshSurfacePo
 	GENERATED_BODY()
 public:
 	virtual UMeshSurfacePointTool* CreateNewTool(const FToolBuilderState& SceneState) const override;
+	virtual bool CanBuildTool(const FToolBuilderState& SceneState) const override;
 };
 
 
diff --git a/Engine/Plugins/Experimental/MeshModelingToolsetExp/Source/MeshModelingToolsExp/Public/SeamSculptTool.h b/Engine/Plugins/Experimental/MeshModelingToolsetExp/Source/MeshModelingToolsExp/Public/SeamSculptTool.h
--- a/Engine/Plugins/Experimental/MeshModelingToolsetExp/Source/MeshModelingToolsExp/Public/SeamSculptTool.h
+++ b/Engine/Plugins/Experimental/MeshModelingToolsetExp/Source/MeshModelingToolsExp/Public/SeamSculptTool.h
@@ -21,6 +21,7 @@ class MESHMODELINGTOOLSEXP_API USeamSculptToolBuilder : public UMeshSurfacePoint
 
 public:
 	virtual UMeshSurfacePointTool* CreateNewTool(const FToolBuilderState& SceneState) const override;
+	virtual bool CanBuildTool(const FToolBuilderState& SceneState) const override;
 };
 
 
diff --git a/Engine/Plugins/Runtime/MeshModelingToolset/Source/MeshModelingTools/Private/RecomputeUVsTool.cpp b/Engine/Plugins/Runtime/MeshModelingToolset/Source/MeshModelingTools/Private/RecomputeUVsTool.cpp
--- a/Engine/Plugins/Runtime/MeshModelingToolset/Source/MeshModelingTools/Private/RecomputeUVsTool.cpp
+++ b/Engine/Plugins/Runtime/MeshModelingToolset/Source/MeshModelingTools/Private/RecomputeUVsTool.cpp
@@ -8,6 +8,8 @@
 #include "ToolSetupUtil.h"
 #include "ModelingToolTargetUtil.h"
 #include "ParameterizationOps/RecomputeUVsOp.h"
+#include "ToolTargetManager.h"
+
 
 #include UE_INLINE_GENERATED_CPP_BY_NAME(RecomputeUVsTool)
 
@@ -26,6 +28,13 @@ USingleSelectionMeshEditingTool* URecomputeUVsToolBuilder::CreateNewTool(const F
 	return NewTool;
 }
 
+bool URecomputeUVsToolBuilder::CanBuildTool(const FToolBuilderState& SceneState) const
+{
+	return USingleSelectionMeshEditingToolBuilder::CanBuildTool(SceneState) && 
+		SceneState.TargetManager->CountSelectedAndTargetableWithPredicate(SceneState, GetTargetRequirements(),
+			[](UActorComponent& Component) { return ToolBuilderUtil::ComponentTypeCouldHaveUVs(Component); }) > 0;
+}
+
 /*
  * Tool
  */
diff --git a/Engine/Plugins/Runtime/MeshModelingToolset/Source/MeshModelingTools/Private/UVLayoutTool.cpp b/Engine/Plugins/Runtime/MeshModelingToolset/Source/MeshModelingTools/Private/UVLayoutTool.cpp
--- a/Engine/Plugins/Runtime/MeshModelingToolset/Source/MeshModelingTools/Private/UVLayoutTool.cpp
+++ b/Engine/Plugins/Runtime/MeshModelingToolset/Source/MeshModelingTools/Private/UVLayoutTool.cpp
@@ -15,9 +15,9 @@
 #include "ParameterizationOps/UVLayoutOp.h"
 #include "Properties/UVLayoutProperties.h"
 
-#include "TargetInterfaces/MaterialProvider.h"
-#include "TargetInterfaces/PrimitiveComponentBackedTarget.h"
 #include "ModelingToolTargetUtil.h"
+#include "ToolTargetManager.h"
+
 
 #include UE_INLINE_GENERATED_CPP_BY_NAME(UVLayoutTool)
 
@@ -34,7 +34,12 @@ UMultiSelectionMeshEditingTool* UUVLayoutToolBuilder::CreateNewTool(const FToolB
 	return NewObject<UUVLayoutTool>(SceneState.ToolManager);
 }
 
-
+bool UUVLayoutToolBuilder::CanBuildTool(const FToolBuilderState& SceneState) const
+{
+	return UMultiSelectionMeshEditingToolBuilder::CanBuildTool(SceneState) &&
+		SceneState.TargetManager->CountSelectedAndTargetableWithPredicate(SceneState, GetTargetRequirements(),
+			[](UActorComponent& Component) { return ToolBuilderUtil::ComponentTypeCouldHaveUVs(Component); }) > 0;
+}
 
 /*
  * Tool
diff --git a/Engine/Plugins/Runtime/MeshModelingToolset/Source/MeshModelingTools/Private/UVProjectionTool.cpp b/Engine/Plugins/Runtime/MeshModelingToolset/Source/MeshModelingTools/Private/UVProjectionTool.cpp
--- a/Engine/Plugins/Runtime/MeshModelingToolset/Source/MeshModelingTools/Private/UVProjectionTool.cpp
+++ b/Engine/Plugins/Runtime/MeshModelingToolset/Source/MeshModelingTools/Private/UVProjectionTool.cpp
@@ -19,6 +19,8 @@
 #include "Parameterization/DynamicMeshUVEditor.h"
 #include "Operations/MeshConvexHull.h"
 #include "MinVolumeBox3.h"
+#include "ToolTargetManager.h"
+
 
 #include UE_INLINE_GENERATED_CPP_BY_NAME(UVProjectionTool)
 
@@ -35,6 +37,13 @@ USingleSelectionMeshEditingTool* UUVProjectionToolBuilder::CreateNewTool(const F
 	return NewObject<UUVProjectionTool>(SceneState.ToolManager);
 }
 
+bool UUVProjectionToolBuilder::CanBuildTool(const FToolBuilderState& SceneState) const
+{
+	return USingleSelectionMeshEditingToolBuilder::CanBuildTool(SceneState) && 
+		SceneState.TargetManager->CountSelectedAndTargetableWithPredicate(SceneState, GetTargetRequirements(),
+			[](UActorComponent& Component) { return ToolBuilderUtil::ComponentTypeCouldHaveUVs(Component); }) > 0;
+}
+
 /*
  * Propert Sets
  */
diff --git a/Engine/Plugins/Runtime/MeshModelingToolset/Source/MeshModelingTools/Public/RecomputeUVsTool.h b/Engine/Plugins/Runtime/MeshModelingToolset/Source/MeshModelingTools/Public/RecomputeUVsTool.h
--- a/Engine/Plugins/Runtime/MeshModelingToolset/Source/MeshModelingTools/Public/RecomputeUVsTool.h
+++ b/Engine/Plugins/Runtime/MeshModelingToolset/Source/MeshModelingTools/Public/RecomputeUVsTool.h
@@ -34,6 +34,7 @@ class MESHMODELINGTOOLS_API URecomputeUVsToolBuilder : public USingleSelectionMe
 
 public:
 	virtual USingleSelectionMeshEditingTool* CreateNewTool(const FToolBuilderState& SceneState) const override;
+	virtual bool CanBuildTool(const FToolBuilderState& SceneState) const override;
 };
 
 
diff --git a/Engine/Plugins/Runtime/MeshModelingToolset/Source/MeshModelingTools/Public/UVLayoutTool.h b/Engine/Plugins/Runtime/MeshModelingToolset/Source/MeshModelingTools/Public/UVLayoutTool.h
--- a/Engine/Plugins/Runtime/MeshModelingToolset/Source/MeshModelingTools/Public/UVLayoutTool.h
+++ b/Engine/Plugins/Runtime/MeshModelingToolset/Source/MeshModelingTools/Public/UVLayoutTool.h
@@ -34,6 +34,7 @@ class MESHMODELINGTOOLS_API UUVLayoutToolBuilder : public UMultiSelectionMeshEdi
 
 public:
 	virtual UMultiSelectionMeshEditingTool* CreateNewTool(const FToolBuilderState& SceneState) const override;
+	virtual bool CanBuildTool(const FToolBuilderState& SceneState) const override;
 };
 
 
diff --git a/Engine/Plugins/Runtime/MeshModelingToolset/Source/MeshModelingTools/Public/UVProjectionTool.h b/Engine/Plugins/Runtime/MeshModelingToolset/Source/MeshModelingTools/Public/UVProjectionTool.h
--- a/Engine/Plugins/Runtime/MeshModelingToolset/Source/MeshModelingTools/Public/UVProjectionTool.h
+++ b/Engine/Plugins/Runtime/MeshModelingToolset/Source/MeshModelingTools/Public/UVProjectionTool.h
@@ -35,6 +35,7 @@ public:
 	virtual USingleSelectionMeshEditingTool* CreateNewTool(const FToolBuilderState& SceneState) const override;
 
 	virtual bool WantsInputSelectionIfAvailable() const override { return true; }
+	virtual bool CanBuildTool(const FToolBuilderState& SceneState) const override;
 };
 
 
diff --git a/Engine/Plugins/Runtime/MeshModelingToolset/Source/MeshModelingToolsEditorOnly/Private/ParameterizeMeshTool.cpp b/Engine/Plugins/Runtime/MeshModelingToolset/Source/MeshModelingToolsEditorOnly/Private/ParameterizeMeshTool.cpp
--- a/Engine/Plugins/Runtime/MeshModelingToolset/Source/MeshModelingToolsEditorOnly/Private/ParameterizeMeshTool.cpp
+++ b/Engine/Plugins/Runtime/MeshModelingToolset/Source/MeshModelingToolsEditorOnly/Private/ParameterizeMeshTool.cpp
@@ -10,6 +10,7 @@
 #include "Properties/ParameterizeMeshProperties.h"
 #include "Polygroups/PolygroupUtil.h"
 #include "Polygroups/PolygroupSet.h"
+#include "ToolTargetManager.h"
 
 #include UE_INLINE_GENERATED_CPP_BY_NAME(ParameterizeMeshTool)
 
@@ -27,6 +28,13 @@ USingleSelectionMeshEditingTool* UParameterizeMeshToolBuilder::CreateNewTool(con
 	return NewTool;
 }
 
+bool UParameterizeMeshToolBuilder::CanBuildTool(const FToolBuilderState& SceneState) const
+{
+	return USingleSelectionMeshEditingToolBuilder::CanBuildTool(SceneState) &&
+		SceneState.TargetManager->CountSelectedAndTargetableWithPredicate(SceneState, GetTargetRequirements(),
+			[](UActorComponent& Component) { return ToolBuilderUtil::ComponentTypeCouldHaveUVs(Component); }) == 1;
+}
+
 /*
  * Tool
  */
diff --git a/Engine/Plugins/Runtime/MeshModelingToolset/Source/MeshModelingToolsEditorOnly/Public/ParameterizeMeshTool.h b/Engine/Plugins/Runtime/MeshModelingToolset/Source/MeshModelingToolsEditorOnly/Public/ParameterizeMeshTool.h
--- a/Engine/Plugins/Runtime/MeshModelingToolset/Source/MeshModelingToolsEditorOnly/Public/ParameterizeMeshTool.h
+++ b/Engine/Plugins/Runtime/MeshModelingToolset/Source/MeshModelingToolsEditorOnly/Public/ParameterizeMeshTool.h
@@ -30,6 +30,7 @@ class MESHMODELINGTOOLSEDITORONLY_API UParameterizeMeshToolBuilder : public USin
 	GENERATED_BODY()
 public:
 	virtual USingleSelectionMeshEditingTool* CreateNewTool(const FToolBuilderState& SceneState) const override;
+	virtual bool CanBuildTool(const FToolBuilderState& SceneState) const override;
 };
 
 
diff --git a/Engine/Source/Runtime/InteractiveToolsFramework/Private/ToolBuilderUtil.cpp b/Engine/Source/Runtime/InteractiveToolsFramework/Private/ToolBuilderUtil.cpp
--- a/Engine/Source/Runtime/InteractiveToolsFramework/Private/ToolBuilderUtil.cpp
+++ b/Engine/Source/Runtime/InteractiveToolsFramework/Private/ToolBuilderUtil.cpp
@@ -6,6 +6,8 @@
 #include "GameFramework/Actor.h"
 #include "Engine/Selection.h"
 #include "Components/StaticMeshComponent.h"
+#include "Components/BrushComponent.h"
+#include "Components/MeshComponent.h"
 
 int ToolBuilderUtil::CountComponents(const FToolBuilderState& InputState, const TFunction<bool(UActorComponent*)>& Predicate)
 {
@@ -129,3 +131,14 @@ TArray<AActor*> ToolBuilderUtil::FindAllActors(const FToolBuilderState& InputSta
 	return InputState.SelectedActors.FilterByPredicate(Predicate);
 }
 
+
+bool ToolBuilderUtil::ComponentTypeCouldHaveUVs(const UActorComponent& Component)
+{
+	return Cast<UMeshComponent>(&Component) != nullptr;
+}
+
+
+bool ToolBuilderUtil::IsVolume(const UActorComponent& Component)
+{
+	return Cast<UBrushComponent>(&Component) != nullptr;
+}
\ No newline at end of file
diff --git a/Engine/Source/Runtime/InteractiveToolsFramework/Private/ToolTargetManager.cpp b/Engine/Source/Runtime/InteractiveToolsFramework/Private/ToolTargetManager.cpp
--- a/Engine/Source/Runtime/InteractiveToolsFramework/Private/ToolTargetManager.cpp
+++ b/Engine/Source/Runtime/InteractiveToolsFramework/Private/ToolTargetManager.cpp
@@ -91,7 +91,15 @@ void UToolTargetManager::EnumerateSelectedAndTargetableComponents(const FToolBui
 }
 
 
-
+int32 UToolTargetManager::CountSelectedAndTargetableWithPredicate(const FToolBuilderState& SceneState,
+	const FToolTargetTypeRequirements& TargetRequirements,
+	TFunctionRef<bool(UActorComponent&)> ComponentPred) const
+{
+	return ToolBuilderUtil::CountComponents(SceneState, [&](UActorComponent* Object)
+	{
+		return CanBuildTarget(Object, TargetRequirements) && ComponentPred(*Object);
+	});
+}
 
 UToolTarget* UToolTargetManager::BuildFirstSelectedTargetable(const FToolBuilderState& SceneState, const FToolTargetTypeRequirements& TargetType)
 {
diff --git a/Engine/Source/Runtime/InteractiveToolsFramework/Public/ToolBuilderUtil.h b/Engine/Source/Runtime/InteractiveToolsFramework/Public/ToolBuilderUtil.h
--- a/Engine/Source/Runtime/InteractiveToolsFramework/Public/ToolBuilderUtil.h
+++ b/Engine/Source/Runtime/InteractiveToolsFramework/Public/ToolBuilderUtil.h
@@ -57,6 +57,14 @@ namespace ToolBuilderUtil
 	/** Find first first available Actor of given type, or return nullptr if not found */
 	template<typename ActorType>
 	ActorType* FindFirstActorOfType(const FToolBuilderState& InputState);
+
+	/** Determine if component support UVs. */
+	INTERACTIVETOOLSFRAMEWORK_API
+	bool ComponentTypeCouldHaveUVs(const UActorComponent& Component);
+
+	/** Determine if component is a volume or not. */
+	INTERACTIVETOOLSFRAMEWORK_API
+	bool IsVolume(const UActorComponent& Component);
 }
 
 /*
diff --git a/Engine/Source/Runtime/InteractiveToolsFramework/Public/ToolTargetManager.h b/Engine/Source/Runtime/InteractiveToolsFramework/Public/ToolTargetManager.h
--- a/Engine/Source/Runtime/InteractiveToolsFramework/Public/ToolTargetManager.h
+++ b/Engine/Source/Runtime/InteractiveToolsFramework/Public/ToolTargetManager.h
@@ -101,7 +101,15 @@ public:
 		const FToolTargetTypeRequirements& TargetRequirements,
 		TFunctionRef<void(UActorComponent*)> ComponentFunc) const;
 
-	
+	/**
+	 * Looks through the currently selected components and actors and counts the number of
+	 * inputs that could be used to create qualifying tool targets with an additional test predicate.
+	 */
+	virtual int32 CountSelectedAndTargetableWithPredicate(const FToolBuilderState& SceneState,
+		const FToolTargetTypeRequirements& TargetRequirements,
+		TFunctionRef<bool(UActorComponent&)> ComponentPred) const;
+
+
 	/**
 	 * Looks through the currently selected components and actors and builds a target out of
 	 * the first encountered element that satisfies the requirements. 
