Modeling Mode Tools: Remove exponential behavior from unbounded sliders in various tools.

Added UI/Clamp meta tags where appropriate to implicitly remove the exponential behavior and added Delta/LinearDeltaSensitivity tags when values shouldn't be bounded but don't necessarily depend on the scale of the working mesh.

In the case of values that shouldn't be bounded and do change depending on the mesh scale, the exponential behavior remains.

diff --git a/Engine/Plugins/Experimental/MeshModelingToolsetExp/Source/MeshModelingToolsEditorOnlyExp/Public/SubdividePolyTool.h b/Engine/Plugins/Experimental/MeshModelingToolsetExp/Source/MeshModelingToolsEditorOnlyExp/Public/SubdividePolyTool.h
--- a/Engine/Plugins/Experimental/MeshModelingToolsetExp/Source/MeshModelingToolsEditorOnlyExp/Public/SubdividePolyTool.h
+++ b/Engine/Plugins/Experimental/MeshModelingToolsetExp/Source/MeshModelingToolsEditorOnlyExp/Public/SubdividePolyTool.h
@@ -38,7 +38,7 @@ class MESHMODELINGTOOLSEDITORONLYEXP_API USubdividePolyToolProperties : public U
 
 public:
 
-	UPROPERTY(EditAnywhere, Category=Settings, meta = (UIMin = "1", ClampMin = "1"))
+	UPROPERTY(EditAnywhere, Category=Settings, meta = (UIMin = "1", ClampMin = "1", Delta = 1, LinearDeltaSensitivity = 50))
 	int SubdivisionLevel = 3;
 
 	// Controls whether the user can select Catmull-Clark or is forced to use Loop
diff --git a/Engine/Plugins/Experimental/MeshModelingToolsetExp/Source/MeshModelingToolsExp/Public/CubeGridTool.h b/Engine/Plugins/Experimental/MeshModelingToolsetExp/Source/MeshModelingToolsExp/Public/CubeGridTool.h
--- a/Engine/Plugins/Experimental/MeshModelingToolsetExp/Source/MeshModelingToolsExp/Public/CubeGridTool.h
+++ b/Engine/Plugins/Experimental/MeshModelingToolsetExp/Source/MeshModelingToolsExp/Public/CubeGridTool.h
@@ -113,7 +113,7 @@ public:
 
 	/** How many blocks each push/pull invocation will do at a time.*/
 	UPROPERTY(EditAnywhere, Category = Options, meta = (
-		UIMin = "0", ClampMin = "0"))
+		UIMin = "0", ClampMin = "0", Delta = 1, LinearDeltaSensitivity = 50))
 	int32 BlocksPerStep = 1;
 
 	/** 
diff --git a/Engine/Plugins/Experimental/MeshModelingToolsetExp/Source/MeshModelingToolsExp/Public/DisplaceMeshTool.h b/Engine/Plugins/Experimental/MeshModelingToolsetExp/Source/MeshModelingToolsExp/Public/DisplaceMeshTool.h
--- a/Engine/Plugins/Experimental/MeshModelingToolsetExp/Source/MeshModelingToolsExp/Public/DisplaceMeshTool.h
+++ b/Engine/Plugins/Experimental/MeshModelingToolsetExp/Source/MeshModelingToolsExp/Public/DisplaceMeshTool.h
@@ -93,7 +93,7 @@ public:
 
 	/** Seed for randomization */
 	UPROPERTY(EditAnywhere, Category = Options,
-		meta = (EditCondition = "DisplacementType == EDisplaceMeshToolDisplaceType::RandomNoise || DisplacementType == EDisplaceMeshToolDisplaceType::PerlinNoise"))
+		meta = (EditCondition = "DisplacementType == EDisplaceMeshToolDisplaceType::RandomNoise || DisplacementType == EDisplaceMeshToolDisplaceType::PerlinNoise", Delta = 1, LinearDeltaSensitivity = 50))
 	int RandomSeed = 31337;
 
 	/** Type of the  mesh subdivision. */
@@ -173,7 +173,7 @@ public:
 	/** When sampling from the texture map, how to scale the mesh UV's in the x and y directions. For a mesh with
 	UV's in the range 0 to 1, setting a scale above 1 will result in tiling the texture map, and scaling below
 	1 will result in using only part of the texture map.*/
-	UPROPERTY(EditAnywhere, Category = Options)
+	UPROPERTY(EditAnywhere, Category = Options, meta = (Delta = 0.01, LinearDeltaSensitivity = 1))
 	FVector2D UVScale = FVector2D(1,1);
 
 	/** When sampling from the texture map, how to offset the mesh UV's. This will result in offsetting the
@@ -211,7 +211,7 @@ public:
 
 	/** Unit vector representing the direction to filter along. */
 	UPROPERTY(EditAnywhere, Category = DirectionalFilterOptions, 
-		meta = (EditCondition = "bEnableFilter == true"))
+		meta = (EditCondition = "bEnableFilter == true", Delta = 0.5, LinearDeltaSensitivity = 1))
 	FVector FilterDirection = {0.0f, 0.0f, 1.0f};
 
 	/** Scalar value determining how close to the filter direction the vertex normals must be in order to be displaced.
@@ -280,7 +280,7 @@ public:
 	float SineWavePhaseShift = 0.0f;
 
 	/** Unit vector representing the direction of wave displacement. */
-	UPROPERTY(EditAnywhere, Category = SineWaveOptions)
+	UPROPERTY(EditAnywhere, Category = SineWaveOptions, meta = (Delta = 0.5, LinearDeltaSensitivity = 1))
 	FVector SineWaveDirection = { 0.0f, 0.0f, 1.0f };
 };
 
diff --git a/Engine/Plugins/Experimental/MeshModelingToolsetExp/Source/MeshModelingToolsExp/Public/HoleFillTool.h b/Engine/Plugins/Experimental/MeshModelingToolsetExp/Source/MeshModelingToolsExp/Public/HoleFillTool.h
--- a/Engine/Plugins/Experimental/MeshModelingToolsetExp/Source/MeshModelingToolsExp/Public/HoleFillTool.h
+++ b/Engine/Plugins/Experimental/MeshModelingToolsetExp/Source/MeshModelingToolsExp/Public/HoleFillTool.h
@@ -52,15 +52,15 @@ public:
 
 	/** Number of vertex rings outside of the fill region to allow remeshing */
 	UPROPERTY(EditAnywhere, Category = SmoothHoleFillOptions, 
-		meta = (UIMin = "0", ClampMin = "0", EditCondition = "!bConstrainToHoleInterior"))
+		meta = (UIMin = "0", ClampMin = "0", EditCondition = "!bConstrainToHoleInterior", Delta = 1, LinearDeltaSensitivity = 50))
 	int RemeshingExteriorRegionWidth;
 
 	/** Number of vertex rings outside of the fill region to perform smoothing */
-	UPROPERTY(EditAnywhere, Category = SmoothHoleFillOptions, meta = (UIMin = "0", ClampMin = "0"))
+	UPROPERTY(EditAnywhere, Category = SmoothHoleFillOptions, meta = (UIMin = "0", ClampMin = "0", Delta = 1, LinearDeltaSensitivity = 50))
 	int SmoothingExteriorRegionWidth;
 
 	/** Number of vertex rings away from the fill region boundary to constrain smoothing */
-	UPROPERTY(EditAnywhere, Category = SmoothHoleFillOptions, meta = (UIMin = "0", ClampMin = "0"))
+	UPROPERTY(EditAnywhere, Category = SmoothHoleFillOptions, meta = (UIMin = "0", ClampMin = "0", Delta = 1, LinearDeltaSensitivity = 50))
 	int SmoothingInteriorRegionWidth;
 
 	/** Desired Smoothness. This is not a linear quantity, but larger numbers produce smoother results */
diff --git a/Engine/Plugins/Experimental/MeshModelingToolsetExp/Source/MeshModelingToolsExp/Public/MeshGroupPaintTool.h b/Engine/Plugins/Experimental/MeshModelingToolsetExp/Source/MeshModelingToolsExp/Public/MeshGroupPaintTool.h
--- a/Engine/Plugins/Experimental/MeshModelingToolsetExp/Source/MeshModelingToolsExp/Public/MeshGroupPaintTool.h
+++ b/Engine/Plugins/Experimental/MeshModelingToolsetExp/Source/MeshModelingToolsExp/Public/MeshGroupPaintTool.h
@@ -142,7 +142,7 @@ public:
 	bool bHitBackFaces = true;
 
 	/** The group that will be assigned to triangles */
-	UPROPERTY(EditAnywhere, Category = ActionType, meta = (UIMin = 0, ClampMin = 0) )
+	UPROPERTY(EditAnywhere, Category = ActionType, meta = (UIMin = 0, ClampMin = 0, Delta = 1, LinearDeltaSensitivity = 50))
 	int32 SetGroup = 1;
 
 	/** If true, only triangles with no group assigned will be painted */
@@ -150,7 +150,7 @@ public:
 	bool bOnlySetUngrouped = false;
 
 	/** Group to set as Erased value */
-	UPROPERTY(EditAnywhere, Category = ActionType, meta = (UIMin = 0, ClampMin = 0))
+	UPROPERTY(EditAnywhere, Category = ActionType, meta = (UIMin = 0, ClampMin = 0, Delta = 1, LinearDeltaSensitivity = 50))
 	int32 EraseGroup = 0;
 
 	/** When enabled, only the current group configured in the Paint brush is erased */
diff --git a/Engine/Plugins/Experimental/MeshModelingToolsetExp/Source/MeshModelingToolsExp/Public/PatternTool.h b/Engine/Plugins/Experimental/MeshModelingToolsetExp/Source/MeshModelingToolsExp/Public/PatternTool.h
--- a/Engine/Plugins/Experimental/MeshModelingToolsetExp/Source/MeshModelingToolsExp/Public/PatternTool.h
+++ b/Engine/Plugins/Experimental/MeshModelingToolsetExp/Source/MeshModelingToolsExp/Public/PatternTool.h
@@ -105,7 +105,7 @@ class MESHMODELINGTOOLSEXP_API UPatternToolSettings : public UInteractiveToolPro
 	GENERATED_BODY()
 public:
 	/** The seed used to introduce random transform variations when enabled */
-	UPROPERTY(EditAnywhere, Category = General, meta = (NoResetToDefault))
+	UPROPERTY(EditAnywhere, Category = General, meta = (NoResetToDefault, Delta = 1, LinearDeltaSensitivity = 50))
 	int32 Seed = FMath::Rand();
 
 	/** Whether or not the pattern items should be projected along the negative Z axis of the plane mechanic */
@@ -113,7 +113,7 @@ public:
 	bool bProjectElementsDown = false;
 
 	/** How much each pattern item should be moved along the negative Z axis of the plane mechanic if Project Elements Down is enabled */
-	UPROPERTY(EditAnywhere, Category = General, meta = (EditCondition = "bProjectElementsDown == true", EditConditionHides))
+	UPROPERTY(EditAnywhere, Category = General, meta = (EditCondition = "bProjectElementsDown == true", EditConditionHides, Delta = 0.1, LinearDeltaSensitivity = 1))
 	float ProjectionOffset = 0.0f;
 
 	/** Hide the source meshes when enabled */
@@ -150,7 +150,7 @@ public:
 	bool bIgnoreTransforms = false;
 	
 	/** Value added to the all pattern elements' bounding boxes for adjusting the behavior of packed spacing mode manually */
-	UPROPERTY(EditAnywhere, Category = BoundingBox)
+	UPROPERTY(EditAnywhere, Category = BoundingBox, meta = (Delta = 0.1, LinearDeltaSensitivity = 1))
 	float Adjustment = 0.0f;
 
 	/** If true, the bounding boxes of each element are rendered in green and the combined bounding box of all source elements is rendered in red */
@@ -280,7 +280,7 @@ public:
 	UPROPERTY(EditAnywhere, Category = RadialPattern, meta = (Units = "Degrees", ClampMin = -180, ClampMax = 180))
 	double AngleShift = 0.0;
 
-	/** If true, Pattern elements are rotated to align with the Circle tangen */
+	/** If true, Pattern elements are rotated to align with the Circle tangent */
 	UPROPERTY(EditAnywhere, Category = RadialPattern)
 	bool bOriented = true;
 };
@@ -305,15 +305,15 @@ public:
 	bool bJitter = false;
 
 	/** Rotation applied to all Pattern Elements, or to first Pattern Element for Interpolated rotation */
-	UPROPERTY(EditAnywhere, Category = Rotation)
+	UPROPERTY(EditAnywhere, Category = Rotation, meta = (UIMin = -360, UIMax = 360))
 	FRotator StartRotation = FRotator::ZeroRotator;
 
 	/** Rotation applied to last Pattern Elements for Interpolated rotation */
-	UPROPERTY(EditAnywhere, Category = Rotation, meta = (EditCondition = "bInterpolate"))
+	UPROPERTY(EditAnywhere, Category = Rotation, meta = (EditCondition = "bInterpolate", UIMin = -360, UIMax = 360))
 	FRotator EndRotation = FRotator::ZeroRotator;
 
 	/** Upper bound of the range which is sampled to randomly rotate each Pattern Element if Jitter is true */
-	UPROPERTY(EditAnywhere, Category = Rotation, meta = (ClampMin = 0, EditCondition = "bJitter"))
+	UPROPERTY(EditAnywhere, Category = Rotation, meta = (ClampMin = 0, EditCondition = "bJitter", UIMin = -360, UIMax = 360))
 	FRotator Jitter = FRotator::ZeroRotator;
 };
 
@@ -371,15 +371,15 @@ public:
 	bool bJitter = false;
 	
 	/** Scale applied to all Pattern Elements, or to first Pattern Element for Interpolated scale */
-	UPROPERTY(EditAnywhere, Category = Scale, meta = (ClampMin = 0.001))
+	UPROPERTY(EditAnywhere, Category = Scale, meta = (ClampMin = 0.001, Delta = 0.01, LinearDeltaSensitivity = 1))
 	FVector StartScale = FVector::OneVector;
 	
 	/** Scale applied to last Pattern Element for Interpolated scale */
-	UPROPERTY(EditAnywhere, Category = Scale, meta = (ClampMin = 0.001, EditCondition = "bInterpolate"))
+	UPROPERTY(EditAnywhere, Category = Scale, meta = (ClampMin = 0.001, EditCondition = "bInterpolate", Delta = 0.01, LinearDeltaSensitivity = 1))
 	FVector EndScale = FVector::OneVector;
 
 	/** Upper bound of the range which is sampled to randomly scale each Pattern Element if Jitter is true */
-	UPROPERTY(EditAnywhere, Category = Scale, meta = (ClampMin = 0.001, EditCondition = "bJitter"))
+	UPROPERTY(EditAnywhere, Category = Scale, meta = (ClampMin = 0.001, EditCondition = "bJitter", Delta = 0.01, LinearDeltaSensitivity = 1))
 	FVector Jitter = FVector(MinScale);
 };
 
diff --git a/Engine/Plugins/Experimental/MeshModelingToolsetExp/Source/MeshModelingToolsExp/Public/PlaneCutTool.h b/Engine/Plugins/Experimental/MeshModelingToolsetExp/Source/MeshModelingToolsExp/Public/PlaneCutTool.h
--- a/Engine/Plugins/Experimental/MeshModelingToolsetExp/Source/MeshModelingToolsExp/Public/PlaneCutTool.h
+++ b/Engine/Plugins/Experimental/MeshModelingToolsetExp/Source/MeshModelingToolsExp/Public/PlaneCutTool.h
@@ -67,7 +67,7 @@ public:
 	bool bKeepBothHalves = false;
 
 	/** If keeping both halves, separate the two pieces by this amount */
-	UPROPERTY(EditAnywhere, Category = Options, meta = (EditCondition = "bKeepBothHalves == true", UIMin = "0", ClampMin = "0") )
+	UPROPERTY(EditAnywhere, Category = Options, meta = (EditCondition = "bKeepBothHalves == true", UIMin = "0", ClampMin = "0", Delta = 0.5, LinearDeltaSensitivity = 1))
 	float SpacingBetweenHalves = 0;
 
 	UPROPERTY(EditAnywhere, Category = Options)
diff --git a/Engine/Plugins/Experimental/MeshModelingToolsetExp/Source/MeshModelingToolsExp/Public/RemeshMeshTool.h b/Engine/Plugins/Experimental/MeshModelingToolsetExp/Source/MeshModelingToolsExp/Public/RemeshMeshTool.h
--- a/Engine/Plugins/Experimental/MeshModelingToolsetExp/Source/MeshModelingToolsExp/Public/RemeshMeshTool.h
+++ b/Engine/Plugins/Experimental/MeshModelingToolsetExp/Source/MeshModelingToolsExp/Public/RemeshMeshTool.h
@@ -45,7 +45,7 @@ public:
 	URemeshMeshToolProperties();
 
 	/** Target triangle count */
-	UPROPERTY(EditAnywhere, Category = Remeshing, meta = (EditCondition = "bUseTargetEdgeLength == false"))
+	UPROPERTY(EditAnywhere, Category = Remeshing, meta = (ClampMin = 0, EditCondition = "bUseTargetEdgeLength == false"))
 	int TargetTriangleCount;
 
 	/** Smoothing type */
diff --git a/Engine/Plugins/Experimental/MeshModelingToolsetExp/Source/MeshModelingToolsExp/Public/RemoveOccludedTrianglesTool.h b/Engine/Plugins/Experimental/MeshModelingToolsetExp/Source/MeshModelingToolsExp/Public/RemoveOccludedTrianglesTool.h
--- a/Engine/Plugins/Experimental/MeshModelingToolsetExp/Source/MeshModelingToolsExp/Public/RemoveOccludedTrianglesTool.h
+++ b/Engine/Plugins/Experimental/MeshModelingToolsetExp/Source/MeshModelingToolsExp/Public/RemoveOccludedTrianglesTool.h
@@ -95,13 +95,13 @@ public:
 	bool bOnlySelfOcclude = false;
 
 	/** Shrink (erode) the boundary of the set of triangles to remove. */
-	UPROPERTY(EditAnywhere, Category = OcclusionCalculation, meta = (UIMin = "0", ClampMin = "0"))
+	UPROPERTY(EditAnywhere, Category = OcclusionCalculation, meta = (UIMin = "0", ClampMin = "0", Delta = 1, LinearDeltaSensitivity = 1))
 	int ShrinkRemoval = 0;
 
-	UPROPERTY(EditAnywhere, Category = RemoveIslands, meta = (UIMin = "0", ClampMin = "0"))
+	UPROPERTY(EditAnywhere, Category = RemoveIslands, meta = (UIMin = "0", ClampMin = "0", Delta = 0.5, LinearDeltaSensitivity = 1))
 	double MinAreaIsland = 0;
 
-	UPROPERTY(EditAnywhere, Category = RemoveIslands, meta = (UIMin = "0", ClampMin = "0"))
+	UPROPERTY(EditAnywhere, Category = RemoveIslands, meta = (UIMin = "0", ClampMin = "0", Delta = 1, LinearDeltaSensitivity = 1))
 	int MinTriCountIsland = 0;
 
 	/** What action to perform on occluded triangles */
diff --git a/Engine/Plugins/Runtime/MeshModelingToolset/Source/MeshModelingTools/Public/Properties/MeshMaterialProperties.h b/Engine/Plugins/Runtime/MeshModelingToolset/Source/MeshModelingTools/Public/Properties/MeshMaterialProperties.h
--- a/Engine/Plugins/Runtime/MeshModelingToolset/Source/MeshModelingTools/Public/Properties/MeshMaterialProperties.h
+++ b/Engine/Plugins/Runtime/MeshModelingToolset/Source/MeshModelingTools/Public/Properties/MeshMaterialProperties.h
@@ -26,7 +26,7 @@ public:
 	TWeakObjectPtr<UMaterialInterface> Material;
 
 	/** Scale factor for generated UVs */
-	UPROPERTY(EditAnywhere, Category = Material, meta = (DisplayName = "UV Scale", HideEditConditionToggle, EditConditionHides, EditCondition = "bShowExtendedOptions"))
+	UPROPERTY(EditAnywhere, Category = Material, meta = (DisplayName = "UV Scale", HideEditConditionToggle, EditConditionHides, EditCondition = "bShowExtendedOptions", Delta = 0.01, LinearDeltaSensitivity = 1))
 	float UVScale = 1.0;
 
 	/** If true, UV scale will be relative to world space. This means objects of different sizes created with the same UV scale have the same average texel size. */
diff --git a/Engine/Plugins/Runtime/MeshModelingToolset/Source/MeshModelingTools/Public/UVProjectionTool.h b/Engine/Plugins/Runtime/MeshModelingToolset/Source/MeshModelingTools/Public/UVProjectionTool.h
--- a/Engine/Plugins/Runtime/MeshModelingToolset/Source/MeshModelingTools/Public/UVProjectionTool.h
+++ b/Engine/Plugins/Runtime/MeshModelingToolset/Source/MeshModelingTools/Public/UVProjectionTool.h
@@ -111,7 +111,7 @@ public:
 	EUVProjectionMethod ProjectionType = EUVProjectionMethod::Plane;
 
 	/** Width, length, and height of the projection shape before rotation */
-	UPROPERTY(EditAnywhere, Category = "UV Projection")
+	UPROPERTY(EditAnywhere, Category = "UV Projection", meta = (Delta = 0.5, LinearDeltaSensitivity = 1))
 	FVector Dimensions = FVector(100.0f, 100.0f, 100.0f);
 
 	/** If true, changes to Dimensions result in all components be changed proportionally */
@@ -155,11 +155,11 @@ public:
 	//
 
 	/** Rotation in degrees applied after computing projection */
-	UPROPERTY(EditAnywhere, Category = "UV Transform")
+	UPROPERTY(EditAnywhere, Category = "UV Transform", meta = (ClampMin = -360, ClampMax = 360))
 	float Rotation = 0.0;
 
 	/** Scaling applied after computing projection */
-	UPROPERTY(EditAnywhere, Category = "UV Transform")
+	UPROPERTY(EditAnywhere, Category = "UV Transform", meta = (Delta = 0.01, LinearDeltaSensitivity = 1))
 	FVector2D Scale = FVector2D::UnitVector;
 
 	/** Translation applied after computing projection */
diff --git a/Engine/Plugins/Runtime/MeshModelingToolset/Source/ModelingComponents/Public/Drawing/UVLayoutPreview.h b/Engine/Plugins/Runtime/MeshModelingToolset/Source/ModelingComponents/Public/Drawing/UVLayoutPreview.h
--- a/Engine/Plugins/Runtime/MeshModelingToolset/Source/ModelingComponents/Public/Drawing/UVLayoutPreview.h
+++ b/Engine/Plugins/Runtime/MeshModelingToolset/Source/ModelingComponents/Public/Drawing/UVLayoutPreview.h
@@ -47,7 +47,7 @@ public:
 	float Scale = 1.0;
 
 	/** Offset relative to the center of the selected object */
-	UPROPERTY(EditAnywhere, Category = "Preview UV Layout", meta = (EditCondition = "bEnabled"))
+	UPROPERTY(EditAnywhere, Category = "Preview UV Layout", meta = (EditCondition = "bEnabled", Delta = 0.5, LinearDeltaSensitivity = 1))
 	FVector2D Offset = FVector2D(1.0, 0.5);
 
 	/** Show wireframe mesh in the preview UV layout */
