ModelingTools: Fix PivotAct tool not working on things with Static and Stationary mobility.

diff --git a/Engine/Plugins/Experimental/MeshModelingToolsetExp/Source/MeshModelingToolsEditorOnlyExp/Private/AddPivotActorTool.cpp b/Engine/Plugins/Experimental/MeshModelingToolsetExp/Source/MeshModelingToolsEditorOnlyExp/Private/AddPivotActorTool.cpp
--- a/Engine/Plugins/Experimental/MeshModelingToolsetExp/Source/MeshModelingToolsEditorOnlyExp/Private/AddPivotActorTool.cpp
+++ b/Engine/Plugins/Experimental/MeshModelingToolsetExp/Source/MeshModelingToolsEditorOnlyExp/Private/AddPivotActorTool.cpp
@@ -54,8 +54,6 @@ bool UAddPivotActorToolBuilder::CanBuildTool(const FToolBuilderState& SceneState
 	// 2. If there are multiple actors selected, they need to have a common parent (or no parent),
 	//   because otherwise we will be breaking up the user's hierarchy when we nest everything under
 	//   the empty actor.
-	// 3. All of the actors need to be marked as movable because non-movable items can't be nested
-	//   under a movable one.
 
 	if (SceneState.SelectedActors.Num() == 1 && ExactCast<AActor>(SceneState.SelectedActors[0]) 
 		&& GetPivotBillboard(SceneState.SelectedActors[0]))
@@ -64,19 +62,14 @@ bool UAddPivotActorToolBuilder::CanBuildTool(const FToolBuilderState& SceneState
 	}
 
 	TSet<AActor*> ParentActors;
-	bool bAllActorsMovable = true;
 
 	SceneState.TargetManager->EnumerateSelectedAndTargetableComponents(SceneState, GetTargetRequirements(), 
-		[&bAllActorsMovable, &ParentActors](UActorComponent* Component) {
+		[&ParentActors](UActorComponent* Component) {
 			AActor* Actor = Component->GetOwner();
-			bAllActorsMovable = bAllActorsMovable && Actor->IsRootComponentMovable();
-			if (bAllActorsMovable)
-			{
-				ParentActors.Add(Actor->GetAttachParentActor());
-			}
+			ParentActors.Add(Actor->GetAttachParentActor());
 		});
 
-	return bAllActorsMovable && ParentActors.Num() == 1;
+	return ParentActors.Num() == 1;
 }
 
 UMultiSelectionMeshEditingTool* UAddPivotActorToolBuilder::CreateNewTool(const FToolBuilderState& SceneState) const
@@ -142,7 +135,6 @@ void UAddPivotActorTool::Setup()
 		}
 	}
 
-
 	TransformProperties = NewObject<UPivotActorTransformProperties>();
 	AddToolPropertySource(TransformProperties);
 	TransformProperties->Position = StartTransform.GetTranslation();
@@ -261,6 +253,20 @@ void UAddPivotActorTool::OnShutdown(EToolShutdownType ShutdownType)
 				TransformProxy->GetTransform(),
 				SpawnParams);
 
+			// Set the mobility of the pivot actor. It's usually Movable unless one of the actors is Static,
+			// because you can only nest Static under Static.
+			EComponentMobility::Type MobilityToUse = EComponentMobility::Movable;
+			for (TObjectPtr<UToolTarget> Target : Targets)
+			{
+				AActor* TargetActor = UE::ToolTarget::GetTargetActor(Target);
+				if (TargetActor->GetRootComponent()->Mobility == EComponentMobility::Static)
+				{
+					MobilityToUse = EComponentMobility::Static;
+					break;
+				}
+			}
+			NewActor->GetRootComponent()->SetMobility(MobilityToUse);
+
 			// Grab the first selected target. It will have the same parent as the other ones, so
 			// we'll use it to figure out the empty actor's parent.
 			AActor* TargetActor = UE::ToolTarget::GetTargetActor(Targets[0]);
