Render Capture Baking: Add SubsurfaceColor and Opacity channels

diff --git a/Engine/Plugins/Experimental/GeometryScripting/Source/GeometryScriptingCore/Private/MeshBakeFunctions.cpp b/Engine/Plugins/Experimental/GeometryScripting/Source/GeometryScriptingCore/Private/MeshBakeFunctions.cpp
--- a/Engine/Plugins/Experimental/GeometryScripting/Source/GeometryScriptingCore/Private/MeshBakeFunctions.cpp
+++ b/Engine/Plugins/Experimental/GeometryScripting/Source/GeometryScriptingCore/Private/MeshBakeFunctions.cpp
@@ -827,6 +827,8 @@ namespace GeometryScriptBakeLocals
 			BakeOptions.bBaseColorMap ||
 			BakeOptions.bNormalMap    ||
 			BakeOptions.bEmissiveMap  ||
+			BakeOptions.bOpacityMap ||
+			BakeOptions.bSubsurfaceColorMap ||
 			BakeOptions.bPackedMRSMap ||
 			BakeOptions.bMetallicMap  ||
 			BakeOptions.bRoughnessMap ||
@@ -840,6 +842,8 @@ namespace GeometryScriptBakeLocals
 		Options.NearPlaneDist = BakeOptions.NearPlaneDist;
 		Options.bBakeBaseColor = BakeOptions.bBaseColorMap;
 		Options.bBakeEmissive = BakeOptions.bEmissiveMap;
+		Options.bBakeOpacity = BakeOptions.bOpacityMap;
+		Options.bBakeSubsurfaceColor = BakeOptions.bSubsurfaceColorMap;
 		Options.bBakeNormalMap = BakeOptions.bNormalMap;
 
 		// Enforce the PackedMRS precondition here
@@ -1166,6 +1170,14 @@ FGeometryScriptRenderCaptureTextures UGeometryScriptLibrary_MeshBakeFunctions::B
 	{
 		FTexture2DBuilder::CopyPlatformDataToSourceData(TexturesOut.EmissiveMap, FTexture2DBuilder::ETextureType::EmissiveHDR);
 	}
+	if (Options.bBakeOpacity && TexturesOut.OpacityMap)
+	{
+		FTexture2DBuilder::CopyPlatformDataToSourceData(TexturesOut.OpacityMap, FTexture2DBuilder::ETextureType::ColorLinear);
+	}
+	if (Options.bBakeSubsurfaceColor && TexturesOut.SubsurfaceColorMap)
+	{
+		FTexture2DBuilder::CopyPlatformDataToSourceData(TexturesOut.SubsurfaceColorMap, FTexture2DBuilder::ETextureType::Color);
+	}
 
 	FGeometryScriptRenderCaptureTextures Result;
 
@@ -1176,6 +1188,8 @@ FGeometryScriptRenderCaptureTextures UGeometryScriptLibrary_MeshBakeFunctions::B
 	Result.RoughnessMap = TexturesOut.RoughnessMap;
 	Result.SpecularMap  = TexturesOut.SpecularMap;
 	Result.EmissiveMap  = TexturesOut.EmissiveMap;
+	Result.OpacityMap   = TexturesOut.OpacityMap;
+	Result.SubsurfaceColorMap  = TexturesOut.SubsurfaceColorMap;
 
 	Result.bHasBaseColorMap = (Result.BaseColorMap != nullptr);
 	Result.bHasNormalMap    = (Result.NormalMap    != nullptr);
@@ -1184,6 +1198,8 @@ FGeometryScriptRenderCaptureTextures UGeometryScriptLibrary_MeshBakeFunctions::B
 	Result.bHasRoughnessMap = (Result.RoughnessMap != nullptr);
 	Result.bHasSpecularMap  = (Result.SpecularMap  != nullptr);
 	Result.bHasEmissiveMap  = (Result.EmissiveMap  != nullptr);
+	Result.bHasOpacityMap   = (Result.OpacityMap   != nullptr);
+	Result.bHasSubsurfaceColorMap  = (Result.SubsurfaceColorMap  != nullptr);
 
 	return Result;
 }
diff --git a/Engine/Plugins/Experimental/GeometryScripting/Source/GeometryScriptingCore/Public/GeometryScript/MeshBakeFunctions.h b/Engine/Plugins/Experimental/GeometryScripting/Source/GeometryScriptingCore/Public/GeometryScript/MeshBakeFunctions.h
--- a/Engine/Plugins/Experimental/GeometryScripting/Source/GeometryScriptingCore/Public/GeometryScript/MeshBakeFunctions.h
+++ b/Engine/Plugins/Experimental/GeometryScripting/Source/GeometryScriptingCore/Public/GeometryScript/MeshBakeFunctions.h
@@ -355,6 +355,12 @@ struct GEOMETRYSCRIPTINGCORE_API FGeometryScriptBakeRenderCaptureOptions
 	UPROPERTY(EditAnywhere, BlueprintReadWrite, Category = Options)
 	bool bEmissiveMap = true;
 
+	UPROPERTY(EditAnywhere, BlueprintReadWrite, Category = Options)
+	bool bOpacityMap = true;
+
+	UPROPERTY(EditAnywhere, BlueprintReadWrite, Category = Options)
+	bool bSubsurfaceColorMap = true;
+
 	// These are hidden in the UI right now, we might want to expose them if they turn out to be useful for very large
 	// or very small objects (not tested yet) TODO Figure out if we want to expose these options
 
@@ -411,6 +417,18 @@ struct GEOMETRYSCRIPTINGCORE_API FGeometryScriptRenderCaptureTextures
 
 	UPROPERTY(EditAnywhere, BlueprintReadWrite, Category = Options)
 	bool bHasEmissiveMap = false;
+
+	UPROPERTY(EditAnywhere, BlueprintReadWrite, Category = Options)
+	TObjectPtr<UTexture2D> OpacityMap = nullptr;
+
+	UPROPERTY(EditAnywhere, BlueprintReadWrite, Category = Options)
+	bool bHasOpacityMap = false;
+
+	UPROPERTY(EditAnywhere, BlueprintReadWrite, Category = Options)
+	TObjectPtr<UTexture2D> SubsurfaceColorMap = nullptr;
+
+	UPROPERTY(EditAnywhere, BlueprintReadWrite, Category = Options)
+	bool bHasSubsurfaceColorMap = false;
 };
 
 UCLASS(meta = (ScriptName = "GeometryScript_Bake"))
diff --git a/Engine/Plugins/Experimental/MeshModelingToolsetExp/Source/MeshModelingToolsExp/Private/BakeRenderCaptureTool.cpp b/Engine/Plugins/Experimental/MeshModelingToolsetExp/Source/MeshModelingToolsExp/Private/BakeRenderCaptureTool.cpp
--- a/Engine/Plugins/Experimental/MeshModelingToolsetExp/Source/MeshModelingToolsExp/Private/BakeRenderCaptureTool.cpp
+++ b/Engine/Plugins/Experimental/MeshModelingToolsetExp/Source/MeshModelingToolsExp/Private/BakeRenderCaptureTool.cpp
@@ -22,6 +22,7 @@
 #include "Sampling/MeshImageBakingCache.h"
 #include "Sampling/MeshMapBaker.h"
 #include "Misc/ScopedSlowTask.h"
+#include "Sampling/RenderCaptureMapEvaluator.h"
 
 #include UE_INLINE_GENERATED_CPP_BY_NAME(BakeRenderCaptureTool)
 
@@ -36,6 +37,8 @@ static FString RoughnessTexParamName = TEXT("Roughness");
 static FString MetallicTexParamName = TEXT("Metallic");
 static FString SpecularTexParamName = TEXT("Specular");
 static FString EmissiveTexParamName = TEXT("Emissive");
+static FString OpacityTexParamName  = TEXT("Opacity");
+static FString SubsurfaceColorTexParamName = TEXT("SubsurfaceColor");
 static FString NormalTexParamName = TEXT("NormalMap");
 static FString PackedMRSTexParamName = TEXT("PackedMRS");
 
@@ -56,6 +59,8 @@ MakeRenderCaptureOptions(
 	Options.bBakeBaseColor = RenderCaptureProperties.bBaseColorMap;
 	Options.bBakeNormalMap = RenderCaptureProperties.bNormalMap;
 	Options.bBakeEmissive =  RenderCaptureProperties.bEmissiveMap;
+	Options.bBakeOpacity = RenderCaptureProperties.bOpacityMap;
+	Options.bBakeSubsurfaceColor = RenderCaptureProperties.bSubsurfaceColorMap;
 	Options.bBakeDeviceDepth = RenderCaptureProperties.bDeviceDepthMap;
 	
 	// Enforce the PackedMRS precondition here so we don't have to check it at each usage site.  Note: We don't
@@ -217,6 +222,8 @@ void UBakeRenderCaptureTool::Setup()
 			RenderCaptureProperties->bBaseColorMap ||
 			RenderCaptureProperties->bNormalMap    ||
 			RenderCaptureProperties->bEmissiveMap  ||
+			RenderCaptureProperties->bOpacityMap   ||
+			RenderCaptureProperties->bSubsurfaceColorMap ||
 			RenderCaptureProperties->bPackedMRSMap ||
 			RenderCaptureProperties->bMetallicMap  ||
 			RenderCaptureProperties->bRoughnessMap ||
@@ -238,6 +245,8 @@ void UBakeRenderCaptureTool::Setup()
 	RenderCaptureProperties->WatchProperty(RenderCaptureProperties->bSpecularMap, [this](bool) { OpState |= EBakeOpState::Evaluate; });
 	RenderCaptureProperties->WatchProperty(RenderCaptureProperties->bPackedMRSMap, [this](bool) { OpState |= EBakeOpState::Evaluate; });
 	RenderCaptureProperties->WatchProperty(RenderCaptureProperties->bEmissiveMap, [this](bool) { OpState |= EBakeOpState::Evaluate; });
+	RenderCaptureProperties->WatchProperty(RenderCaptureProperties->bOpacityMap, [this](bool) { OpState |= EBakeOpState::Evaluate; });
+	RenderCaptureProperties->WatchProperty(RenderCaptureProperties->bSubsurfaceColorMap, [this](bool) { OpState |= EBakeOpState::Evaluate; });
 	RenderCaptureProperties->WatchProperty(RenderCaptureProperties->bAntiAliasing, [this](bool) { OpState |= EBakeOpState::Evaluate; });
 	// These are not exposed to the UI, but we watch them anyway because we might change that later
 	RenderCaptureProperties->WatchProperty(RenderCaptureProperties->CaptureFieldOfView, [this](float) { OpState |= EBakeOpState::Evaluate; });
@@ -258,6 +267,8 @@ void UBakeRenderCaptureTool::Setup()
 		Settings->MapPreviewNamesList.Add(RoughnessTexParamName);
 		Settings->MapPreviewNamesList.Add(SpecularTexParamName);
 		Settings->MapPreviewNamesList.Add(EmissiveTexParamName);
+		Settings->MapPreviewNamesList.Add(OpacityTexParamName);
+		Settings->MapPreviewNamesList.Add(SubsurfaceColorTexParamName);
 	}
 
 	ResultSettings = NewObject<UBakeRenderCaptureResults>(this);
@@ -301,6 +312,8 @@ void UBakeRenderCaptureTool::Render(IToolsContextRenderAPI* RenderAPI)
 	const FVector BrightnessColor(Brightness, Brightness, Brightness);
 	PreviewMaterialRC->SetVectorParameterValue(TEXT("Brightness"), BrightnessColor);
 	PreviewMaterialPackedRC->SetVectorParameterValue(TEXT("Brightness"), BrightnessColor);
+	PreviewMaterialRC_Subsurface->SetVectorParameterValue(TEXT("Brightness"), BrightnessColor);
+	PreviewMaterialPackedRC_Subsurface->SetVectorParameterValue(TEXT("Brightness"), BrightnessColor);
 }
 
 
@@ -382,6 +395,18 @@ void UBakeRenderCaptureTool::CreateTextureAssetsRC(UWorld* SourceWorld, UObject*
 		CreateTextureAsset(TexName, FTexture2DBuilder::ETextureType::EmissiveHDR, ResultSettings->EmissiveMap);
 	}
 
+	if (RenderCaptureProperties->bOpacityMap && ResultSettings->OpacityMap != nullptr)
+	{
+		const FString TexName = FString::Printf(TEXT("%s_%s"), *BaseName, *OpacityTexParamName);
+		CreateTextureAsset(TexName, FTexture2DBuilder::ETextureType::ColorLinear, ResultSettings->OpacityMap);
+	}
+
+	if (RenderCaptureProperties->bSubsurfaceColorMap && ResultSettings->SubsurfaceColorMap != nullptr)
+	{
+		const FString TexName = FString::Printf(TEXT("%s_%s"), *BaseName, *SubsurfaceColorTexParamName);
+		CreateTextureAsset(TexName, FTexture2DBuilder::ETextureType::Color, ResultSettings->SubsurfaceColorMap);
+	}
+
 	// We need different code paths based on PackedMRS here because we don't want to uncheck the separate channels
 	// when PackedMRS is enabled to give the user a better UX (they don't have to re-check them after disabling
 	// PackedMRS). In other place we can test the PackedMRS and separate channel booleans in series and avoid the
@@ -438,6 +463,14 @@ bool UBakeRenderCaptureTool::CanAccept() const
 	{
 		return false;
 	}
+	if (RenderCaptureProperties->bOpacityMap && ResultSettings->OpacityMap == nullptr)
+	{
+		return false;
+	}
+	if (RenderCaptureProperties->bSubsurfaceColorMap && ResultSettings->SubsurfaceColorMap == nullptr)
+	{
+		return false;
+	}
 	
 	// We need different code paths based on PackedMRS here because we don't want to uncheck the separate channels
 	// when PackedMRS is enabled to give the user a better UX (they don't have to re-check them after disabling
@@ -500,6 +533,8 @@ void UBakeRenderCaptureTool::OnMapsUpdatedRC(const TUniquePtr<FMeshMapBaker>& Ne
 	ResultSettings->RoughnessMap = TexturesOut.RoughnessMap;
 	ResultSettings->SpecularMap = TexturesOut.SpecularMap;
 	ResultSettings->EmissiveMap = TexturesOut.EmissiveMap;
+	ResultSettings->OpacityMap = TexturesOut.OpacityMap;
+	ResultSettings->SubsurfaceColorMap = TexturesOut.SubsurfaceColorMap;
 
 	GatherAnalytics(*NewResult);
 	UpdateVisualization();
@@ -518,6 +553,24 @@ void UBakeRenderCaptureTool::InitializePreviewMaterials()
 		EmptyEmissiveMap = Builder.GetTexture2D();
 	}
 
+	{
+		FTexture2DBuilder Builder;
+		Builder.Initialize(FTexture2DBuilder::ETextureType::ColorLinear, FImageDimensions(16, 16));
+		// The Opacity texture is passed to the Material's Opacity pin as well as the Opacity Mask pin, so set white
+		// here so we see something when previewing the subsurface material when opacity is not baked
+		Builder.Clear(FColor::White);
+		Builder.Commit(false);
+		EmptyOpacityMap = Builder.GetTexture2D();
+	}
+
+	{
+		FTexture2DBuilder Builder;
+		Builder.Initialize(FTexture2DBuilder::ETextureType::Color, FImageDimensions(16, 16));
+		Builder.Clear(FColor::Black);
+		Builder.Commit(false);
+		EmptySubsurfaceColorMap = Builder.GetTexture2D();
+	}
+
 	{
 		FTexture2DBuilder Builder;
 		Builder.Initialize(FTexture2DBuilder::ETextureType::ColorLinear, FImageDimensions(16, 16));
@@ -548,30 +601,68 @@ void UBakeRenderCaptureTool::InitializePreviewMaterials()
 	}
 	
 	{
-		UMaterial* Material = LoadObject<UMaterial>(nullptr, TEXT("/MeshModelingToolsetExp/Materials/BakeRenderCapturePreviewMaterial"));
+		UMaterial* Material = LoadObject<UMaterial>(nullptr,
+			TEXT("/MeshModelingToolsetExp/Materials/BakeRenderCapturePreviewMaterial"));
 		check(Material);
 		if (Material != nullptr)
 		{
 			PreviewMaterialRC = UMaterialInstanceDynamic::Create(Material, GetToolManager());
-			PreviewMaterialRC->SetTextureParameterValue(TEXT("BaseColor"), EmptyColorMapWhite);
-			PreviewMaterialRC->SetTextureParameterValue(TEXT("Roughness"), EmptyRoughnessMap);
-			PreviewMaterialRC->SetTextureParameterValue(TEXT("Metallic"), EmptyMetallicMap);
-			PreviewMaterialRC->SetTextureParameterValue(TEXT("Specular"), EmptySpecularMap);
-			PreviewMaterialRC->SetTextureParameterValue(TEXT("Emissive"), EmptyEmissiveMap);
-			PreviewMaterialRC->SetTextureParameterValue(TEXT("NormalMap"), EmptyNormalMap);
+			PreviewMaterialRC->SetTextureParameterValue(FName(*BaseColorTexParamName), EmptyColorMapWhite);
+			PreviewMaterialRC->SetTextureParameterValue(FName(*RoughnessTexParamName), EmptyRoughnessMap);
+			PreviewMaterialRC->SetTextureParameterValue(FName(*MetallicTexParamName), EmptyMetallicMap);
+			PreviewMaterialRC->SetTextureParameterValue(FName(*SpecularTexParamName), EmptySpecularMap);
+			PreviewMaterialRC->SetTextureParameterValue(FName(*EmissiveTexParamName), EmptyEmissiveMap);
+			PreviewMaterialRC->SetTextureParameterValue(FName(*NormalTexParamName), EmptyNormalMap);
 		}
 	}
-	
+
 	{
-		UMaterial* Material = LoadObject<UMaterial>(nullptr, TEXT("/MeshModelingToolsetExp/Materials/FullMaterialBakePreviewMaterial_PackedMRS"));
+		UMaterial* Material = LoadObject<UMaterial>(nullptr,
+			TEXT("/MeshModelingToolsetExp/Materials/BakeRenderCapturePreviewSubsurfaceMaterial"));
+		check(Material);
+		if (Material != nullptr)
+		{
+			ensure(Material->GetShadingModels().HasShadingModel(EMaterialShadingModel::MSM_Subsurface));
+			PreviewMaterialRC_Subsurface = UMaterialInstanceDynamic::Create(Material, GetToolManager());
+			PreviewMaterialRC_Subsurface->SetTextureParameterValue(FName(*BaseColorTexParamName), EmptyColorMapWhite);
+			PreviewMaterialRC_Subsurface->SetTextureParameterValue(FName(*RoughnessTexParamName), EmptyRoughnessMap);
+			PreviewMaterialRC_Subsurface->SetTextureParameterValue(FName(*MetallicTexParamName), EmptyMetallicMap);
+			PreviewMaterialRC_Subsurface->SetTextureParameterValue(FName(*SpecularTexParamName), EmptySpecularMap);
+			PreviewMaterialRC_Subsurface->SetTextureParameterValue(FName(*EmissiveTexParamName), EmptyEmissiveMap);
+			PreviewMaterialRC_Subsurface->SetTextureParameterValue(FName(*NormalTexParamName), EmptyNormalMap);
+			PreviewMaterialRC_Subsurface->SetTextureParameterValue(FName(*OpacityTexParamName), EmptyOpacityMap);
+			PreviewMaterialRC_Subsurface->SetTextureParameterValue(FName(*SubsurfaceColorTexParamName), EmptySubsurfaceColorMap);
+		}
+	}
+
+	{
+		UMaterial* Material = LoadObject<UMaterial>(nullptr,
+			TEXT("/MeshModelingToolsetExp/Materials/FullMaterialBakePreviewMaterial_PackedMRS"));
 		check(Material);
 		if (Material != nullptr)
 		{
 			PreviewMaterialPackedRC = UMaterialInstanceDynamic::Create(Material, GetToolManager());
-			PreviewMaterialPackedRC->SetTextureParameterValue(TEXT("BaseColor"), EmptyColorMapWhite);
-			PreviewMaterialPackedRC->SetTextureParameterValue(TEXT("PackedMRS"), EmptyPackedMRSMap);
-			PreviewMaterialPackedRC->SetTextureParameterValue(TEXT("Emissive"), EmptyEmissiveMap);
-			PreviewMaterialPackedRC->SetTextureParameterValue(TEXT("NormalMap"), EmptyNormalMap);
+			PreviewMaterialPackedRC->SetTextureParameterValue(FName(*BaseColorTexParamName), EmptyColorMapWhite);
+			PreviewMaterialPackedRC->SetTextureParameterValue(FName(*PackedMRSTexParamName), EmptyPackedMRSMap);
+			PreviewMaterialPackedRC->SetTextureParameterValue(FName(*EmissiveTexParamName), EmptyEmissiveMap);
+			PreviewMaterialPackedRC->SetTextureParameterValue(FName(*NormalTexParamName), EmptyNormalMap);
+		}
+	}
+
+	{
+		UMaterial* Material = LoadObject<UMaterial>(nullptr,
+			TEXT("/MeshModelingToolsetExp/Materials/FullMaterialBakePreviewSubsurfaceMaterial_PackedMRS"));
+		check(Material);
+		if (Material != nullptr)
+		{
+			ensure(Material->GetShadingModels().HasShadingModel(EMaterialShadingModel::MSM_Subsurface));
+			PreviewMaterialPackedRC_Subsurface = UMaterialInstanceDynamic::Create(Material, GetToolManager());
+			PreviewMaterialPackedRC_Subsurface->SetTextureParameterValue(FName(*BaseColorTexParamName), EmptyColorMapWhite);
+			PreviewMaterialPackedRC_Subsurface->SetTextureParameterValue(FName(*PackedMRSTexParamName), EmptyPackedMRSMap);
+			PreviewMaterialPackedRC_Subsurface->SetTextureParameterValue(FName(*EmissiveTexParamName), EmptyEmissiveMap);
+			PreviewMaterialPackedRC_Subsurface->SetTextureParameterValue(FName(*NormalTexParamName), EmptyNormalMap);
+			PreviewMaterialPackedRC_Subsurface->SetTextureParameterValue(FName(*OpacityTexParamName), EmptyOpacityMap);
+			PreviewMaterialPackedRC_Subsurface->SetTextureParameterValue(FName(*SubsurfaceColorTexParamName), EmptySubsurfaceColorMap);
 		}
 	}
 }
@@ -657,6 +748,8 @@ void UBakeRenderCaptureTool::UpdateResult()
 			RenderCaptureProperties->bSpecularMap       = ComputedRenderCaptureProperties->bSpecularMap;
 			RenderCaptureProperties->bPackedMRSMap      = ComputedRenderCaptureProperties->bPackedMRSMap;
 			RenderCaptureProperties->bEmissiveMap       = ComputedRenderCaptureProperties->bEmissiveMap;
+			RenderCaptureProperties->bOpacityMap        = ComputedRenderCaptureProperties->bOpacityMap;
+			RenderCaptureProperties->bSubsurfaceColorMap = ComputedRenderCaptureProperties->bSubsurfaceColorMap;
 			RenderCaptureProperties->bAntiAliasing      = ComputedRenderCaptureProperties->bAntiAliasing;
 			RenderCaptureProperties->bDeviceDepthMap    = ComputedRenderCaptureProperties->bDeviceDepthMap;
 			RenderCaptureProperties->CaptureFieldOfView = ComputedRenderCaptureProperties->CaptureFieldOfView;
@@ -681,6 +774,8 @@ void UBakeRenderCaptureTool::UpdateResult()
 		ComputedRenderCaptureProperties->bSpecularMap       = RenderCaptureProperties->bSpecularMap;
 		ComputedRenderCaptureProperties->bPackedMRSMap      = RenderCaptureProperties->bPackedMRSMap;
 		ComputedRenderCaptureProperties->bEmissiveMap       = RenderCaptureProperties->bEmissiveMap;
+		ComputedRenderCaptureProperties->bOpacityMap        = RenderCaptureProperties->bOpacityMap;
+		ComputedRenderCaptureProperties->bSubsurfaceColorMap = RenderCaptureProperties->bSubsurfaceColorMap;
 		ComputedRenderCaptureProperties->bAntiAliasing      = RenderCaptureProperties->bAntiAliasing;
 		ComputedRenderCaptureProperties->bDeviceDepthMap    = RenderCaptureProperties->bDeviceDepthMap;
 		ComputedRenderCaptureProperties->CaptureFieldOfView = RenderCaptureProperties->CaptureFieldOfView;
@@ -765,101 +860,81 @@ void UBakeRenderCaptureTool::UpdateVisualization()
 		return;
 	}
 
-	if (ResultSettings->PackedMRSMap)
+	const bool bSubsurfaceMaterial = ResultSettings->SubsurfaceColorMap || ResultSettings->OpacityMap;
+	const bool bPackedMRS = ResultSettings->PackedMRSMap != nullptr;
+
+	// Choose the material
+	TObjectPtr<UMaterialInstanceDynamic> Material = bPackedMRS ? PreviewMaterialPackedRC : PreviewMaterialRC;
+	if (bSubsurfaceMaterial)
 	{
-		TObjectPtr<UMaterialInstanceDynamic> Material = PreviewMaterialPackedRC;
-		PreviewMesh->SetOverrideRenderMaterial(Material);
+		Material = bPackedMRS ? PreviewMaterialPackedRC_Subsurface : PreviewMaterialRC_Subsurface;
+		ensure(Material->GetShadingModels().HasShadingModel(EMaterialShadingModel::MSM_Subsurface));
+		ensure(Material->GetBlendMode() == EBlendMode::BLEND_Masked);
+	}
 
-		if (VisualizationProps->bPreviewAsMaterial)
-		{
-			// We set all textures which were computed in the corresponding texture channels
-			Material->SetTextureParameterValue(FName(BaseColorTexParamName), ResultSettings->BaseColorMap ? ResultSettings->BaseColorMap : EmptyColorMapWhite);
-			Material->SetTextureParameterValue(FName(EmissiveTexParamName),  ResultSettings->EmissiveMap  ? ResultSettings->EmissiveMap  : EmptyEmissiveMap);
-			Material->SetTextureParameterValue(FName(NormalTexParamName),    ResultSettings->NormalMap    ? ResultSettings->NormalMap    : EmptyNormalMap);
-			Material->SetTextureParameterValue(FName(PackedMRSTexParamName), ResultSettings->PackedMRSMap);
-		}
-		else
+	if (VisualizationProps->bPreviewAsMaterial)
+	{
+		const auto TrySetTexture =
+			[Material](const FString& TextureName, TObjectPtr<UTexture2D> Texture, TObjectPtr<UTexture2D> Fallback, bool bMaterialHasTexture = true)
 		{
-			// The BaseColor texture channel will be set according to the selected MapPreview
-			TObjectPtr<UTexture2D> BaseColorMap = EmptyColorMapWhite;
-			if (ResultSettings->BaseColorMap && Settings->MapPreview == BaseColorTexParamName)
+			if (bMaterialHasTexture)
 			{
-				BaseColorMap = ResultSettings->BaseColorMap;
+				Material->SetTextureParameterValue(FName(TextureName), Texture ? Texture : Fallback);
 			}
-			else if (ResultSettings->EmissiveMap && Settings->MapPreview == EmissiveTexParamName)
-			{
-				BaseColorMap = ResultSettings->EmissiveMap;
-			}
-			else if (ResultSettings->NormalMap && Settings->MapPreview == NormalTexParamName)
-			{
-				BaseColorMap = ResultSettings->NormalMap;
-			}
-			else if (ResultSettings->PackedMRSMap && Settings->MapPreview == PackedMRSTexParamName)
-			{
-				BaseColorMap = ResultSettings->PackedMRSMap;
-			}
-			Material->SetTextureParameterValue(FName(BaseColorTexParamName), BaseColorMap);
-			Material->SetTextureParameterValue(FName(EmissiveTexParamName),  EmptyEmissiveMap);
-			Material->SetTextureParameterValue(FName(NormalTexParamName), EmptyNormalMap);
-			Material->SetTextureParameterValue(FName(PackedMRSTexParamName), EmptyPackedMRSMap);
-		}
-
-		Material->SetScalarParameterValue(TEXT("UVChannel"), InputMeshSettings->GetTargetUVLayerIndex());
+		};
+
+		// Set all computed textures or fallback to the empty texture map
+		TrySetTexture(BaseColorTexParamName, ResultSettings->BaseColorMap, EmptyColorMapWhite);
+		TrySetTexture(EmissiveTexParamName,  ResultSettings->EmissiveMap,  EmptyEmissiveMap);
+		TrySetTexture(NormalTexParamName,    ResultSettings->NormalMap,    EmptyNormalMap);
+		TrySetTexture(PackedMRSTexParamName, ResultSettings->PackedMRSMap, EmptyPackedMRSMap,  bPackedMRS);
+		TrySetTexture(RoughnessTexParamName, ResultSettings->RoughnessMap, EmptyRoughnessMap, !bPackedMRS);
+		TrySetTexture(MetallicTexParamName,  ResultSettings->MetallicMap,  EmptyMetallicMap,  !bPackedMRS);
+		TrySetTexture(SpecularTexParamName,  ResultSettings->SpecularMap,  EmptySpecularMap,  !bPackedMRS);
+		TrySetTexture(OpacityTexParamName,   ResultSettings->OpacityMap,   EmptyOpacityMap,    bSubsurfaceMaterial);
+		TrySetTexture(SubsurfaceColorTexParamName, ResultSettings->SubsurfaceColorMap, EmptySubsurfaceColorMap, bSubsurfaceMaterial);
 	}
 	else
 	{
-		TObjectPtr<UMaterialInstanceDynamic> Material = PreviewMaterialRC;
-		PreviewMesh->SetOverrideRenderMaterial(Material);
-
-		if (VisualizationProps->bPreviewAsMaterial)
-		{
-			// We set all textures which were computed in the corresponding texture channels
-			Material->SetTextureParameterValue(FName(BaseColorTexParamName), ResultSettings->BaseColorMap ? ResultSettings->BaseColorMap : EmptyColorMapWhite);
-			Material->SetTextureParameterValue(FName(RoughnessTexParamName), ResultSettings->RoughnessMap ? ResultSettings->RoughnessMap : EmptyRoughnessMap);
-			Material->SetTextureParameterValue(FName(MetallicTexParamName),  ResultSettings->MetallicMap  ? ResultSettings->MetallicMap  : EmptyMetallicMap);
-			Material->SetTextureParameterValue(FName(SpecularTexParamName),  ResultSettings->SpecularMap  ? ResultSettings->SpecularMap  : EmptySpecularMap);
-			Material->SetTextureParameterValue(FName(EmissiveTexParamName),  ResultSettings->EmissiveMap  ? ResultSettings->EmissiveMap  : EmptyEmissiveMap);
-			Material->SetTextureParameterValue(FName(NormalTexParamName),    ResultSettings->NormalMap    ? ResultSettings->NormalMap    : EmptyNormalMap);
-		}
-		else
+		const auto TrySetTexture =
+			[Material, this](const FString& TextureName, TObjectPtr<UTexture2D> Texture, TObjectPtr<UTexture2D> Fallback, bool bMaterialHasTexture = true)
 		{
-			// The BaseColor texture channel will be set according to the selected MapPreview
-			TObjectPtr<UTexture2D> BaseColorMap = EmptyColorMapWhite;
-			if (ResultSettings->BaseColorMap && Settings->MapPreview == BaseColorTexParamName)
+			// Set the BaseColor texture to the MapPreview texture if it exists and use white otherwise
+			if (TextureName == Settings->MapPreview)
 			{
-				BaseColorMap = ResultSettings->BaseColorMap;
-			}
-			else if (ResultSettings->RoughnessMap && Settings->MapPreview == RoughnessTexParamName)
-			{
-				BaseColorMap = ResultSettings->RoughnessMap;
-			}
-			else if (ResultSettings->MetallicMap && Settings->MapPreview == MetallicTexParamName)
-			{
-				BaseColorMap = ResultSettings->MetallicMap;
-			}
-			else if (ResultSettings->SpecularMap && Settings->MapPreview == SpecularTexParamName)
-			{
-				BaseColorMap = ResultSettings->SpecularMap;
-			}
-			else if (ResultSettings->EmissiveMap && Settings->MapPreview == EmissiveTexParamName)
-			{
-				BaseColorMap = ResultSettings->EmissiveMap;
+				if (bMaterialHasTexture && Texture)
+				{
+					Material->SetTextureParameterValue(FName(BaseColorTexParamName), Texture);
+				}
+				else
+				{
+					Material->SetTextureParameterValue(FName(BaseColorTexParamName), EmptyColorMapWhite);
+				}
 			}
-			else if (ResultSettings->NormalMap && Settings->MapPreview == NormalTexParamName)
+
+			// Set the non-BaseColor texture parameters to empty fallback textures
+			if (TextureName != BaseColorTexParamName)
 			{
-				BaseColorMap = ResultSettings->NormalMap;
+				if (bMaterialHasTexture)
+				{
+					Material->SetTextureParameterValue(FName(TextureName), Fallback);
+				}
 			}
-			Material->SetTextureParameterValue(TEXT("BaseColor"), BaseColorMap);
-			
-			Material->SetTextureParameterValue(TEXT("Roughness"), EmptyRoughnessMap);
-			Material->SetTextureParameterValue(TEXT("Metallic"),  EmptyMetallicMap);
-			Material->SetTextureParameterValue(TEXT("Specular"),  EmptySpecularMap);
-			Material->SetTextureParameterValue(TEXT("Emissive"),  EmptyEmissiveMap);
-			Material->SetTextureParameterValue(TEXT("NormalMap"), EmptyNormalMap);
-		}
-		
-		Material->SetScalarParameterValue(TEXT("UVChannel"), InputMeshSettings->GetTargetUVLayerIndex());
+		};
+
+		TrySetTexture(BaseColorTexParamName, ResultSettings->BaseColorMap, EmptyColorMapWhite);
+		TrySetTexture(EmissiveTexParamName,  ResultSettings->EmissiveMap,  EmptyEmissiveMap);
+		TrySetTexture(NormalTexParamName,    ResultSettings->NormalMap,    EmptyNormalMap);
+		TrySetTexture(PackedMRSTexParamName, ResultSettings->PackedMRSMap, EmptyPackedMRSMap,  bPackedMRS);
+		TrySetTexture(RoughnessTexParamName, ResultSettings->RoughnessMap, EmptyRoughnessMap, !bPackedMRS);
+		TrySetTexture(MetallicTexParamName,  ResultSettings->MetallicMap,  EmptyMetallicMap,  !bPackedMRS);
+		TrySetTexture(SpecularTexParamName,  ResultSettings->SpecularMap,  EmptySpecularMap,  !bPackedMRS);
+		TrySetTexture(OpacityTexParamName,   ResultSettings->OpacityMap,   EmptyOpacityMap,    bSubsurfaceMaterial);
+		TrySetTexture(SubsurfaceColorTexParamName, ResultSettings->SubsurfaceColorMap, EmptySubsurfaceColorMap, bSubsurfaceMaterial);
 	}
+
+	Material->SetScalarParameterValue(TEXT("UVChannel"), InputMeshSettings->GetTargetUVLayerIndex());
+	PreviewMesh->SetOverrideRenderMaterial(Material);
 }
 
 
@@ -872,6 +947,8 @@ void UBakeRenderCaptureTool::InvalidateResults()
 	ResultSettings->SpecularMap = nullptr;
 	ResultSettings->PackedMRSMap = nullptr;
 	ResultSettings->EmissiveMap = nullptr;
+	ResultSettings->OpacityMap = nullptr;
+	ResultSettings->SubsurfaceColorMap = nullptr;
 	ResultSettings->NormalMap = nullptr;
 }
 
@@ -914,6 +991,8 @@ void UBakeRenderCaptureTool::RecordAnalytics() const
 	Attributes.Add(FAnalyticsEventAttribute(TEXT("Settings.RenderCapture.SpecularMap.Enabled"), RenderCaptureProperties->bSpecularMap));
 	Attributes.Add(FAnalyticsEventAttribute(TEXT("Settings.RenderCapture.PackedMRSMap.Enabled"), RenderCaptureProperties->bPackedMRSMap));
 	Attributes.Add(FAnalyticsEventAttribute(TEXT("Settings.RenderCapture.EmissiveMap.Enabled"), RenderCaptureProperties->bEmissiveMap));
+	Attributes.Add(FAnalyticsEventAttribute(TEXT("Settings.RenderCapture.OpacityMap.Enabled"), RenderCaptureProperties->bOpacityMap));
+	Attributes.Add(FAnalyticsEventAttribute(TEXT("Settings.RenderCapture.SubsurfaceColorMap.Enabled"), RenderCaptureProperties->bSubsurfaceColorMap));
 	Attributes.Add(FAnalyticsEventAttribute(TEXT("Settings.RenderCapture.CaptureFieldOfView"), RenderCaptureProperties->CaptureFieldOfView));
 	Attributes.Add(FAnalyticsEventAttribute(TEXT("Settings.RenderCapture.NearPlaneDistance"), RenderCaptureProperties->NearPlaneDist));
 
diff --git a/Engine/Plugins/Experimental/MeshModelingToolsetExp/Source/MeshModelingToolsExp/Public/BakeRenderCaptureTool.h b/Engine/Plugins/Experimental/MeshModelingToolsetExp/Source/MeshModelingToolsExp/Public/BakeRenderCaptureTool.h
--- a/Engine/Plugins/Experimental/MeshModelingToolsetExp/Source/MeshModelingToolsExp/Public/BakeRenderCaptureTool.h
+++ b/Engine/Plugins/Experimental/MeshModelingToolsetExp/Source/MeshModelingToolsExp/Public/BakeRenderCaptureTool.h
@@ -58,6 +58,12 @@ public:
 
 	UPROPERTY(VisibleAnywhere, Category = Results, meta = (TransientToolProperty))
 	TObjectPtr<UTexture2D> EmissiveMap;
+
+	UPROPERTY(VisibleAnywhere, Category = Results, meta = (TransientToolProperty))
+	TObjectPtr<UTexture2D> OpacityMap;
+
+	UPROPERTY(VisibleAnywhere, Category = Results, meta = (TransientToolProperty))
+	TObjectPtr<UTexture2D> SubsurfaceColorMap;
 };
 
 
@@ -124,6 +130,14 @@ public:
 	UPROPERTY(Category = RenderCaptureOptions, EditAnywhere, meta = (NoResetToDefault))
 	bool bEmissiveMap = false;
 
+	// Whether to generate a texture for the Emissive property
+	UPROPERTY(Category = RenderCaptureOptions, EditAnywhere, meta = (NoResetToDefault))
+	bool bOpacityMap = false;
+
+	// Whether to generate a texture for the SubsurfaceColor property
+	UPROPERTY(Category = RenderCaptureOptions, EditAnywhere, meta = (NoResetToDefault))
+	bool bSubsurfaceColorMap = false;
+
 	// Whether to use anti-aliasing in the render captures, this may introduce artefacts if pixels at different scene depths get blended
 	UPROPERTY(Category = RenderCaptureOptions, EditAnywhere, AdvancedDisplay)
 	bool bAntiAliasing = false;
@@ -150,6 +164,8 @@ public:
 			&& bSpecularMap == Other.bSpecularMap
 			&& bPackedMRSMap == Other.bPackedMRSMap
 			&& bEmissiveMap == Other.bEmissiveMap
+			&& bOpacityMap == Other.bOpacityMap
+			&& bSubsurfaceColorMap == Other.bSubsurfaceColorMap
 			&& bAntiAliasing == Other.bAntiAliasing
 			&& CaptureFieldOfView == Other.CaptureFieldOfView
 			&& NearPlaneDist == Other.NearPlaneDist
@@ -311,6 +327,10 @@ protected:
 	UPROPERTY()
 	TObjectPtr<UTexture2D> EmptyEmissiveMap;
 	UPROPERTY()
+	TObjectPtr<UTexture2D> EmptyOpacityMap;
+	UPROPERTY()
+	TObjectPtr<UTexture2D> EmptySubsurfaceColorMap;
+	UPROPERTY()
 	TObjectPtr<UTexture2D> EmptyPackedMRSMap;
 	UPROPERTY()
 	TObjectPtr<UTexture2D> EmptyRoughnessMap;
@@ -322,6 +342,10 @@ protected:
 	TObjectPtr<UMaterialInstanceDynamic> PreviewMaterialRC;
 	UPROPERTY()
 	TObjectPtr<UMaterialInstanceDynamic> PreviewMaterialPackedRC;
+	UPROPERTY()
+	TObjectPtr<UMaterialInstanceDynamic> PreviewMaterialRC_Subsurface;
+	UPROPERTY()
+	TObjectPtr<UMaterialInstanceDynamic> PreviewMaterialPackedRC_Subsurface;
 	void InitializePreviewMaterials();
 
 	// TODO We currently need to compute this on the game thread because the implementation has checks for this
diff --git a/Engine/Plugins/Runtime/GeometryProcessing/Source/DynamicMesh/Public/Sampling/RenderCaptureMapEvaluator.h b/Engine/Plugins/Runtime/GeometryProcessing/Source/DynamicMesh/Public/Sampling/RenderCaptureMapEvaluator.h
--- a/Engine/Plugins/Runtime/GeometryProcessing/Source/DynamicMesh/Public/Sampling/RenderCaptureMapEvaluator.h
+++ b/Engine/Plugins/Runtime/GeometryProcessing/Source/DynamicMesh/Public/Sampling/RenderCaptureMapEvaluator.h
@@ -19,7 +19,9 @@ enum class ERenderCaptureChannel
 	Emissive = 16,
 	WorldNormal = 32,
 	DeviceDepth = 64,
-	CombinedMRS = 128
+	CombinedMRS = 128,
+	Opacity = 256,
+	SubsurfaceColor = 512
 };
 
 template <typename DataType>
diff --git a/Engine/Plugins/Runtime/MeshModelingToolset/Source/ModelingComponents/Private/Baking/RenderCaptureFunctions.cpp b/Engine/Plugins/Runtime/MeshModelingToolset/Source/ModelingComponents/Private/Baking/RenderCaptureFunctions.cpp
--- a/Engine/Plugins/Runtime/MeshModelingToolset/Source/ModelingComponents/Private/Baking/RenderCaptureFunctions.cpp
+++ b/Engine/Plugins/Runtime/MeshModelingToolset/Source/ModelingComponents/Private/Baking/RenderCaptureFunctions.cpp
@@ -239,6 +239,8 @@ void UE::Geometry::UpdatePhotoSet(
 	SceneCapture->SetCaptureTypeEnabled(ERenderCaptureType::BaseColor,   Options.bBakeBaseColor);
 	SceneCapture->SetCaptureTypeEnabled(ERenderCaptureType::WorldNormal, Options.bBakeNormalMap);
 	SceneCapture->SetCaptureTypeEnabled(ERenderCaptureType::Emissive,    Options.bBakeEmissive);
+	SceneCapture->SetCaptureTypeEnabled(ERenderCaptureType::Opacity,     Options.bBakeOpacity);
+	SceneCapture->SetCaptureTypeEnabled(ERenderCaptureType::SubsurfaceColor, Options.bBakeSubsurfaceColor);
 	SceneCapture->SetCaptureTypeEnabled(ERenderCaptureType::CombinedMRS, Options.bUsePackedMRS);
 	SceneCapture->SetCaptureTypeEnabled(ERenderCaptureType::Metallic,    Options.bBakeMetallic);
 	SceneCapture->SetCaptureTypeEnabled(ERenderCaptureType::Roughness,   Options.bBakeRoughness);
@@ -253,7 +255,9 @@ void UE::Geometry::UpdatePhotoSet(
 	SceneCapture->SetCaptureConfig(ERenderCaptureType::Roughness,   Config);
 	SceneCapture->SetCaptureConfig(ERenderCaptureType::Specular,    Config);
 	SceneCapture->SetCaptureConfig(ERenderCaptureType::Emissive,    Config);
-	
+	SceneCapture->SetCaptureConfig(ERenderCaptureType::Opacity,     Config);
+	SceneCapture->SetCaptureConfig(ERenderCaptureType::SubsurfaceColor, Config);
+
 	SceneCapture->SetCaptureSceneActors(Actors[0]->GetWorld(), Actors);
 
 	SceneCapture->AddStandardExteriorCapturesFromBoundingBox(
@@ -295,6 +299,12 @@ MakeColorEvaluator(
 	case ERenderCaptureType::CombinedMRS:
 		Evaluator->Channel = ERenderCaptureChannel::CombinedMRS;
 		break;
+	case ERenderCaptureType::Opacity:
+		Evaluator->Channel = ERenderCaptureChannel::Opacity;
+		break;
+	case ERenderCaptureType::SubsurfaceColor:
+		Evaluator->Channel = ERenderCaptureChannel::SubsurfaceColor;
+		break;
 	case ERenderCaptureType::DeviceDepth:
 		Evaluator->Channel = ERenderCaptureChannel::DeviceDepth;
 		break;
@@ -372,8 +382,10 @@ TUniquePtr<FMeshMapBaker> UE::Geometry::MakeRenderCaptureBaker(
 	DefaultColorSample.Specular = InvalidColor.X;
 	DefaultColorSample.Metallic = InvalidColor.X;
 	DefaultColorSample.Emissive = FVector3f(InvalidColor.X, InvalidColor.Y, InvalidColor.Z);
+	DefaultColorSample.Opacity = InvalidColor.X;
+	DefaultColorSample.SubsurfaceColor = FVector3f(InvalidColor.X, InvalidColor.Y, InvalidColor.Z);
 	DefaultColorSample.WorldNormal = FVector4f((DefaultNormal + FVector3f::One()) * .5f, InvalidColor.W);
-	
+
 	// We use 0 here since this corresponds to the infinite far plane value. Also, we don't preview the depth texture
 	DefaultColorSample.DeviceDepth = 0;
 
@@ -411,6 +423,14 @@ TUniquePtr<FMeshMapBaker> UE::Geometry::MakeRenderCaptureBaker(
 	{
 		AddColorEvaluator(MakeColorEvaluator<ERenderCaptureType::Emissive>(DefaultColorSample, SceneCapture));
 	}
+	if (Options.bBakeOpacity)
+	{
+		AddColorEvaluator(MakeColorEvaluator<ERenderCaptureType::Opacity>(DefaultColorSample, SceneCapture));
+	}
+	if (Options.bBakeSubsurfaceColor)
+	{
+		AddColorEvaluator(MakeColorEvaluator<ERenderCaptureType::SubsurfaceColor>(DefaultColorSample, SceneCapture));
+	}
 	if (Options.bBakeNormalMap)
 	{
 		TSharedPtr<FRenderCaptureMapEvaluator<FVector3f>> Evaluator = MakeShared<FRenderCaptureMapEvaluator<FVector3f>>();
@@ -499,6 +519,27 @@ void UE::Geometry::GetTexturesFromRenderCaptureBaker(const TUniquePtr<FMeshMapBa
 						bPopulateSourceData);
 					break;
 
+				case ERenderCaptureChannel::SubsurfaceColor:
+
+					// The SubsurfaceColor GBuffer channel is gamma encoded so:
+					// 1. Dont convert the data in ImageBuilder to sRGB, it is already gamma encoded
+					// 2. Pass the ETextureType::Color enum so the built UTexture2D will have SRGB checked
+					TexturesOut.SubsurfaceColorMap = FTexture2DBuilder::BuildTextureFromImage(
+						*ImageBuilder,
+						FTexture2DBuilder::ETextureType::Color,
+						false,
+						bPopulateSourceData);
+					break;
+
+				case ERenderCaptureChannel::Opacity:
+
+					TexturesOut.OpacityMap = FTexture2DBuilder::BuildTextureFromImage(
+						*ImageBuilder,
+						FTexture2DBuilder::ETextureType::ColorLinear,
+						false,
+						bPopulateSourceData);
+					break;
+
 				case ERenderCaptureChannel::Roughness:
 
 					TexturesOut.RoughnessMap = FTexture2DBuilder::BuildTextureFromImage(
diff --git a/Engine/Plugins/Runtime/MeshModelingToolset/Source/ModelingComponents/Private/Scene/SceneCapturePhotoSet.cpp b/Engine/Plugins/Runtime/MeshModelingToolset/Source/ModelingComponents/Private/Scene/SceneCapturePhotoSet.cpp
--- a/Engine/Plugins/Runtime/MeshModelingToolset/Source/ModelingComponents/Private/Scene/SceneCapturePhotoSet.cpp
+++ b/Engine/Plugins/Runtime/MeshModelingToolset/Source/ModelingComponents/Private/Scene/SceneCapturePhotoSet.cpp
@@ -17,6 +17,8 @@ FSceneCapturePhotoSet::FSceneCapturePhotoSet()
 	SpecularConfig = GetDefaultRenderCaptureConfig(ERenderCaptureType::Specular);
 	EmissiveConfig = GetDefaultRenderCaptureConfig(ERenderCaptureType::Emissive);
 	PackedMRSConfig = GetDefaultRenderCaptureConfig(ERenderCaptureType::CombinedMRS);
+	OpacityConfig = GetDefaultRenderCaptureConfig(ERenderCaptureType::Opacity);
+	SubsurfaceColorConfig = GetDefaultRenderCaptureConfig(ERenderCaptureType::SubsurfaceColor);
 	DeviceDepthConfig = GetDefaultRenderCaptureConfig(ERenderCaptureType::DeviceDepth);
 }
 
@@ -49,6 +51,12 @@ void FSceneCapturePhotoSet::SetCaptureConfig(ERenderCaptureType CaptureType, con
 		case ERenderCaptureType::Emissive:
 			EmissiveConfig = Config;
 			break;
+		case ERenderCaptureType::Opacity:
+			OpacityConfig = Config;
+			break;
+		case ERenderCaptureType::SubsurfaceColor:
+			SubsurfaceColorConfig = Config;
+			break;
 		case ERenderCaptureType::CombinedMRS:
 			PackedMRSConfig = Config;
 			break;
@@ -78,6 +86,10 @@ FRenderCaptureConfig FSceneCapturePhotoSet::GetCaptureConfig(ERenderCaptureType
 			return EmissiveConfig;
 		case ERenderCaptureType::CombinedMRS:
 			return PackedMRSConfig;
+		case ERenderCaptureType::Opacity:
+			return OpacityConfig;
+		case ERenderCaptureType::SubsurfaceColor:
+			return SubsurfaceColorConfig;
 		case ERenderCaptureType::DeviceDepth:
 			return DeviceDepthConfig;
 		default:
@@ -94,6 +106,8 @@ void FSceneCapturePhotoSet::DisableAllCaptureTypes()
 	SetCaptureTypeEnabled(ERenderCaptureType::Metallic, false);
 	SetCaptureTypeEnabled(ERenderCaptureType::Specular, false);
 	SetCaptureTypeEnabled(ERenderCaptureType::Emissive, false);
+	SetCaptureTypeEnabled(ERenderCaptureType::Opacity, false);
+	SetCaptureTypeEnabled(ERenderCaptureType::SubsurfaceColor, false);
 	SetCaptureTypeEnabled(ERenderCaptureType::CombinedMRS, false);
 	SetCaptureTypeEnabled(ERenderCaptureType::DeviceDepth, false);
 }
@@ -124,6 +138,12 @@ void FSceneCapturePhotoSet::SetCaptureTypeEnabled(ERenderCaptureType CaptureType
 	case ERenderCaptureType::CombinedMRS:
 		bEnablePackedMRS = bEnabled;
 		break;
+	case ERenderCaptureType::Opacity:
+		bEnableOpacity = bEnabled;
+		break;
+	case ERenderCaptureType::SubsurfaceColor:
+		bEnableSubsurfaceColor = bEnabled;
+		break;
 	case ERenderCaptureType::DeviceDepth:
 		bEnableDeviceDepth = bEnabled;
 		break;
@@ -154,6 +174,12 @@ void FSceneCapturePhotoSet::SetCaptureTypeEnabled(ERenderCaptureType CaptureType
 		case ERenderCaptureType::Emissive:
 			EmissivePhotoSet.Empty();
 			break;
+		case ERenderCaptureType::Opacity:
+			OpacityPhotoSet.Empty();
+			break;
+		case ERenderCaptureType::SubsurfaceColor:
+			SubsurfaceColorPhotoSet.Empty();
+			break;
 		case ERenderCaptureType::CombinedMRS:
 			PackedMRSPhotoSet.Empty();
 			break;
@@ -173,6 +199,8 @@ void FSceneCapturePhotoSet::SetCaptureTypeEnabled(ERenderCaptureType CaptureType
 		!bEnablePackedMRS &&
 		!bEnableWorldNormal &&
 		!bEnableEmissive &&
+		!bEnableOpacity &&
+		!bEnableSubsurfaceColor &&
 		!bEnableDeviceDepth)
 	{
 		PhotoSetParams.Empty();
@@ -195,6 +223,10 @@ bool FSceneCapturePhotoSet::GetCaptureTypeEnabled(ERenderCaptureType CaptureType
 			return bEnableSpecular;
 		case ERenderCaptureType::Emissive:
 			return bEnableEmissive;
+		case ERenderCaptureType::Opacity:
+			return bEnableOpacity;
+		case ERenderCaptureType::SubsurfaceColor:
+			return bEnableSubsurfaceColor;
 		case ERenderCaptureType::CombinedMRS:
 			return bEnablePackedMRS;
 		case ERenderCaptureType::DeviceDepth:
@@ -408,6 +440,14 @@ void FSceneCapturePhotoSet::AddExteriorCaptures(
 		{
 			CaptureImageTypeFunc_3f(ERenderCaptureType::Emissive, EmissivePhotoSet);
 		}
+		if (bEnableOpacity)
+		{
+			CaptureImageTypeFunc_1f(ERenderCaptureType::Opacity, OpacityPhotoSet);
+		}
+		if (bEnableSubsurfaceColor)
+		{
+			CaptureImageTypeFunc_3f(ERenderCaptureType::SubsurfaceColor, SubsurfaceColorPhotoSet);
+		}
 
 		// AddExteriorCaptures can be called on an FSceneCapturePhotoSet which already has some capture types computed
 		// and in this case we should not modify the existing photo sets/cached photo set parameters.
@@ -440,6 +480,8 @@ FSceneCapturePhotoSet::FSceneSample::FSceneSample()
 	Specular = 0.0f;
 	Metallic = 0.0f;
 	Emissive = FVector3f(0, 0, 0);
+	Opacity = 0.0f;
+	SubsurfaceColor = FVector3f(0, 0, 0);
 	WorldNormal = FVector3f(0, 0, 1);
 	DeviceDepth = 0.0f;
 }
@@ -460,6 +502,10 @@ FVector3f FSceneCapturePhotoSet::FSceneSample::GetValue3f(ERenderCaptureType Cap
 		return Specular * FVector3f::One();
 	case ERenderCaptureType::Emissive:
 		return Emissive;
+	case ERenderCaptureType::Opacity:
+		return Opacity * FVector3f::One();
+	case ERenderCaptureType::SubsurfaceColor:
+		return SubsurfaceColor;
 	case ERenderCaptureType::DeviceDepth:
 		return DeviceDepth * FVector3f::One();
 	default:
@@ -486,6 +532,10 @@ FVector4f FSceneCapturePhotoSet::FSceneSample::GetValue4f(ERenderCaptureType Cap
 		return FVector4f(Metallic, Roughness, Specular, 1.0f);
 	case ERenderCaptureType::Emissive:
 		return FVector4f(Emissive.X, Emissive.Y, Emissive.Z, 1.0f);
+	case ERenderCaptureType::Opacity:
+		return FVector4f(Opacity, Opacity, Opacity, 1.0f);
+	case ERenderCaptureType::SubsurfaceColor:
+		return FVector4f(SubsurfaceColor.X, SubsurfaceColor.Y, SubsurfaceColor.Z, 1.0f);
 	case ERenderCaptureType::DeviceDepth:
 		return FVector4f(DeviceDepth, DeviceDepth, DeviceDepth, 1.0f);
 	default:
@@ -647,6 +697,18 @@ bool FSceneCapturePhotoSet::ComputeSample(
 			EmissivePhotoSet.ComputeSample(Position, Normal, VisibilityFunction, DefaultsInResultsOut.Emissive);
 		DefaultsInResultsOut.HaveValues.bEmissive = true;
 	}
+	if (SampleChannels.bOpacity)
+	{
+		DefaultsInResultsOut.Opacity =
+			OpacityPhotoSet.ComputeSample(Position, Normal, VisibilityFunction, DefaultsInResultsOut.Opacity);
+		DefaultsInResultsOut.HaveValues.bOpacity = true;
+	}
+	if (SampleChannels.bSubsurfaceColor)
+	{
+		DefaultsInResultsOut.SubsurfaceColor =
+			SubsurfaceColorPhotoSet.ComputeSample(Position, Normal, VisibilityFunction, DefaultsInResultsOut.SubsurfaceColor);
+		DefaultsInResultsOut.HaveValues.bSubsurfaceColor = true;
+	}
 	if (SampleChannels.bWorldNormal)
 	{
 		DefaultsInResultsOut.WorldNormal =
diff --git a/Engine/Plugins/Runtime/MeshModelingToolset/Source/ModelingComponents/Private/Scene/WorldRenderCapture.cpp b/Engine/Plugins/Runtime/MeshModelingToolset/Source/ModelingComponents/Private/Scene/WorldRenderCapture.cpp
--- a/Engine/Plugins/Runtime/MeshModelingToolset/Source/ModelingComponents/Private/Scene/WorldRenderCapture.cpp
+++ b/Engine/Plugins/Runtime/MeshModelingToolset/Source/ModelingComponents/Private/Scene/WorldRenderCapture.cpp
@@ -29,10 +29,17 @@ static TAutoConsoleVariable<int32> CVarModelingWorldRenderCaptureVTWarmupFrames(
 	5,
 	TEXT("Number of frames to render before each capture in order to warmup the VT."));
 
-FRenderCaptureTypeFlags FRenderCaptureTypeFlags::All()
+FRenderCaptureTypeFlags FRenderCaptureTypeFlags::All(bool bCombinedMRS)
 {
-	// TODO Cleanup this function, its confusing due to CombinedMRS overlapping with the individual channels
-	FRenderCaptureTypeFlags Result{ true, true, true, true, true, true };
+	FRenderCaptureTypeFlags Result;
+	Result.bBaseColor = true;
+	Result.bRoughness = true;
+	Result.bMetallic = true;
+	Result.bSpecular = true;
+	Result.bEmissive = true;
+	Result.bWorldNormal = true;
+	Result.bCombinedMRS = bCombinedMRS;
+	Result.bSubsurfaceColor = true;
 	Result.bDeviceDepth = true;
 	return Result;
 }
@@ -84,6 +91,9 @@ void FRenderCaptureTypeFlags::SetEnabled(ERenderCaptureType CaptureType, bool bE
 	case ERenderCaptureType::CombinedMRS:
 		bCombinedMRS = bEnabled;
 		break;
+	case ERenderCaptureType::SubsurfaceColor:
+		bSubsurfaceColor = bEnabled;
+		break;
 	case ERenderCaptureType::DeviceDepth:
 		bDeviceDepth = bEnabled;
 		break;
@@ -1032,6 +1042,12 @@ bool FWorldRenderCapture::CaptureFromPosition(
 		return bOK;
 	}
 
+	// The following handles buffer visualization materials found in /Engine/BufferVisualization/<CaptureTypeName>.
+	// Sometimes these postprocess materials change the raw GBuffer data, in particular:
+	// - The Roughness postprocess material output is:           GBufferRoughness^Gamma,   with Gamma=2.2
+	// - The SubsurfaceColor postprocess material output is:     GBufferSSColor^(1/Gamma), with Gamma=2.2
+	// Note the relationship between linear and gamma encoded values: EncodedValue = LinearValue^(1/Gamma)
+
 	// Roughness visualization is rendered with gamma correction (unclear why)
 	bool bLinear = (CaptureType != ERenderCaptureType::Roughness);
 	UTextureRenderTarget2D* RenderTargetTexture = GetRenderTexture(bLinear);
@@ -1061,10 +1077,24 @@ bool FWorldRenderCapture::CaptureFromPosition(
 	FName CaptureTypeName;
 	switch (CaptureType)
 	{
-	case ERenderCaptureType::WorldNormal:	CaptureTypeName = FName("WorldNormal");	break;
-	case ERenderCaptureType::Roughness:		CaptureTypeName = FName("Roughness");	break;
-	case ERenderCaptureType::Metallic:		CaptureTypeName = FName("Metallic");	break;
-	case ERenderCaptureType::Specular:		CaptureTypeName = FName("Specular");	break;
+	case ERenderCaptureType::WorldNormal:
+		CaptureTypeName = FName("WorldNormal");
+		break;
+	case ERenderCaptureType::Roughness:
+		CaptureTypeName = FName("Roughness");
+		break;
+	case ERenderCaptureType::Metallic:
+		CaptureTypeName = FName("Metallic");
+		break;
+	case ERenderCaptureType::Specular:
+		CaptureTypeName = FName("Specular");
+		break;
+	case ERenderCaptureType::Opacity:
+		CaptureTypeName = FName("Opacity");
+		break;
+	case ERenderCaptureType::SubsurfaceColor:
+		CaptureTypeName = FName("SubsurfaceColor");
+		break;
 	case ERenderCaptureType::BaseColor:
 	default:
 		CaptureTypeName = FName("BaseColor");
diff --git a/Engine/Plugins/Runtime/MeshModelingToolset/Source/ModelingComponents/Public/Baking/RenderCaptureFunctions.h b/Engine/Plugins/Runtime/MeshModelingToolset/Source/ModelingComponents/Public/Baking/RenderCaptureFunctions.h
--- a/Engine/Plugins/Runtime/MeshModelingToolset/Source/ModelingComponents/Public/Baking/RenderCaptureFunctions.h
+++ b/Engine/Plugins/Runtime/MeshModelingToolset/Source/ModelingComponents/Public/Baking/RenderCaptureFunctions.h
@@ -117,6 +117,8 @@ struct MODELINGCOMPONENTS_API FRenderCaptureOptions
 	bool bBakeSpecular = true;
 	bool bBakeEmissive = true;
 	bool bBakeNormalMap = true;
+	bool bBakeOpacity = true;
+	bool bBakeSubsurfaceColor = true;
 	bool bBakeDeviceDepth = true;
 	
 	bool bUsePackedMRS = true;
@@ -180,6 +182,8 @@ struct MODELINGCOMPONENTS_API FRenderCaptureTextures
 	UTexture2D* RoughnessMap = nullptr;
 	UTexture2D* SpecularMap = nullptr;
 	UTexture2D* EmissiveMap = nullptr;
+	UTexture2D* OpacityMap = nullptr;
+	UTexture2D* SubsurfaceColorMap = nullptr;
 };
 
 // Note: The source data in the textures is *not* updated by this function
diff --git a/Engine/Plugins/Runtime/MeshModelingToolset/Source/ModelingComponents/Public/Scene/SceneCapturePhotoSet.h b/Engine/Plugins/Runtime/MeshModelingToolset/Source/ModelingComponents/Public/Scene/SceneCapturePhotoSet.h
--- a/Engine/Plugins/Runtime/MeshModelingToolset/Source/ModelingComponents/Public/Scene/SceneCapturePhotoSet.h
+++ b/Engine/Plugins/Runtime/MeshModelingToolset/Source/ModelingComponents/Public/Scene/SceneCapturePhotoSet.h
@@ -26,6 +26,8 @@ namespace Geometry
  *		Specular
  *		PackedMRS   (Metallic / Roughness / Specular)
  *		Emissive
+ *		Opacity
+ *		SubsurfaceColor
  *		WorldNormal
  *		DeviceDepth
  *
@@ -113,6 +115,8 @@ public:
 		FVector3f Emissive;
 		FVector3f WorldNormal;
 		float DeviceDepth;
+		float Opacity;
+		FVector3f SubsurfaceColor;
 
 		FSceneSample();
 
@@ -168,6 +172,8 @@ public:
 	const FSpatialPhotoSet3f& GetPackedMRSPhotoSet() { return PackedMRSPhotoSet; }
 	const FSpatialPhotoSet3f& GetWorldNormalPhotoSet() { return WorldNormalPhotoSet; }
 	const FSpatialPhotoSet3f& GetEmissivePhotoSet() { return EmissivePhotoSet; }
+	const FSpatialPhotoSet1f& GetOpacityPhotoSet() { return OpacityPhotoSet; }
+	const FSpatialPhotoSet3f& GetSubsurfaceColorPhotoSet() { return SubsurfaceColorPhotoSet; }
 	const FSpatialPhotoSet1f& GetDeviceDepthPhotoSet() { return DeviceDepthPhotoSet; }
 
 	/**
@@ -208,6 +214,8 @@ protected:
 	bool bEnablePackedMRS = true;
 	bool bEnableWorldNormal = true;
 	bool bEnableEmissive = true;
+	bool bEnableOpacity = true;
+	bool bEnableSubsurfaceColor = true;
 	bool bEnableDeviceDepth = true;
 
 	FSpatialPhotoSet3f BaseColorPhotoSet;
@@ -217,6 +225,8 @@ protected:
 	FSpatialPhotoSet3f PackedMRSPhotoSet;
 	FSpatialPhotoSet3f WorldNormalPhotoSet;
 	FSpatialPhotoSet3f EmissivePhotoSet;
+	FSpatialPhotoSet1f OpacityPhotoSet;
+	FSpatialPhotoSet3f SubsurfaceColorPhotoSet;
 	FSpatialPhotoSet1f DeviceDepthPhotoSet;
 
 	FRenderCaptureConfig BaseColorConfig;
@@ -226,6 +236,8 @@ protected:
 	FRenderCaptureConfig PackedMRSConfig;
 	FRenderCaptureConfig WorldNormalConfig;
 	FRenderCaptureConfig EmissiveConfig;
+	FRenderCaptureConfig OpacityConfig;
+	FRenderCaptureConfig SubsurfaceColorConfig;
 	FRenderCaptureConfig DeviceDepthConfig;
 
 	TArray<FSpatialPhotoParams> PhotoSetParams;
@@ -281,6 +293,18 @@ FVector4f FSceneCapturePhotoSet::ComputeSampleNearest(
 		return FVector4f(Emissive, 1.f);
 	}
 
+	if constexpr (CaptureType == ERenderCaptureType::Opacity)
+	{
+		float Opacity = OpacityPhotoSet.ComputeSampleNearest(PhotoIndex, PhotoCoords, DefaultSample.Opacity);
+		return FVector4f(Opacity, Opacity, Opacity, 1.f);
+	}
+
+	if constexpr (CaptureType == ERenderCaptureType::SubsurfaceColor)
+	{
+		FVector3f SubsurfaceColor = SubsurfaceColorPhotoSet.ComputeSampleNearest(PhotoIndex, PhotoCoords, DefaultSample.SubsurfaceColor);
+		return FVector4f(SubsurfaceColor, 1.f);
+	}
+
 	if constexpr (CaptureType == ERenderCaptureType::WorldNormal)
 	{
 		FVector3f WorldNormal = WorldNormalPhotoSet.ComputeSampleNearest(PhotoIndex, PhotoCoords, DefaultSample.WorldNormal);
diff --git a/Engine/Plugins/Runtime/MeshModelingToolset/Source/ModelingComponents/Public/Scene/WorldRenderCapture.h b/Engine/Plugins/Runtime/MeshModelingToolset/Source/ModelingComponents/Public/Scene/WorldRenderCapture.h
--- a/Engine/Plugins/Runtime/MeshModelingToolset/Source/ModelingComponents/Public/Scene/WorldRenderCapture.h
+++ b/Engine/Plugins/Runtime/MeshModelingToolset/Source/ModelingComponents/Public/Scene/WorldRenderCapture.h
@@ -25,7 +25,9 @@ enum class ERenderCaptureType
 	Emissive = 16,
 	WorldNormal = 32,
 	DeviceDepth = 64,
-	CombinedMRS = 128
+	CombinedMRS = 128,
+	Opacity = 256,
+	SubsurfaceColor = 512
 };
 
 struct MODELINGCOMPONENTS_API FRenderCaptureConfig
@@ -50,9 +52,11 @@ struct MODELINGCOMPONENTS_API FRenderCaptureTypeFlags
 	bool bWorldNormal = false;
 	bool bCombinedMRS = false;
 	bool bDeviceDepth = false;
+	bool bOpacity = false;
+	bool bSubsurfaceColor = false;
 
 	/** @return FRenderCaptureTypeFlags with all types enabled/true */
-	static FRenderCaptureTypeFlags All();
+	static FRenderCaptureTypeFlags All(bool bCombinedMRS = false);
 
 	/** @return FRenderCaptureTypeFlags with all types disabled/false */
 	static FRenderCaptureTypeFlags None();
