MeshModelingToolset: Mark hole fill statistics as NoResetToDefault

diff --git a/Engine/Plugins/Experimental/MeshModelingToolsetExp/Source/MeshModelingToolsExp/Public/HoleFillTool.h b/Engine/Plugins/Experimental/MeshModelingToolsetExp/Source/MeshModelingToolsExp/Public/HoleFillTool.h
--- a/Engine/Plugins/Experimental/MeshModelingToolsetExp/Source/MeshModelingToolsExp/Public/HoleFillTool.h
+++ b/Engine/Plugins/Experimental/MeshModelingToolsetExp/Source/MeshModelingToolsExp/Public/HoleFillTool.h
@@ -171,19 +171,19 @@ class MESHMODELINGTOOLSEXP_API UHoleFillStatisticsProperties : public UInteracti
 	GENERATED_BODY()
 public:
 
-	UPROPERTY(VisibleAnywhere, Category = HoleFillStatistics)
+	UPROPERTY(VisibleAnywhere, Category = HoleFillStatistics, meta = (NoResetToDefault))
 	FString InitialHoles;
 
-	UPROPERTY(VisibleAnywhere, Category = HoleFillStatistics)
+	UPROPERTY(VisibleAnywhere, Category = HoleFillStatistics, meta = (NoResetToDefault))
 	FString SelectedHoles;
 
-	UPROPERTY(VisibleAnywhere, Category = HoleFillStatistics)
+	UPROPERTY(VisibleAnywhere, Category = HoleFillStatistics, meta = (NoResetToDefault))
 	FString SuccessfulFills;
 
-	UPROPERTY(VisibleAnywhere, Category = HoleFillStatistics)
+	UPROPERTY(VisibleAnywhere, Category = HoleFillStatistics, meta = (NoResetToDefault))
 	FString FailedFills;
 
-	UPROPERTY(VisibleAnywhere, Category = HoleFillStatistics)
+	UPROPERTY(VisibleAnywhere, Category = HoleFillStatistics, meta = (NoResetToDefault))
 	FString RemainingHoles;
 
 	void Initialize(const UHoleFillTool& HoleFillTool);
