Mesh Modeling: prevent subdivide tool from running with level 0 specified (for example if the input mesh is very high resolution and the subd level is automatically capped)

diff --git a/Engine/Plugins/Experimental/MeshModelingToolsetExp/Source/MeshModelingToolsEditorOnlyExp/Private/SubdividePolyTool.cpp b/Engine/Plugins/Experimental/MeshModelingToolsetExp/Source/MeshModelingToolsEditorOnlyExp/Private/SubdividePolyTool.cpp
--- a/Engine/Plugins/Experimental/MeshModelingToolsetExp/Source/MeshModelingToolsEditorOnlyExp/Private/SubdividePolyTool.cpp
+++ b/Engine/Plugins/Experimental/MeshModelingToolsetExp/Source/MeshModelingToolsEditorOnlyExp/Private/SubdividePolyTool.cpp
@@ -51,6 +51,11 @@ public:
 			return;
 		}
 
+		if (SubdivisionLevel == 0)
+		{
+			return;
+		}
+
 		constexpr bool bAutoCompute = true;
 		FGroupTopology Topo(&Mesh, bAutoCompute);
 		FSubdividePoly Subd(Topo, Mesh, SubdivisionLevel);
