Render Capture Baking: Refactor how the scene capture and baking steps are managed to fix a crash when scene capture parameters changed and implement other improvements

The crash that was fixed occured when, for example, the userturned off a capture channel while the baking background compute was running
The refactoring has enabled the tool to only bake the newly captured channels only, previously when the scene capture was incrementally updated all pre-existing channels were baked again
The refactoring should enable us to more easily support a feature where the user toggles capture directions or adds new ones manually

#rnx

diff --git a/Engine/Plugins/Experimental/GeometryScripting/Source/GeometryScriptingCore/Private/MeshBakeFunctions.cpp b/Engine/Plugins/Experimental/GeometryScripting/Source/GeometryScriptingCore/Private/MeshBakeFunctions.cpp
--- a/Engine/Plugins/Experimental/GeometryScripting/Source/GeometryScriptingCore/Private/MeshBakeFunctions.cpp
+++ b/Engine/Plugins/Experimental/GeometryScripting/Source/GeometryScriptingCore/Private/MeshBakeFunctions.cpp
@@ -815,10 +815,7 @@ namespace GeometryScriptBakeLocals
 	}
 
 
-	FRenderCaptureOptions MakeRenderCaptureOptions(
-		FGeometryScriptBakeRenderCaptureOptions BakeOptions,
-		FGeometryScriptBakeTargetMeshOptions TargetOptions
-		)
+	FRenderCaptureOptions MakeRenderCaptureOptions(FGeometryScriptBakeRenderCaptureOptions BakeOptions)
 	{
 		FRenderCaptureOptions Options;
 
@@ -834,7 +831,6 @@ namespace GeometryScriptBakeLocals
 			BakeOptions.bRoughnessMap ||
 			BakeOptions.bSpecularMap
 			);
-		Options.ValidSampleDepthThreshold = BakeOptions.CleanupTolerance;
 
 		Options.RenderCaptureImageSize = GeometryScriptBakeLocals::GetDimensions(BakeOptions.RenderCaptureResolution).GetHeight();
 		Options.bAntiAliasing = BakeOptions.bRenderCaptureAntiAliasing;
@@ -852,8 +848,6 @@ namespace GeometryScriptBakeLocals
 		Options.bBakeRoughness = BakeOptions.bPackedMRSMap ? false : BakeOptions.bRoughnessMap;
 		Options.bBakeSpecular =  BakeOptions.bPackedMRSMap ? false : BakeOptions.bSpecularMap;
 
-		Options.TargetUVLayer = static_cast<int32>(TargetOptions.TargetUVLayer);
-
 		return Options;
 	}
 } // end namespace GeometryScriptBakeLocals
@@ -1045,8 +1039,6 @@ FGeometryScriptRenderCaptureTextures UGeometryScriptLibrary_MeshBakeFunctions::B
 	FGeometryScriptBakeRenderCaptureOptions BakeOptions,
 	UGeometryScriptDebug* Debug)
 {
-	FRenderCaptureOptions Options = GeometryScriptBakeLocals::MakeRenderCaptureOptions(BakeOptions, TargetOptions);
-
 	// Its possible to pass nullptrs in SourceActors so we filter these out here
 	TArray<TObjectPtr<AActor>> ValidSourceActors;
 	for (AActor* Actor : SourceActors)
@@ -1102,14 +1094,16 @@ FGeometryScriptRenderCaptureTextures UGeometryScriptLibrary_MeshBakeFunctions::B
 		return {};
 	}
 
-	if (Options.bBakeNormalMap && !FDynamicMeshTangents(TargetMesh->GetMeshPtr()).HasValidTangents(true))
+	if (BakeOptions.bNormalMap && !FDynamicMeshTangents(TargetMesh->GetMeshPtr()).HasValidTangents(true))
 	{
 		AppendWarning(Debug, EGeometryScriptErrorType::InvalidInputs,
 			LOCTEXT("BakeTextureFromRenderCaptures_InvalidMeshTangents", "BakeTextureFromRenderCaptures: TargetMesh has invalid tangents so the requested normal map cannot be baked"));
 		return {};
 	}
 
-	TUniquePtr<FSceneCapturePhotoSet> SceneCapture = CapturePhotoSet(ValidSourceActors, Options, false);
+	const FRenderCaptureOptions Options = GeometryScriptBakeLocals::MakeRenderCaptureOptions(BakeOptions);
+	FRenderCaptureUpdate UnusedUpdate;
+	TUniquePtr<FSceneCapturePhotoSet> SceneCapture = CapturePhotoSet(ValidSourceActors, Options, UnusedUpdate, false);
 
 	const FDynamicMeshAABBTree3 TargetMeshSpatial(TargetMesh->GetMeshPtr());
 	TSharedPtr<FMeshTangentsd, ESPMode::ThreadSafe> TargetMeshTangents = MakeShared<FMeshTangentsd>(TargetMesh->GetMeshPtr());
@@ -1124,12 +1118,14 @@ FGeometryScriptRenderCaptureTextures UGeometryScriptLibrary_MeshBakeFunctions::B
 
 	FRenderCaptureOcclusionHandler OcclusionHandler(GeometryScriptBakeLocals::GetDimensions(BakeOptions.Resolution));
 
+	const FRenderCaptureOptions PendingBake = Options; // All specified channels need baking
 	TUniquePtr<FMeshMapBaker> Baker = MakeRenderCaptureBaker(
 		TargetMesh->GetMeshPtr(),
 		TargetMeshTangents,
 		SceneCapture.Get(),
 		&Sampler,
-		Options,
+		PendingBake,
+		TargetOptions.TargetUVLayer,
 		GeometryScriptBakeLocals::ConvertEnum(BakeOptions.Resolution),
 		GeometryScriptBakeLocals::ConvertEnum(BakeOptions.SamplesPerPixel),
 		&OcclusionHandler);
@@ -1142,39 +1138,39 @@ FGeometryScriptRenderCaptureTextures UGeometryScriptLibrary_MeshBakeFunctions::B
 	GetTexturesFromRenderCaptureBaker(Baker, TexturesOut);
 
 	// Update source data
-	if (Options.bBakeBaseColor && TexturesOut.BaseColorMap)
+	if (TexturesOut.BaseColorMap)
 	{
 		FTexture2DBuilder::CopyPlatformDataToSourceData(TexturesOut.BaseColorMap, FTexture2DBuilder::ETextureType::Color);
 	}
-	if (Options.bBakeNormalMap && TexturesOut.NormalMap)
+	if (TexturesOut.NormalMap)
 	{
 		FTexture2DBuilder::CopyPlatformDataToSourceData(TexturesOut.NormalMap, FTexture2DBuilder::ETextureType::NormalMap);
 	}
-	if (Options.bUsePackedMRS && TexturesOut.PackedMRSMap)
+	if (TexturesOut.PackedMRSMap)
 	{
 		FTexture2DBuilder::CopyPlatformDataToSourceData(TexturesOut.PackedMRSMap, FTexture2DBuilder::ETextureType::ColorLinear);
 	}
-	if (Options.bBakeMetallic && TexturesOut.MetallicMap)
+	if (TexturesOut.MetallicMap)
 	{
 		FTexture2DBuilder::CopyPlatformDataToSourceData(TexturesOut.MetallicMap, FTexture2DBuilder::ETextureType::Metallic);
 	}
-	if (Options.bBakeRoughness && TexturesOut.RoughnessMap)
+	if (TexturesOut.RoughnessMap)
 	{
 		FTexture2DBuilder::CopyPlatformDataToSourceData(TexturesOut.RoughnessMap, FTexture2DBuilder::ETextureType::Roughness);
 	}
-	if (Options.bBakeSpecular && TexturesOut.SpecularMap)
+	if (TexturesOut.SpecularMap)
 	{
 		FTexture2DBuilder::CopyPlatformDataToSourceData(TexturesOut.SpecularMap, FTexture2DBuilder::ETextureType::Specular);
 	}
-	if (Options.bBakeEmissive && TexturesOut.EmissiveMap)
+	if (TexturesOut.EmissiveMap)
 	{
 		FTexture2DBuilder::CopyPlatformDataToSourceData(TexturesOut.EmissiveMap, FTexture2DBuilder::ETextureType::EmissiveHDR);
 	}
-	if (Options.bBakeOpacity && TexturesOut.OpacityMap)
+	if (TexturesOut.OpacityMap)
 	{
 		FTexture2DBuilder::CopyPlatformDataToSourceData(TexturesOut.OpacityMap, FTexture2DBuilder::ETextureType::ColorLinear);
 	}
-	if (Options.bBakeSubsurfaceColor && TexturesOut.SubsurfaceColorMap)
+	if (TexturesOut.SubsurfaceColorMap)
 	{
 		FTexture2DBuilder::CopyPlatformDataToSourceData(TexturesOut.SubsurfaceColorMap, FTexture2DBuilder::ETextureType::Color);
 	}
diff --git a/Engine/Plugins/Experimental/MeshModelingToolsetExp/Source/GeometryProcessingAdapters/Private/ApproximateActorsImpl.cpp b/Engine/Plugins/Experimental/MeshModelingToolsetExp/Source/GeometryProcessingAdapters/Private/ApproximateActorsImpl.cpp
--- a/Engine/Plugins/Experimental/MeshModelingToolsetExp/Source/GeometryProcessingAdapters/Private/ApproximateActorsImpl.cpp
+++ b/Engine/Plugins/Experimental/MeshModelingToolsetExp/Source/GeometryProcessingAdapters/Private/ApproximateActorsImpl.cpp
@@ -111,13 +111,22 @@ static TUniquePtr<FSceneCapturePhotoSet> CapturePhotoSet(
 		SceneCapture->SetCaptureTypeEnabled(ERenderCaptureType::Specular, bSpecular);
 	}
 
+	UWorld* World = Actors.IsEmpty() ? nullptr : Actors[0]->GetWorld();
 
-	SceneCapture->SetCaptureSceneActors(Actors[0]->GetWorld(), Actors);
+	SceneCapture->SetCaptureSceneActors(World, Actors);
 
-	SceneCapture->AddStandardExteriorCapturesFromBoundingBox(
-		CaptureDimensions, FieldOfView, NearPlaneDist,
+	const TArray<FSpatialPhotoParams> SpatialParams = ComputeStandardExteriorSpatialPhotoParameters(
+		World,
+		Actors,
+		CaptureDimensions,
+		FieldOfView,
+		NearPlaneDist,
 		true, true, true, true, true);
-	
+
+	SceneCapture->SetSpatialPhotoParams(SpatialParams);
+
+	SceneCapture->Compute();
+
 	return SceneCapture;
 }
 
diff --git a/Engine/Plugins/Experimental/MeshModelingToolsetExp/Source/MeshModelingToolsExp/Private/BakeRenderCaptureTool.cpp b/Engine/Plugins/Experimental/MeshModelingToolsetExp/Source/MeshModelingToolsExp/Private/BakeRenderCaptureTool.cpp
--- a/Engine/Plugins/Experimental/MeshModelingToolsetExp/Source/MeshModelingToolsExp/Private/BakeRenderCaptureTool.cpp
+++ b/Engine/Plugins/Experimental/MeshModelingToolsetExp/Source/MeshModelingToolsExp/Private/BakeRenderCaptureTool.cpp
@@ -17,7 +17,6 @@
 
 #include "EngineAnalytics.h"
 
-#include "Baking/RenderCaptureFunctions.h"
 #include "Baking/BakingTypes.h"
 #include "Sampling/MeshImageBakingCache.h"
 #include "Sampling/MeshMapBaker.h"
@@ -44,18 +43,10 @@ static FString PackedMRSTexParamName = TEXT("PackedMRS");
 
 
 FRenderCaptureOptions
-MakeRenderCaptureOptions(
-	const URenderCaptureProperties& RenderCaptureProperties,
-	const UBakeRenderCaptureToolProperties& ToolProperties,
-	const UBakeRenderCaptureInputToolProperties& InputMeshSettings)
+MakeRenderCaptureOptions(const URenderCaptureProperties& RenderCaptureProperties)
 {
 	FRenderCaptureOptions Options;
 
-	Options.TargetUVLayer = InputMeshSettings.GetTargetUVLayerIndex();
-
-	Options.RenderCaptureImageSize = static_cast<int32>(RenderCaptureProperties.Resolution);
-	Options.ValidSampleDepthThreshold = ToolProperties.ValidSampleDepthThreshold;
-
 	Options.bBakeBaseColor = RenderCaptureProperties.bBaseColorMap;
 	Options.bBakeNormalMap = RenderCaptureProperties.bNormalMap;
 	Options.bBakeEmissive =  RenderCaptureProperties.bEmissiveMap;
@@ -71,6 +62,7 @@ MakeRenderCaptureOptions(
 	Options.bBakeRoughness = RenderCaptureProperties.bPackedMRSMap ? false : RenderCaptureProperties.bRoughnessMap;
 	Options.bBakeSpecular =  RenderCaptureProperties.bPackedMRSMap ? false : RenderCaptureProperties.bSpecularMap;
 
+	Options.RenderCaptureImageSize = static_cast<int32>(RenderCaptureProperties.Resolution);
 	Options.bAntiAliasing = RenderCaptureProperties.bAntiAliasing;
 	Options.FieldOfViewDegrees = RenderCaptureProperties.CaptureFieldOfView;
 	Options.NearPlaneDist = RenderCaptureProperties.NearPlaneDist;
@@ -88,11 +80,15 @@ class FRenderCaptureMapBakerOp : public TGenericDataOperator<FMeshMapBaker>
 public:
 	UE::Geometry::FDynamicMesh3* BaseMesh = nullptr;
 	TSharedPtr<UE::Geometry::FMeshTangentsd, ESPMode::ThreadSafe> BaseMeshTangents;
-	FRenderCaptureOptions Options;
+	int32 TargetUVLayer;
+	double ValidSampleDepthThreshold;
 	EBakeTextureResolution TextureImageSize;
 	EBakeTextureSamplesPerPixel SamplesPerPixel;
 	FSceneCapturePhotoSet* SceneCapture = nullptr;
 
+	// Used to pass the channels which need baking via the bBakeXXX and bUsePackedMRS members
+	FRenderCaptureOptions PendingBake;
+
 	// Begin TGenericDataOperator interface
 	virtual void CalculateResult(FProgressCancel* Progress) override;
 	// End TGenericDataOperator interface
@@ -105,7 +101,7 @@ void FRenderCaptureMapBakerOp::CalculateResult(FProgressCancel*)
 
 	FSceneCapturePhotoSetSampler Sampler(
 		SceneCapture,
-		Options.ValidSampleDepthThreshold,
+		ValidSampleDepthThreshold,
 		BaseMesh,
 		&BaseMeshSpatial,
 		BaseMeshTangents.Get());
@@ -121,7 +117,8 @@ void FRenderCaptureMapBakerOp::CalculateResult(FProgressCancel*)
 		BaseMeshTangents,
 		SceneCapture,
 		&Sampler,
-		Options,
+		PendingBake,
+		TargetUVLayer,
 		TextureImageSize,
 		SamplesPerPixel,
 		&OcclusionHandler);
@@ -210,7 +207,7 @@ void UBakeRenderCaptureTool::Setup()
 	AddToolPropertySource(Settings);
 
 	Settings->MapPreview = BaseColorTexParamName;
-	Settings->WatchProperty(Settings->MapPreview, [this](FString) { UpdateVisualization(); GetToolManager()->PostInvalidation(); });
+	Settings->WatchProperty(Settings->MapPreview, [this](FString) { UpdateVisualization(); });
 	Settings->WatchProperty(Settings->SamplesPerPixel, [this](EBakeTextureSamplesPerPixel) { OpState |= EBakeOpState::Evaluate; });
 	Settings->WatchProperty(Settings->TextureSize, [this](EBakeTextureResolution) { OpState |= EBakeOpState::Evaluate; });
 	Settings->WatchProperty(Settings->ValidSampleDepthThreshold, [this](float ValidSampleDepthThreshold)
@@ -278,9 +275,6 @@ void UBakeRenderCaptureTool::Setup()
 
 	TargetUVLayerToError.Reset();
 
-	// Used to implement SceneCapture cancellation
-	ComputedRenderCaptureProperties = NewObject<URenderCaptureProperties>(this);
-
 	// Hide the render capture meshes since this baker operates solely in world space which will occlude the preview of
 	// the target mesh.
 	for (int Idx = 1; Idx < NumTargets; ++Idx)
@@ -290,7 +284,6 @@ void UBakeRenderCaptureTool::Setup()
 
 	// Make sure we trigger SceneCapture computation in UpdateResult
 	OpState |= EBakeOpState::Evaluate;
-	ComputedRenderCaptureProperties->NearPlaneDist = 0.f; // Arbitrary invalid value
 
 	SetToolDisplayName(LOCTEXT("ToolName", "Bake Render Capture"));
 	GetToolManager()->DisplayMessage(
@@ -509,10 +502,29 @@ TUniquePtr<TGenericDataOperator<FMeshMapBaker>> UBakeRenderCaptureTool::MakeNewO
 	TUniquePtr<FRenderCaptureMapBakerOp> Op = MakeUnique<FRenderCaptureMapBakerOp>();
 	Op->BaseMesh = &TargetMesh;
 	Op->BaseMeshTangents = TargetMeshTangents;
-	Op->Options = MakeRenderCaptureOptions(*RenderCaptureProperties, *Settings, *InputMeshSettings);
+	Op->TargetUVLayer = InputMeshSettings->GetTargetUVLayerIndex();
+	Op->ValidSampleDepthThreshold = Settings->ValidSampleDepthThreshold;
 	Op->TextureImageSize = Settings->TextureSize;
 	Op->SamplesPerPixel = Settings->SamplesPerPixel;
 	Op->SceneCapture = SceneCapture.Get();
+
+	auto IsPendingBake = [this](UTexture *CurrentResult, ERenderCaptureType CaptureType) -> bool
+	{
+		const bool bRequested = SceneCapture->GetCaptureTypeEnabled(CaptureType);
+		const bool bBaked = (CurrentResult != nullptr);
+		return (bRequested && !bBaked);
+	};
+
+	Op->PendingBake.bBakeBaseColor = IsPendingBake(ResultSettings->BaseColorMap, ERenderCaptureType::BaseColor);
+	Op->PendingBake.bBakeRoughness = IsPendingBake(ResultSettings->RoughnessMap, ERenderCaptureType::Roughness);
+	Op->PendingBake.bBakeMetallic =  IsPendingBake(ResultSettings->MetallicMap,  ERenderCaptureType::Metallic);
+	Op->PendingBake.bBakeSpecular =  IsPendingBake(ResultSettings->SpecularMap,  ERenderCaptureType::Specular);
+	Op->PendingBake.bUsePackedMRS =  IsPendingBake(ResultSettings->PackedMRSMap, ERenderCaptureType::CombinedMRS);
+	Op->PendingBake.bBakeEmissive =  IsPendingBake(ResultSettings->EmissiveMap,  ERenderCaptureType::Emissive);
+	Op->PendingBake.bBakeNormalMap = IsPendingBake(ResultSettings->NormalMap,    ERenderCaptureType::WorldNormal);
+	Op->PendingBake.bBakeOpacity =   IsPendingBake(ResultSettings->OpacityMap,   ERenderCaptureType::Opacity);
+	Op->PendingBake.bBakeSubsurfaceColor = IsPendingBake(ResultSettings->SubsurfaceColorMap, ERenderCaptureType::SubsurfaceColor);
+
 	return Op;
 }
 
@@ -525,20 +537,48 @@ void UBakeRenderCaptureTool::OnMapsUpdatedRC(const TUniquePtr<FMeshMapBaker>& Ne
 	FRenderCaptureTextures TexturesOut;
 	GetTexturesFromRenderCaptureBaker(NewResult, TexturesOut);
 
-	// Unpack TexturesOut to store in ResultSettings
-	ResultSettings->BaseColorMap = TexturesOut.BaseColorMap;
-	ResultSettings->NormalMap = TexturesOut.NormalMap;
-	ResultSettings->PackedMRSMap = TexturesOut.PackedMRSMap;
-	ResultSettings->MetallicMap = TexturesOut.MetallicMap;
-	ResultSettings->RoughnessMap = TexturesOut.RoughnessMap;
-	ResultSettings->SpecularMap = TexturesOut.SpecularMap;
-	ResultSettings->EmissiveMap = TexturesOut.EmissiveMap;
-	ResultSettings->OpacityMap = TexturesOut.OpacityMap;
-	ResultSettings->SubsurfaceColorMap = TexturesOut.SubsurfaceColorMap;
+	// The NewResult will contain the newly baked textures so we only update those and not overwrite any already baked
+	// valid results. Note that if results are invalidated by some parameter change they are cleared by the
+	// InvalidateResults function, by comparing/syncing with the enabled capture types in SceneCapture
+	if (TexturesOut.BaseColorMap)
+	{
+		ResultSettings->BaseColorMap = TexturesOut.BaseColorMap;
+	}
+	if (TexturesOut.NormalMap)
+	{
+		ResultSettings->NormalMap = TexturesOut.NormalMap;
+	}
+	if (TexturesOut.PackedMRSMap)
+	{
+		ResultSettings->PackedMRSMap = TexturesOut.PackedMRSMap;
+	}
+	if (TexturesOut.MetallicMap)
+	{
+		ResultSettings->MetallicMap = TexturesOut.MetallicMap;
+	}
+	if (TexturesOut.RoughnessMap)
+	{
+		ResultSettings->RoughnessMap = TexturesOut.RoughnessMap;
+	}
+	if (TexturesOut.SpecularMap)
+	{
+		ResultSettings->SpecularMap = TexturesOut.SpecularMap;
+	}
+	if (TexturesOut.EmissiveMap)
+	{
+		ResultSettings->EmissiveMap = TexturesOut.EmissiveMap;
+	}
+	if (TexturesOut.OpacityMap)
+	{
+		ResultSettings->OpacityMap = TexturesOut.OpacityMap;
+	}
+	if (TexturesOut.SubsurfaceColorMap)
+	{
+		ResultSettings->SubsurfaceColorMap = TexturesOut.SubsurfaceColorMap;
+	}
 
 	GatherAnalytics(*NewResult);
 	UpdateVisualization();
-	GetToolManager()->PostInvalidation();
 }
 
 
@@ -682,107 +722,86 @@ void UBakeRenderCaptureTool::InvalidateComputeRC()
 	Compute->InvalidateResult();
 }
 
-
-void UBakeRenderCaptureTool::UpdateResult()
+FRenderCaptureUpdate UBakeRenderCaptureTool::UpdateSceneCapture()
 {
-	if (OpState == EBakeOpState::Clean)
+	for (int Idx = 1; Idx < Targets.Num(); ++Idx)
 	{
-		// Evaluation already launched/complete. Note that the Compute background compute updates ResultSettings when
-		// they are available by calling OnMapsUpdatedRC in its OnResultUpdated delegate.
-		return;
+		UE::ToolTarget::ShowSourceObject(Targets[Idx]);
 	}
 
-	//
-	// Create a set of spatially located render captures of the scene ("photo set"). We need to recompute this if the
-	// render capture properties changed. Note we only compare the URenderCaptureProperties, and not the
-	// ValidSampleDepthThreshold, this is intentional so that we only trigger a scene capture recompute when we go from
-	// a zero to a positive threshold (we need to compute the depth capture), or a positive to a zero threshold (we can
-	// save memory and not compute the depth capture), we don't need to recompute the scene capture when the user is
-	// changing between positive threshold values.
-	//
-	if (*RenderCaptureProperties != *ComputedRenderCaptureProperties)
+	FRenderCaptureUpdate Update;
+
+	if (!SceneCapture)
 	{
-		for (int Idx = 1; Idx < Targets.Num(); ++Idx)
-		{
-			UE::ToolTarget::ShowSourceObject(Targets[Idx]);
-		}
+		constexpr bool bAllowCancel = false;
 
-		// Do not allow user-cancellation on the call that occurs when the Render Capture Tool starts up
-		const bool bAllowCancel = (bFirstEverSceneCapture == false);
+		FScopedSlowTask Progress(1.f, LOCTEXT("CapturingScene", "Capturing Scene..."));
+		Progress.EnterProgressFrame(1.f);
+		Progress.MakeDialog(bAllowCancel);
 
-		FRenderCaptureOptions Options = MakeRenderCaptureOptions(*RenderCaptureProperties, *Settings, *InputMeshSettings);
-		{
-			FScopedSlowTask Progress(1.f, LOCTEXT("CapturingScene", "Capturing Scene..."));
-			Progress.EnterProgressFrame(1.f);
-			Progress.MakeDialog(bAllowCancel);
+		const FRenderCaptureOptions Options = MakeRenderCaptureOptions(*RenderCaptureProperties);
+		SceneCapture = CapturePhotoSet(Actors, Options, Update, bAllowCancel);
+	}
+	else
+	{
+		constexpr bool bAllowCancel = true;
 
-			const bool bRequireFullRecompute = (
-				RenderCaptureProperties->Resolution         != ComputedRenderCaptureProperties->Resolution ||
-				RenderCaptureProperties->bAntiAliasing      != ComputedRenderCaptureProperties->bAntiAliasing ||
-				RenderCaptureProperties->CaptureFieldOfView != ComputedRenderCaptureProperties->CaptureFieldOfView ||
-				RenderCaptureProperties->NearPlaneDist      != ComputedRenderCaptureProperties->NearPlaneDist);
+		FScopedSlowTask Progress(1.f, LOCTEXT("CapturingScene", "Capturing Scene..."));
+		Progress.EnterProgressFrame(1.f);
+		Progress.MakeDialog(bAllowCancel);
 
-			if (!SceneCapture || (SceneCapture && bRequireFullRecompute))
-			{
-				SceneCapture = CapturePhotoSet(Actors, Options, bAllowCancel);
-			}
-			else
-			{
-				UpdatePhotoSet(SceneCapture, Actors, Options, bAllowCancel);
-			}
-		}
+		// Get already computed options from the existing SceneCapture so we can restore them if the capture is cancelled
+		const FRenderCaptureOptions ComputedOptions = GetComputedPhotoSetOptions(SceneCapture);
 
-		for (int Idx = 1; Idx < Targets.Num(); ++Idx)
-		{
-			UE::ToolTarget::HideSourceObject(Targets[Idx]);
-		}
+		const FRenderCaptureOptions Options = MakeRenderCaptureOptions(*RenderCaptureProperties);
+		Update = UpdatePhotoSets(SceneCapture, Actors, Options, bAllowCancel);
 
 		if (SceneCapture->Cancelled())
 		{
 			// Restore the settings present before the change that invoked the scene capture recompute
-			RenderCaptureProperties->Resolution         = ComputedRenderCaptureProperties->Resolution;
-			RenderCaptureProperties->bBaseColorMap      = ComputedRenderCaptureProperties->bBaseColorMap;
-			RenderCaptureProperties->bNormalMap         = ComputedRenderCaptureProperties->bNormalMap;
-			RenderCaptureProperties->bMetallicMap       = ComputedRenderCaptureProperties->bMetallicMap;
-			RenderCaptureProperties->bRoughnessMap      = ComputedRenderCaptureProperties->bRoughnessMap;
-			RenderCaptureProperties->bSpecularMap       = ComputedRenderCaptureProperties->bSpecularMap;
-			RenderCaptureProperties->bPackedMRSMap      = ComputedRenderCaptureProperties->bPackedMRSMap;
-			RenderCaptureProperties->bEmissiveMap       = ComputedRenderCaptureProperties->bEmissiveMap;
-			RenderCaptureProperties->bOpacityMap        = ComputedRenderCaptureProperties->bOpacityMap;
-			RenderCaptureProperties->bSubsurfaceColorMap = ComputedRenderCaptureProperties->bSubsurfaceColorMap;
-			RenderCaptureProperties->bAntiAliasing      = ComputedRenderCaptureProperties->bAntiAliasing;
-			RenderCaptureProperties->bDeviceDepthMap    = ComputedRenderCaptureProperties->bDeviceDepthMap;
-			RenderCaptureProperties->CaptureFieldOfView = ComputedRenderCaptureProperties->CaptureFieldOfView;
-			RenderCaptureProperties->NearPlaneDist      = ComputedRenderCaptureProperties->NearPlaneDist;
-			Settings->ValidSampleDepthThreshold = ComputedValidDepthThreshold;
+			RenderCaptureProperties->bBaseColorMap      = ComputedOptions.bBakeBaseColor;
+			RenderCaptureProperties->bNormalMap         = ComputedOptions.bBakeNormalMap;
+			RenderCaptureProperties->bMetallicMap       = ComputedOptions.bBakeMetallic;
+			RenderCaptureProperties->bRoughnessMap      = ComputedOptions.bBakeRoughness;
+			RenderCaptureProperties->bSpecularMap       = ComputedOptions.bBakeSpecular;
+			RenderCaptureProperties->bPackedMRSMap      = ComputedOptions.bUsePackedMRS;
+			RenderCaptureProperties->bEmissiveMap       = ComputedOptions.bBakeEmissive;
+			RenderCaptureProperties->bOpacityMap        = ComputedOptions.bBakeOpacity;
+			RenderCaptureProperties->bSubsurfaceColorMap = ComputedOptions.bBakeSubsurfaceColor;
+			RenderCaptureProperties->bDeviceDepthMap    = ComputedOptions.bBakeDeviceDepth;
+			RenderCaptureProperties->bAntiAliasing      = ComputedOptions.bAntiAliasing;
+			RenderCaptureProperties->CaptureFieldOfView = ComputedOptions.FieldOfViewDegrees;
+			RenderCaptureProperties->NearPlaneDist      = ComputedOptions.NearPlaneDist;
+			RenderCaptureProperties->Resolution         = static_cast<EBakeTextureResolution>(ComputedOptions.RenderCaptureImageSize);
 
 			// Silently make the above updates so we don't overwrite the change to OpState below and call this function again
 			RenderCaptureProperties->SilentUpdateWatched();
 			Settings->SilentUpdateWatched();
+		}
+	}
 
-			OpState = EBakeOpState::Clean;
+	for (int Idx = 1; Idx < Targets.Num(); ++Idx)
+	{
+		UE::ToolTarget::HideSourceObject(Targets[Idx]);
+	}
 
-			return;
-		}
+	return Update;
+}
 
-		// Cache Settings used to compute this SceneCapture so we can restore them if the user cancels a SceneCapture recompute
-		ComputedRenderCaptureProperties->Resolution         = RenderCaptureProperties->Resolution;
-		ComputedRenderCaptureProperties->bBaseColorMap      = RenderCaptureProperties->bBaseColorMap;
-		ComputedRenderCaptureProperties->bNormalMap         = RenderCaptureProperties->bNormalMap;
-		ComputedRenderCaptureProperties->bMetallicMap       = RenderCaptureProperties->bMetallicMap;
-		ComputedRenderCaptureProperties->bRoughnessMap      = RenderCaptureProperties->bRoughnessMap;
-		ComputedRenderCaptureProperties->bSpecularMap       = RenderCaptureProperties->bSpecularMap;
-		ComputedRenderCaptureProperties->bPackedMRSMap      = RenderCaptureProperties->bPackedMRSMap;
-		ComputedRenderCaptureProperties->bEmissiveMap       = RenderCaptureProperties->bEmissiveMap;
-		ComputedRenderCaptureProperties->bOpacityMap        = RenderCaptureProperties->bOpacityMap;
-		ComputedRenderCaptureProperties->bSubsurfaceColorMap = RenderCaptureProperties->bSubsurfaceColorMap;
-		ComputedRenderCaptureProperties->bAntiAliasing      = RenderCaptureProperties->bAntiAliasing;
-		ComputedRenderCaptureProperties->bDeviceDepthMap    = RenderCaptureProperties->bDeviceDepthMap;
-		ComputedRenderCaptureProperties->CaptureFieldOfView = RenderCaptureProperties->CaptureFieldOfView;
-		ComputedRenderCaptureProperties->NearPlaneDist      = RenderCaptureProperties->NearPlaneDist;
-		ComputedValidDepthThreshold = Settings->ValidSampleDepthThreshold;
+// Process dirty props and update background compute. Called by UBakeMeshAttributeMapsToolBase::Render
+void UBakeRenderCaptureTool::UpdateResult()
+{
+	// Return if the bake is already launched/complete.
+	if (OpState == EBakeOpState::Clean)
+	{
+		return;
+	}
 
-		bFirstEverSceneCapture = false;
+	// The bake operation, Compute, stores a pointer to the SceneCapture so that must not be modified while baking
+	const bool bComputeInProgress = (Compute && (Compute->GetElapsedComputeTime() > 0.f));
+	if (bComputeInProgress)
+	{
+		return;
 	}
 
 	FText ErrorMessage; // Empty message indicates no error
@@ -833,26 +852,53 @@ void UBakeRenderCaptureTool::UpdateResult()
 	// Calling DisplayMessage with an empty string will clear existing messages
 	GetToolManager()->DisplayMessage(ErrorMessage, EToolMessageLevel::UserWarning);
 
-	InvalidateResults();
-
 	const bool bIsInvalid = (ErrorMessage.IsEmpty() == false);
 	if (bIsInvalid)
 	{
+		// Clear all results and wait for user to fix the invalid tool inputs
+		InvalidateResults(FRenderCaptureUpdate{});
+
+		// Only call UpdateVisualization when we first detect the invalid inputs
 		const bool bWasValid = static_cast<bool>(OpState & EBakeOpState::Invalid) == false;
 		if (bWasValid)
 		{
-			UpdateVisualization(); // Clear the preview mesh material inputs
+			UpdateVisualization();
 		}
+
+		// Set an invalid op state so we re-enter this function until the inputs are valid
 		OpState = EBakeOpState::Invalid;
 	}
 	else
 	{
+		FRenderCaptureUpdate Update = UpdateSceneCapture();
+
+		const bool bInvalidateAll =
+			ComputedTextureSize != Settings->TextureSize ||
+			ComputedSamplesPerPixel != Settings->SamplesPerPixel ||
+			ComputedValidDepthThreshold != Settings->ValidSampleDepthThreshold;
+
+		// Invalidate any results which need re-baking because the SceneCapture was updated, or
+		// Invalidate all results if the any baking parameters were changed
+		InvalidateResults(bInvalidateAll ? FRenderCaptureUpdate{} : Update);
+
+		// Update the preview mesh material with the (possibly invalidated) bake results
+		UpdateVisualization();
+
+		// Start another bake operation
 		InvalidateComputeRC();
+
+		// Cache computed parameters which are used to determine if results need re-baking
+		ComputedTextureSize = Settings->TextureSize;
+		ComputedSamplesPerPixel = Settings->SamplesPerPixel;
+		ComputedValidDepthThreshold = Settings->ValidSampleDepthThreshold;
+
+		// Set a clean op state so we don't render this function until its completed
 		OpState = EBakeOpState::Clean;
 	}
 }
 
 
+
 void UBakeRenderCaptureTool::UpdateVisualization()
 {
 	if (Settings->MapPreview.IsEmpty())
@@ -935,21 +981,52 @@ void UBakeRenderCaptureTool::UpdateVisualization()
 
 	Material->SetScalarParameterValue(TEXT("UVChannel"), InputMeshSettings->GetTargetUVLayerIndex());
 	PreviewMesh->SetOverrideRenderMaterial(Material);
-}
 
+	GetToolManager()->PostInvalidation();
+}
 
 
-void UBakeRenderCaptureTool::InvalidateResults()
+void UBakeRenderCaptureTool::InvalidateResults(FRenderCaptureUpdate Update)
 {
-	ResultSettings->BaseColorMap = nullptr;
-	ResultSettings->RoughnessMap = nullptr;
-	ResultSettings->MetallicMap = nullptr;
-	ResultSettings->SpecularMap = nullptr;
-	ResultSettings->PackedMRSMap = nullptr;
-	ResultSettings->EmissiveMap = nullptr;
-	ResultSettings->OpacityMap = nullptr;
-	ResultSettings->SubsurfaceColorMap = nullptr;
-	ResultSettings->NormalMap = nullptr;
+	// Note that the bake operation, Compute, updates ResultSettings when results are available via the
+	// Compute->OnResultUpdated delegate.
+
+	if (!SceneCapture->GetCaptureTypeEnabled(ERenderCaptureType::BaseColor) || Update.bUpdatedBaseColor)
+	{
+		ResultSettings->BaseColorMap = nullptr;
+	}
+	if (!SceneCapture->GetCaptureTypeEnabled(ERenderCaptureType::Roughness) || Update.bUpdatedRoughness)
+	{
+		ResultSettings->RoughnessMap = nullptr;
+	}
+	if (!SceneCapture->GetCaptureTypeEnabled(ERenderCaptureType::Metallic) || Update.bUpdatedMetallic)
+	{
+		ResultSettings->MetallicMap = nullptr;
+	}
+	if (!SceneCapture->GetCaptureTypeEnabled(ERenderCaptureType::Specular) || Update.bUpdatedSpecular)
+	{
+		ResultSettings->SpecularMap = nullptr;
+	}
+	if (!SceneCapture->GetCaptureTypeEnabled(ERenderCaptureType::Emissive) || Update.bUpdatedEmissive)
+	{
+		ResultSettings->EmissiveMap = nullptr;
+	}
+	if (!SceneCapture->GetCaptureTypeEnabled(ERenderCaptureType::WorldNormal) || Update.bUpdatedNormalMap)
+	{
+		ResultSettings->NormalMap = nullptr;
+	}
+	if (!SceneCapture->GetCaptureTypeEnabled(ERenderCaptureType::Opacity) || Update.bUpdatedOpacity)
+	{
+		ResultSettings->OpacityMap = nullptr;
+	}
+	if (!SceneCapture->GetCaptureTypeEnabled(ERenderCaptureType::SubsurfaceColor) || Update.bUpdatedSubsurfaceColor)
+	{
+		ResultSettings->SubsurfaceColorMap = nullptr;
+	}
+	if (!SceneCapture->GetCaptureTypeEnabled(ERenderCaptureType::CombinedMRS) || Update.bUpdatedPackedMRS)
+	{
+		ResultSettings->PackedMRSMap = nullptr;
+	}
 }
 
 
diff --git a/Engine/Plugins/Experimental/MeshModelingToolsetExp/Source/MeshModelingToolsExp/Public/BakeRenderCaptureTool.h b/Engine/Plugins/Experimental/MeshModelingToolsetExp/Source/MeshModelingToolsExp/Public/BakeRenderCaptureTool.h
--- a/Engine/Plugins/Experimental/MeshModelingToolsetExp/Source/MeshModelingToolsExp/Public/BakeRenderCaptureTool.h
+++ b/Engine/Plugins/Experimental/MeshModelingToolsetExp/Source/MeshModelingToolsExp/Public/BakeRenderCaptureTool.h
@@ -6,6 +6,7 @@
 #include "InteractiveToolBuilder.h"
 #include "BaseTools/MultiSelectionMeshEditingTool.h"
 #include "Baking/BakingTypes.h"
+#include "Baking/RenderCaptureFunctions.h"
 #include "BakeMeshAttributeMapsToolBase.h"
 
 // Render Capture algorithm includes
@@ -307,8 +308,8 @@ protected:
 	virtual void GatherAnalytics(FBakeAnalytics::FMeshSettings& Data) override;
 	virtual FString GetAnalyticsEventName() const override { return TEXT("BakeRC"); }
 	// End UBakeMeshAttributeMapsToolBase interface
-	
-	void InvalidateResults();
+
+	void InvalidateResults(UE::Geometry::FRenderCaptureUpdate Update);
 
 	// In this tool we don't call UBakeMeshAttributeMapsToolBase::OnMapsUpdated because it would require e.g, adding
 	// the render capture channels to EBakeMapType.  The implementation is simpler and leads to less coupling if we
@@ -350,13 +351,11 @@ protected:
 
 	// TODO We currently need to compute this on the game thread because the implementation has checks for this
 	TUniquePtr<UE::Geometry::FSceneCapturePhotoSet> SceneCapture;
-	bool bFirstEverSceneCapture = true;
 
-	// If the user cancels a scene capture before the computation completes then the settings which changed to invoke
-	// the capture are reverted to these values
-	UPROPERTY()
-	TObjectPtr<URenderCaptureProperties> ComputedRenderCaptureProperties;
+	// These are used to determine if we need to re-bake results
 	float ComputedValidDepthThreshold;
+	EBakeTextureSamplesPerPixel ComputedSamplesPerPixel = EBakeTextureSamplesPerPixel::Sample1;
+	EBakeTextureResolution ComputedTextureSize = EBakeTextureResolution::Resolution512;
 
 	TMap<int, FText> TargetUVLayerToError;
 
@@ -367,4 +366,6 @@ protected:
 	void RecordAnalytics() const;
 	void GatherAnalytics(const UE::Geometry::FMeshMapBaker& Result);
 
+private:
+	UE::Geometry::FRenderCaptureUpdate UpdateSceneCapture();
 };
\ No newline at end of file
diff --git a/Engine/Plugins/Runtime/MeshModelingToolset/Source/ModelingComponents/Private/Baking/RenderCaptureFunctions.cpp b/Engine/Plugins/Runtime/MeshModelingToolset/Source/ModelingComponents/Private/Baking/RenderCaptureFunctions.cpp
--- a/Engine/Plugins/Runtime/MeshModelingToolset/Source/ModelingComponents/Private/Baking/RenderCaptureFunctions.cpp
+++ b/Engine/Plugins/Runtime/MeshModelingToolset/Source/ModelingComponents/Private/Baking/RenderCaptureFunctions.cpp
@@ -9,6 +9,7 @@
 #include "Engine/Texture2D.h"
 #include "GameFramework/Actor.h"
 #include "Algo/NoneOf.h"
+#include "Algo/AnyOf.h"
 
 
 using namespace UE::Geometry;
@@ -209,32 +210,24 @@ bool FSceneCapturePhotoSetSampler::IsValidCorrespondence(const FMeshMapEvaluator
 TUniquePtr<FSceneCapturePhotoSet> UE::Geometry::CapturePhotoSet(
 	const TArray<TObjectPtr<AActor>>& Actors,
 	const FRenderCaptureOptions& Options,
-	bool bAllowCancel
-)
+	FRenderCaptureUpdate& Update,
+	bool bAllowCancel)
 {
 	TRACE_CPUPROFILER_EVENT_SCOPE(CapturePhotoSet);
 
 	TUniquePtr<FSceneCapturePhotoSet> SceneCapture = MakeUnique<FSceneCapturePhotoSet>();
 
-	UpdatePhotoSet(SceneCapture,Actors, Options, bAllowCancel);
+	Update = UpdatePhotoSets(SceneCapture, Actors, Options, bAllowCancel);
 
 	return SceneCapture;
 }
 
-void UE::Geometry::UpdatePhotoSet(
-	TUniquePtr<FSceneCapturePhotoSet>& SceneCapture,
+FRenderCaptureUpdate UE::Geometry::UpdatePhotoSets(
+	const TUniquePtr<FSceneCapturePhotoSet>& SceneCapture,
 	const TArray<TObjectPtr<AActor>>& Actors,
 	const FRenderCaptureOptions& Options,
-	bool bAllowCancel
-)
+	bool bAllowCancel)
 {
-	double FieldOfView = Options.FieldOfViewDegrees;
-	double NearPlaneDist = Options.NearPlaneDist;
-
-	FImageDimensions CaptureDimensions(Options.RenderCaptureImageSize, Options.RenderCaptureImageSize);
-
-	SceneCapture->SetAllowCancel(bAllowCancel);
-
 	SceneCapture->SetCaptureTypeEnabled(ERenderCaptureType::DeviceDepth, Options.bBakeDeviceDepth);
 	SceneCapture->SetCaptureTypeEnabled(ERenderCaptureType::BaseColor,   Options.bBakeBaseColor);
 	SceneCapture->SetCaptureTypeEnabled(ERenderCaptureType::WorldNormal, Options.bBakeNormalMap);
@@ -258,14 +251,94 @@ void UE::Geometry::UpdatePhotoSet(
 	SceneCapture->SetCaptureConfig(ERenderCaptureType::Opacity,     Config);
 	SceneCapture->SetCaptureConfig(ERenderCaptureType::SubsurfaceColor, Config);
 
-	SceneCapture->SetCaptureSceneActors(Actors[0]->GetWorld(), Actors);
+	UWorld* World = Actors.IsEmpty() ? nullptr : Actors[0]->GetWorld();
+
+	SceneCapture->SetCaptureSceneActors(World, Actors);
 
-	SceneCapture->AddStandardExteriorCapturesFromBoundingBox(
-		CaptureDimensions, FieldOfView, NearPlaneDist,
+	const TArray<FSpatialPhotoParams> SpatialParams = ComputeStandardExteriorSpatialPhotoParameters(
+		World,
+		Actors,
+		FImageDimensions(Options.RenderCaptureImageSize, Options.RenderCaptureImageSize),
+		Options.FieldOfViewDegrees,
+		Options.NearPlaneDist,
 		true, true, true, true, true);
+
+	SceneCapture->SetSpatialPhotoParams(SpatialParams);
+
+	SceneCapture->SetAllowCancel(bAllowCancel);
+
+	// If the provided Options required some photo sets to recompute they will have been cleared by this point.
+	// The updated photo sets are the ones where the photo sets are empty here but are not empty after the Compute call
+	FRenderCaptureUpdate Update;
+
+	Update.bUpdatedBaseColor       = (SceneCapture->GetBaseColorPhotoSet().Num()       == 0);
+	Update.bUpdatedRoughness       = (SceneCapture->GetRoughnessPhotoSet().Num()       == 0);
+	Update.bUpdatedSpecular        = (SceneCapture->GetSpecularPhotoSet().Num()        == 0);
+	Update.bUpdatedMetallic        = (SceneCapture->GetMetallicPhotoSet().Num()        == 0);
+	Update.bUpdatedPackedMRS       = (SceneCapture->GetPackedMRSPhotoSet().Num()       == 0);
+	Update.bUpdatedNormalMap       = (SceneCapture->GetWorldNormalPhotoSet().Num()     == 0);
+	Update.bUpdatedEmissive        = (SceneCapture->GetEmissivePhotoSet().Num()        == 0);
+	Update.bUpdatedOpacity         = (SceneCapture->GetOpacityPhotoSet().Num()         == 0);
+	Update.bUpdatedSubsurfaceColor = (SceneCapture->GetSubsurfaceColorPhotoSet().Num() == 0);
+	Update.bUpdatedDeviceDepth     = (SceneCapture->GetDeviceDepthPhotoSet().Num()     == 0);
+
+	SceneCapture->Compute();
+
+	Update.bUpdatedBaseColor       = (SceneCapture->GetBaseColorPhotoSet().Num()       == 0) != Update.bUpdatedBaseColor;
+	Update.bUpdatedRoughness       = (SceneCapture->GetRoughnessPhotoSet().Num()       == 0) != Update.bUpdatedRoughness;
+	Update.bUpdatedSpecular        = (SceneCapture->GetSpecularPhotoSet().Num()        == 0) != Update.bUpdatedSpecular;
+	Update.bUpdatedMetallic        = (SceneCapture->GetMetallicPhotoSet().Num()        == 0) != Update.bUpdatedMetallic;
+	Update.bUpdatedPackedMRS       = (SceneCapture->GetPackedMRSPhotoSet().Num()       == 0) != Update.bUpdatedPackedMRS;
+	Update.bUpdatedNormalMap       = (SceneCapture->GetWorldNormalPhotoSet().Num()     == 0) != Update.bUpdatedNormalMap;
+	Update.bUpdatedEmissive        = (SceneCapture->GetEmissivePhotoSet().Num()        == 0) != Update.bUpdatedEmissive;
+	Update.bUpdatedOpacity         = (SceneCapture->GetOpacityPhotoSet().Num()         == 0) != Update.bUpdatedOpacity;
+	Update.bUpdatedSubsurfaceColor = (SceneCapture->GetSubsurfaceColorPhotoSet().Num() == 0) != Update.bUpdatedSubsurfaceColor;
+	Update.bUpdatedDeviceDepth     = (SceneCapture->GetDeviceDepthPhotoSet().Num()     == 0) != Update.bUpdatedDeviceDepth ;
+
+	return Update;
 }
 
+FRenderCaptureOptions UE::Geometry::GetComputedPhotoSetOptions(const TUniquePtr<FSceneCapturePhotoSet>& SceneCapture)
+{
+	FRenderCaptureOptions Result;
+
+	Result.bBakeBaseColor       = SceneCapture->GetCaptureTypeEnabled(ERenderCaptureType::BaseColor);
+	Result.bBakeNormalMap       = SceneCapture->GetCaptureTypeEnabled(ERenderCaptureType::WorldNormal);
+	Result.bBakeMetallic        = SceneCapture->GetCaptureTypeEnabled(ERenderCaptureType::Metallic);
+	Result.bBakeRoughness       = SceneCapture->GetCaptureTypeEnabled(ERenderCaptureType::Roughness);
+	Result.bBakeSpecular        = SceneCapture->GetCaptureTypeEnabled(ERenderCaptureType::Specular);
+	Result.bUsePackedMRS        = SceneCapture->GetCaptureTypeEnabled(ERenderCaptureType::CombinedMRS);
+	Result.bBakeEmissive        = SceneCapture->GetCaptureTypeEnabled(ERenderCaptureType::Emissive);
+	Result.bBakeOpacity         = SceneCapture->GetCaptureTypeEnabled(ERenderCaptureType::Opacity);
+	Result.bBakeSubsurfaceColor = SceneCapture->GetCaptureTypeEnabled(ERenderCaptureType::SubsurfaceColor);
+	Result.bBakeDeviceDepth     = SceneCapture->GetCaptureTypeEnabled(ERenderCaptureType::DeviceDepth);
+
+	const TArray<FSpatialPhotoParams> SpatialParams = SceneCapture->GetSpatialPhotoParams();
+
+	if (!SpatialParams.IsEmpty())
+	{
+		ensure(SpatialParams[0].Dimensions.IsSquare());
+		Result.RenderCaptureImageSize = SpatialParams[0].Dimensions.GetWidth();
+		Result.FieldOfViewDegrees = SpatialParams[0].HorzFOVDegrees;
+		Result.NearPlaneDist = SpatialParams[0].NearPlaneDist;
+
+		const bool Values[] = {
+			SceneCapture->GetCaptureConfig(ERenderCaptureType::BaseColor).bAntiAliasing,
+			SceneCapture->GetCaptureConfig(ERenderCaptureType::WorldNormal).bAntiAliasing,
+			SceneCapture->GetCaptureConfig(ERenderCaptureType::Metallic).bAntiAliasing,
+			SceneCapture->GetCaptureConfig(ERenderCaptureType::Roughness).bAntiAliasing,
+			SceneCapture->GetCaptureConfig(ERenderCaptureType::Specular).bAntiAliasing,
+			SceneCapture->GetCaptureConfig(ERenderCaptureType::CombinedMRS).bAntiAliasing,
+			SceneCapture->GetCaptureConfig(ERenderCaptureType::Emissive).bAntiAliasing,
+			SceneCapture->GetCaptureConfig(ERenderCaptureType::Opacity).bAntiAliasing,
+			SceneCapture->GetCaptureConfig(ERenderCaptureType::SubsurfaceColor).bAntiAliasing,
+			SceneCapture->GetCaptureConfig(ERenderCaptureType::DeviceDepth).bAntiAliasing,
+		};
+		Result.bAntiAliasing = Algo::AnyOf(Values);
+	}
 
+	return Result;
+}
 
 
 
@@ -336,7 +409,8 @@ TUniquePtr<FMeshMapBaker> UE::Geometry::MakeRenderCaptureBaker(
 	TSharedPtr<UE::Geometry::FMeshTangentsd, ESPMode::ThreadSafe> BaseMeshTangents,
 	FSceneCapturePhotoSet* SceneCapture,
 	FSceneCapturePhotoSetSampler* Sampler,
-	FRenderCaptureOptions Options,
+	FRenderCaptureOptions PendingBake,
+	int32 TargetUVLayer,
 	EBakeTextureResolution TextureImageSize,
 	EBakeTextureSamplesPerPixel SamplesPerPixel,
 	FRenderCaptureOcclusionHandler* OcclusionHandler)
@@ -365,7 +439,7 @@ TUniquePtr<FMeshMapBaker> UE::Geometry::MakeRenderCaptureBaker(
 	Result->SetDimensions(FImageDimensions(static_cast<int32>(TextureImageSize), static_cast<int32>(TextureImageSize)));
 	Result->SetSamplesPerPixel(static_cast<int32>(SamplesPerPixel));
 	Result->SetFilter(FMeshMapBaker::EBakeFilterType::BSpline);
-	Result->SetTargetMeshUVLayer(Options.TargetUVLayer);
+	Result->SetTargetMeshUVLayer(TargetUVLayer);
 	Result->InteriorSampleCallback = RegisterSampleStats;
 	Result->PostWriteToImageCallback = ComputeAndApplyInfill;
 	Result->SetDetailSampler(Sampler);
@@ -394,44 +468,44 @@ TUniquePtr<FMeshMapBaker> UE::Geometry::MakeRenderCaptureBaker(
 		Result->AddEvaluator(Evaluator);
 		OcclusionHandler->PushInfillRequired(true);
 	};
-	
-	if (Options.bBakeDeviceDepth)
+
+	if (SceneCapture->GetCaptureTypeEnabled(ERenderCaptureType::DeviceDepth) && PendingBake.bBakeDeviceDepth)
 	{
 		AddColorEvaluator(MakeColorEvaluator<ERenderCaptureType::DeviceDepth>(DefaultColorSample, SceneCapture));
 	}
-	if (Options.bBakeBaseColor)
+	if (SceneCapture->GetCaptureTypeEnabled(ERenderCaptureType::BaseColor) && PendingBake.bBakeBaseColor)
 	{
 		AddColorEvaluator(MakeColorEvaluator<ERenderCaptureType::BaseColor>(DefaultColorSample, SceneCapture));
 	}
-	if (Options.bUsePackedMRS)
+	if (SceneCapture->GetCaptureTypeEnabled(ERenderCaptureType::CombinedMRS) && PendingBake.bUsePackedMRS)
 	{
 		AddColorEvaluator(MakeColorEvaluator<ERenderCaptureType::CombinedMRS>(DefaultColorSample, SceneCapture));
 	}
-	if (Options.bBakeRoughness)
+	if (SceneCapture->GetCaptureTypeEnabled(ERenderCaptureType::Roughness) && PendingBake.bBakeRoughness)
 	{
 		AddColorEvaluator(MakeColorEvaluator<ERenderCaptureType::Roughness>(DefaultColorSample, SceneCapture));
 	}
-	if (Options.bBakeMetallic)
+	if (SceneCapture->GetCaptureTypeEnabled(ERenderCaptureType::Metallic) && PendingBake.bBakeMetallic)
 	{
 		AddColorEvaluator(MakeColorEvaluator<ERenderCaptureType::Metallic>(DefaultColorSample, SceneCapture));
 	}
-	if (Options.bBakeSpecular)
+	if (SceneCapture->GetCaptureTypeEnabled(ERenderCaptureType::Specular) && PendingBake.bBakeSpecular)
 	{
 		AddColorEvaluator(MakeColorEvaluator<ERenderCaptureType::Specular>(DefaultColorSample, SceneCapture));
 	}
-	if (Options.bBakeEmissive)
+	if (SceneCapture->GetCaptureTypeEnabled(ERenderCaptureType::Emissive) && PendingBake.bBakeEmissive)
 	{
 		AddColorEvaluator(MakeColorEvaluator<ERenderCaptureType::Emissive>(DefaultColorSample, SceneCapture));
 	}
-	if (Options.bBakeOpacity)
+	if (SceneCapture->GetCaptureTypeEnabled(ERenderCaptureType::Opacity) && PendingBake.bBakeOpacity)
 	{
 		AddColorEvaluator(MakeColorEvaluator<ERenderCaptureType::Opacity>(DefaultColorSample, SceneCapture));
 	}
-	if (Options.bBakeSubsurfaceColor)
+	if (SceneCapture->GetCaptureTypeEnabled(ERenderCaptureType::SubsurfaceColor) && PendingBake.bBakeSubsurfaceColor)
 	{
 		AddColorEvaluator(MakeColorEvaluator<ERenderCaptureType::SubsurfaceColor>(DefaultColorSample, SceneCapture));
 	}
-	if (Options.bBakeNormalMap)
+	if (SceneCapture->GetCaptureTypeEnabled(ERenderCaptureType::WorldNormal) && PendingBake.bBakeNormalMap)
 	{
 		TSharedPtr<FRenderCaptureMapEvaluator<FVector3f>> Evaluator = MakeShared<FRenderCaptureMapEvaluator<FVector3f>>();
 
diff --git a/Engine/Plugins/Runtime/MeshModelingToolset/Source/ModelingComponents/Private/Scene/SceneCapturePhotoSet.cpp b/Engine/Plugins/Runtime/MeshModelingToolset/Source/ModelingComponents/Private/Scene/SceneCapturePhotoSet.cpp
--- a/Engine/Plugins/Runtime/MeshModelingToolset/Source/ModelingComponents/Private/Scene/SceneCapturePhotoSet.cpp
+++ b/Engine/Plugins/Runtime/MeshModelingToolset/Source/ModelingComponents/Private/Scene/SceneCapturePhotoSet.cpp
@@ -24,47 +24,108 @@ FSceneCapturePhotoSet::FSceneCapturePhotoSet()
 
 void FSceneCapturePhotoSet::SetCaptureSceneActors(UWorld* World, const TArray<AActor*>& Actors)
 {
+	if (this->TargetWorld != World || this->VisibleActors != Actors)
+	{
+		// Empty the photo sets because they rendered different actors
+		EmptyAllPhotoSets();
+
+		// Empty the spatial photo parameters because these are computed from the bounding box of the actors
+		PhotoSetParams.Empty();
+	}
 	this->TargetWorld = World;
 	this->VisibleActors = Actors;
 }
 
+TArray<AActor*> FSceneCapturePhotoSet::GetCaptureSceneActors()
+{
+	return VisibleActors;
+}
+
+UWorld* FSceneCapturePhotoSet::GetCaptureTargetWorld()
+{
+	return TargetWorld;
+}
+
+void FSceneCapturePhotoSet::SetSpatialPhotoParams(const TArray<FSpatialPhotoParams>& SpatialParams)
+{
+	// TODO Discard/reset on a per array element level rather than discarding everything when any viewpoint changed
+	if (PhotoSetParams != SpatialParams)
+	{
+		// Empty the photo sets because they were rendered with different viewpoints
+		EmptyAllPhotoSets();
+
+		PhotoSetParams = SpatialParams;
+	}
+}
+
+const TArray<FSpatialPhotoParams>& FSceneCapturePhotoSet::GetSpatialPhotoParams() const
+{
+	return PhotoSetParams;
+}
+
 
 void FSceneCapturePhotoSet::SetCaptureConfig(ERenderCaptureType CaptureType, const FRenderCaptureConfig& Config)
 {
+	auto UpdateCaptureConfig1f = [Config](FRenderCaptureConfig& CaptureConfig, FSpatialPhotoSet1f& PhotoSet)
+	{
+		if (CaptureConfig.bAntiAliasing != Config.bAntiAliasing)
+		{
+			// If we're changing the AntiAliasing state we need to remove any existing photos
+			PhotoSet.Empty();
+		}
+		CaptureConfig = Config;
+	};
+
+	auto UpdateCaptureConfig3f = [Config](FRenderCaptureConfig& CaptureConfig, FSpatialPhotoSet3f& PhotoSet)
+	{
+		if (CaptureConfig.bAntiAliasing != Config.bAntiAliasing)
+		{
+			// If we're changing the AntiAliasing state we need to remove any existing photos
+			PhotoSet.Empty();
+		}
+		CaptureConfig = Config;
+	};
+
 	switch (CaptureType)
 	{
 	case ERenderCaptureType::BaseColor:
-			BaseColorConfig = Config;
-			break;
-		case ERenderCaptureType::WorldNormal:
-			WorldNormalConfig = Config;
-			break;
-		case ERenderCaptureType::Roughness:
-			RoughnessConfig = Config;
-			break;
-		case ERenderCaptureType::Metallic:
-			MetallicConfig = Config;
-			break;
-		case ERenderCaptureType::Specular:
-			SpecularConfig = Config;
-			break;
-		case ERenderCaptureType::Emissive:
-			EmissiveConfig = Config;
-			break;
-		case ERenderCaptureType::Opacity:
-			OpacityConfig = Config;
-			break;
-		case ERenderCaptureType::SubsurfaceColor:
-			SubsurfaceColorConfig = Config;
-			break;
-		case ERenderCaptureType::CombinedMRS:
-			PackedMRSConfig = Config;
-			break;
-		case ERenderCaptureType::DeviceDepth:
-			DeviceDepthConfig = Config;
-			break;
-		default:
-			ensure(false);
+		UpdateCaptureConfig3f(BaseColorConfig, BaseColorPhotoSet);
+		break;
+	case ERenderCaptureType::WorldNormal:
+		UpdateCaptureConfig3f(WorldNormalConfig, WorldNormalPhotoSet);
+		break;
+	case ERenderCaptureType::Roughness:
+		UpdateCaptureConfig1f(RoughnessConfig, RoughnessPhotoSet);
+		break;
+	case ERenderCaptureType::Metallic:
+		UpdateCaptureConfig1f(MetallicConfig, MetallicPhotoSet);
+		break;
+	case ERenderCaptureType::Specular:
+		UpdateCaptureConfig1f(SpecularConfig, SpecularPhotoSet);
+		break;
+	case ERenderCaptureType::Emissive:
+		UpdateCaptureConfig3f(EmissiveConfig, EmissivePhotoSet);
+		break;
+	case ERenderCaptureType::CombinedMRS:
+		UpdateCaptureConfig3f(PackedMRSConfig, PackedMRSPhotoSet);
+		break;
+	case ERenderCaptureType::Opacity:
+		UpdateCaptureConfig1f(OpacityConfig, OpacityPhotoSet);
+		break;
+	case ERenderCaptureType::SubsurfaceColor:
+		UpdateCaptureConfig3f(SubsurfaceColorConfig, SubsurfaceColorPhotoSet);
+		break;
+	case ERenderCaptureType::DeviceDepth:
+		if (DeviceDepthConfig.bAntiAliasing != Config.bAntiAliasing)
+		{
+			// If we're disabling DeviceDepth we need to remove any existing photos as well as the cached view matrices
+			DeviceDepthPhotoSet.Empty();
+			PhotoViewMatricies.Empty();
+		}
+		DeviceDepthConfig = Config;
+		break;
+	default:
+		ensure(false);
 	}
 }
 
@@ -115,96 +176,67 @@ void FSceneCapturePhotoSet::DisableAllCaptureTypes()
 
 void FSceneCapturePhotoSet::SetCaptureTypeEnabled(ERenderCaptureType CaptureType, bool bEnabled)
 {
+	auto UpdateCaptureType1f = [bEnabled](bool& bCaptureEnabled, FSpatialPhotoSet1f& PhotoSet)
+	{
+		if (!bEnabled)
+		{
+			// If we're disabling a CaptureType we need to remove any existing photos
+			PhotoSet.Empty();
+		}
+		bCaptureEnabled = bEnabled;
+	};
+
+	auto UpdateCaptureType3f = [bEnabled](bool& bCaptureEnabled, FSpatialPhotoSet3f& PhotoSet)
+	{
+		if (!bEnabled)
+		{
+			// If we're disabling a CaptureType we need to remove any existing photos
+			PhotoSet.Empty();
+		}
+		bCaptureEnabled = bEnabled;
+	};
+
 	switch (CaptureType)
 	{
 	case ERenderCaptureType::BaseColor:
-		bEnableBaseColor = bEnabled;
+		UpdateCaptureType3f(bEnableBaseColor, BaseColorPhotoSet);
 		break;
 	case ERenderCaptureType::WorldNormal:
-		bEnableWorldNormal = bEnabled;
+		UpdateCaptureType3f(bEnableWorldNormal, WorldNormalPhotoSet);
 		break;
 	case ERenderCaptureType::Roughness:
-		bEnableRoughness = bEnabled;
+		UpdateCaptureType1f(bEnableRoughness, RoughnessPhotoSet);
 		break;
 	case ERenderCaptureType::Metallic:
-		bEnableMetallic = bEnabled;
+		UpdateCaptureType1f(bEnableMetallic, MetallicPhotoSet);
 		break;
 	case ERenderCaptureType::Specular:
-		bEnableSpecular = bEnabled;
+		UpdateCaptureType1f(bEnableSpecular, SpecularPhotoSet);
 		break;
 	case ERenderCaptureType::Emissive:
-		bEnableEmissive = bEnabled;
+		UpdateCaptureType3f(bEnableEmissive, EmissivePhotoSet);
 		break;
 	case ERenderCaptureType::CombinedMRS:
-		bEnablePackedMRS = bEnabled;
+		UpdateCaptureType3f(bEnablePackedMRS, PackedMRSPhotoSet);
 		break;
 	case ERenderCaptureType::Opacity:
-		bEnableOpacity = bEnabled;
+		UpdateCaptureType1f(bEnableOpacity, OpacityPhotoSet);
 		break;
 	case ERenderCaptureType::SubsurfaceColor:
-		bEnableSubsurfaceColor = bEnabled;
+		UpdateCaptureType3f(bEnableSubsurfaceColor, SubsurfaceColorPhotoSet);
 		break;
 	case ERenderCaptureType::DeviceDepth:
+		if (bEnableDeviceDepth != bEnabled)
+		{
+			// If we're disabling DeviceDepth we need to remove any existing photos as well as the cached view matrices
+			DeviceDepthPhotoSet.Empty();
+			PhotoViewMatricies.Empty();
+		}
 		bEnableDeviceDepth = bEnabled;
 		break;
 	default:
 		ensure(false);
 	}
-
-	// If we're disabling a CaptureType we need to remove any existing photos
-	if (bEnabled == false)
-	{
-		switch (CaptureType)
-		{
-		case ERenderCaptureType::BaseColor:
-			BaseColorPhotoSet.Empty();
-			break;
-		case ERenderCaptureType::WorldNormal:
-			WorldNormalPhotoSet.Empty();
-			break;
-		case ERenderCaptureType::Roughness:
-			RoughnessPhotoSet.Empty();
-			break;
-		case ERenderCaptureType::Metallic:
-			MetallicPhotoSet.Empty();
-			break;
-		case ERenderCaptureType::Specular:
-			SpecularPhotoSet.Empty();
-			break;
-		case ERenderCaptureType::Emissive:
-			EmissivePhotoSet.Empty();
-			break;
-		case ERenderCaptureType::Opacity:
-			OpacityPhotoSet.Empty();
-			break;
-		case ERenderCaptureType::SubsurfaceColor:
-			SubsurfaceColorPhotoSet.Empty();
-			break;
-		case ERenderCaptureType::CombinedMRS:
-			PackedMRSPhotoSet.Empty();
-			break;
-		case ERenderCaptureType::DeviceDepth:
-			DeviceDepthPhotoSet.Empty();
-			break;
-		default:
-			ensure(false);
-		}
-	}
-
-	// If all photo sets are disabled we should also clear the photo set params
-	if (!bEnableBaseColor &&
-		!bEnableRoughness &&
-		!bEnableSpecular &&
-		!bEnableMetallic &&
-		!bEnablePackedMRS &&
-		!bEnableWorldNormal &&
-		!bEnableEmissive &&
-		!bEnableOpacity &&
-		!bEnableSubsurfaceColor &&
-		!bEnableDeviceDepth)
-	{
-		PhotoSetParams.Empty();
-	}
 }
 
 bool FSceneCapturePhotoSet::GetCaptureTypeEnabled(ERenderCaptureType CaptureType) const
@@ -237,6 +269,146 @@ bool FSceneCapturePhotoSet::GetCaptureTypeEnabled(ERenderCaptureType CaptureType
 	return false;
 }
 
+void FSceneCapturePhotoSet::Compute()
+{
+	check(this->TargetWorld != nullptr);
+
+	FScopedSlowTask Progress(PhotoSetParams.Num(), LOCTEXT("ComputingViewpoints", "Computing Viewpoints..."));
+	Progress.MakeDialog(bAllowCancel);
+
+	// Unregister all components to remove unwanted proxies from the scene. This was previously the only way to "hide" nanite meshes, now optional.
+	TSet<AActor*> VisibleActorsSet(VisibleActors);
+	TArray<AActor*> ActorsToRegister;
+	if (bEnforceVisibilityViaUnregister)
+	{
+		for (TActorIterator<AActor> Actor(TargetWorld); Actor; ++Actor)
+		{
+			if (!VisibleActorsSet.Contains(*Actor))
+			{
+				Actor->UnregisterAllComponents();
+				ActorsToRegister.Add(*Actor);
+			}
+		}
+	}
+
+	ON_SCOPE_EXIT
+	{
+		// Workaround for Nanite scene proxies visibility
+		// Reregister all components we previously unregistered
+		for (AActor* Actor : ActorsToRegister)
+		{
+			Actor->RegisterAllComponents();
+		}
+	};
+
+	FWorldRenderCapture RenderCapture;
+	RenderCapture.SetWorld(TargetWorld);
+	RenderCapture.SetVisibleActors(VisibleActors);
+	if (bWriteDebugImages)
+	{
+		RenderCapture.SetEnableWriteDebugImage(true, 0, DebugImagesFolderName);
+	}
+
+	for (const FSpatialPhotoParams& Params : PhotoSetParams)
+	{
+		Progress.EnterProgressFrame(1.f);
+		if (Progress.ShouldCancel())
+		{
+			bWasCancelled = true;
+			return;
+		}
+
+		RenderCapture.SetDimensions(Params.Dimensions);
+
+		auto CaptureImageTypeFunc_3f = [this, &Progress, &RenderCapture, &Params]
+			(ERenderCaptureType CaptureType, FSpatialPhotoSet3f& PhotoSet)
+		{
+			if (PhotoSet.Num() < PhotoSetParams.Num())
+			{
+				FSpatialPhoto3f NewPhoto;
+				NewPhoto.Frame = Params.Frame;
+				NewPhoto.NearPlaneDist = Params.NearPlaneDist;
+				NewPhoto.HorzFOVDegrees = Params.HorzFOVDegrees;
+				NewPhoto.Dimensions = Params.Dimensions;
+
+				// TODO Do something with the success boolean returned by RenderCapture.CaptureFromPosition
+				FImageAdapter Image(&NewPhoto.Image);
+				FRenderCaptureConfig Config = GetCaptureConfig(CaptureType);
+				RenderCapture.CaptureFromPosition(CaptureType, NewPhoto.Frame, NewPhoto.HorzFOVDegrees, NewPhoto.NearPlaneDist, Image, Config);
+				PhotoSet.Add(MoveTemp(NewPhoto));
+			}
+
+			Progress.TickProgress();
+		};
+
+		auto CaptureImageTypeFunc_1f = [this, &Progress, &RenderCapture, &Params]
+			(ERenderCaptureType CaptureType, FSpatialPhotoSet1f& PhotoSet)
+		{
+			// Testing NumDirections is how we currently determine if we need to compute this photo
+			if (PhotoSet.Num() < PhotoSetParams.Num())
+			{
+				FSpatialPhoto1f NewPhoto;
+				NewPhoto.Frame = Params.Frame;
+				NewPhoto.NearPlaneDist = Params.NearPlaneDist;
+				NewPhoto.HorzFOVDegrees = Params.HorzFOVDegrees;
+				NewPhoto.Dimensions = Params.Dimensions;
+
+				// TODO Do something with the success boolean returned by RenderCapture.CaptureFromPosition
+				FImageAdapter Image(&NewPhoto.Image);
+				FRenderCaptureConfig Config = GetCaptureConfig(CaptureType);
+				RenderCapture.CaptureFromPosition(CaptureType, NewPhoto.Frame, NewPhoto.HorzFOVDegrees, NewPhoto.NearPlaneDist, Image, Config);
+				PhotoSet.Add(MoveTemp(NewPhoto));
+			}
+
+			Progress.TickProgress();
+		};
+
+		if (bEnableDeviceDepth)
+		{
+			CaptureImageTypeFunc_1f(ERenderCaptureType::DeviceDepth, DeviceDepthPhotoSet);
+			if (PhotoViewMatricies.Num() < PhotoSetParams.Num())
+			{
+				PhotoViewMatricies.Add(RenderCapture.GetLastCaptureViewMatrices());
+			}
+		}
+		if (bEnableBaseColor)
+		{
+			CaptureImageTypeFunc_3f(ERenderCaptureType::BaseColor, BaseColorPhotoSet);
+		}
+		if (bEnableRoughness)
+		{
+			CaptureImageTypeFunc_1f(ERenderCaptureType::Roughness, RoughnessPhotoSet);
+		}
+		if (bEnableSpecular)
+		{
+			CaptureImageTypeFunc_1f(ERenderCaptureType::Specular, SpecularPhotoSet);
+		}
+		if (bEnableMetallic)
+		{
+			CaptureImageTypeFunc_1f(ERenderCaptureType::Metallic, MetallicPhotoSet);
+		}
+		if (bEnablePackedMRS)
+		{
+			CaptureImageTypeFunc_3f(ERenderCaptureType::CombinedMRS, PackedMRSPhotoSet);
+		}
+		if (bEnableWorldNormal)
+		{
+			CaptureImageTypeFunc_3f(ERenderCaptureType::WorldNormal, WorldNormalPhotoSet);
+		}
+		if (bEnableEmissive)
+		{
+			CaptureImageTypeFunc_3f(ERenderCaptureType::Emissive, EmissivePhotoSet);
+		}
+		if (bEnableOpacity)
+		{
+			CaptureImageTypeFunc_1f(ERenderCaptureType::Opacity, OpacityPhotoSet);
+		}
+		if (bEnableSubsurfaceColor)
+		{
+			CaptureImageTypeFunc_3f(ERenderCaptureType::SubsurfaceColor, SubsurfaceColorPhotoSet);
+		}
+	} // end directions loop
+}
 
 
 void FSceneCapturePhotoSet::AddStandardExteriorCapturesFromBoundingBox(
@@ -288,11 +460,11 @@ void FSceneCapturePhotoSet::AddStandardExteriorCapturesFromBoundingBox(
 		Directions.Add(Normalized(FVector3d(1, -1, 0)));
 		Directions.Add(Normalized(FVector3d(-1, -1, 0)));
 	}
+	PRAGMA_DISABLE_DEPRECATION_WARNINGS
 	AddExteriorCaptures(PhotoDimensions, HorizontalFOVDegrees, NearPlaneDist, Directions);
+	PRAGMA_ENABLE_DEPRECATION_WARNINGS
 }
 
-
-
 void FSceneCapturePhotoSet::AddExteriorCaptures(
 	FImageDimensions PhotoDimensions,
 	double HorizontalFOVDegrees,
@@ -411,6 +583,10 @@ void FSceneCapturePhotoSet::AddExteriorCaptures(
 		if (bEnableDeviceDepth)
 		{
 			CaptureImageTypeFunc_1f(ERenderCaptureType::DeviceDepth, DeviceDepthPhotoSet);
+			if (PhotoViewMatricies.Num() < NumDirections)
+			{
+				PhotoViewMatricies.Add(RenderCapture.GetLastCaptureViewMatrices());
+			}
 		}
 		if (bEnableBaseColor)
 		{
@@ -454,10 +630,9 @@ void FSceneCapturePhotoSet::AddExteriorCaptures(
 		// See :SceneCaptureWithExistingCaptures 
 		if (PhotoSetParams.Num() < NumDirections)
 		{
-			Params.ViewMatrices = RenderCapture.GetLastCaptureViewMatrices();
 			PhotoSetParams.Add(Params);
 		}
-	}
+	} // end directions loop
 }
 
 
@@ -554,6 +729,12 @@ bool FSceneCapturePhotoSet::ComputeSampleLocation(
 {
 	double DotTolerance = -0.1;		// dot should be negative for normal pointing towards photo
 
+	if (ValidSampleDepthThreshold > 0)
+	{
+		check(DeviceDepthPhotoSet.Num() == PhotoSetParams.Num());
+		check(PhotoViewMatricies.Num() == PhotoSetParams.Num());
+	}
+
 	PhotoIndex = IndexConstants::InvalidID;
 	PhotoCoords = FVector2d(0., 0.);
 
@@ -610,7 +791,7 @@ bool FSceneCapturePhotoSet::ComputeSampleLocation(
 						{
 							// Compute the pixel position in world space to use it to compute a depth according to the render
 							FVector3d PixelPositionDevice{DeviceXY, DeviceZ};
-							FVector4d PixelPositionWorld = Params.ViewMatrices.GetInvViewProjectionMatrix().TransformPosition(PixelPositionDevice);
+							FVector4d PixelPositionWorld = PhotoViewMatricies[Index].GetInvViewProjectionMatrix().TransformPosition(PixelPositionDevice);
 							PixelPositionWorld /= PixelPositionWorld.W;
 
 							// Compare the depth of the sample with the depth of the pixel and consider the sample invalid
@@ -725,6 +906,23 @@ void FSceneCapturePhotoSet::SetEnableVisibilityByUnregisterMode(bool bEnable)
 	bEnforceVisibilityViaUnregister = bEnable;
 }
 
+void FSceneCapturePhotoSet::EmptyAllPhotoSets()
+{
+	BaseColorPhotoSet.Empty();
+	RoughnessPhotoSet.Empty();
+	SpecularPhotoSet.Empty();
+	MetallicPhotoSet.Empty();
+	PackedMRSPhotoSet.Empty();
+	WorldNormalPhotoSet.Empty();
+	EmissivePhotoSet.Empty();
+	OpacityPhotoSet.Empty();
+	SubsurfaceColorPhotoSet.Empty();
+
+	// For the device depth photo set we have two containers to empty
+	DeviceDepthPhotoSet.Empty();
+	PhotoViewMatricies.Empty();
+}
+
 void FSceneCapturePhotoSet::SetEnableWriteDebugImages(bool bEnable, FString FolderName)
 {
 	bWriteDebugImages = bEnable;
@@ -734,4 +932,95 @@ void FSceneCapturePhotoSet::SetEnableWriteDebugImages(bool bEnable, FString Fold
 	}
 }
 
+TArray<FSpatialPhotoParams> UE::Geometry::ComputeStandardExteriorSpatialPhotoParameters(
+	UWorld* World,
+	const TArray<AActor*>& Actors,
+	FImageDimensions PhotoDimensions,
+	double HorizontalFOVDegrees,
+	double NearPlaneDist,
+	bool bFaces,
+	bool bUpperCorners,
+	bool bLowerCorners,
+	bool bUpperEdges,
+	bool bSideEdges)
+{
+	if (!World || Actors.IsEmpty())
+	{
+		return {};
+	}
+
+	TArray<FVector3d> Directions;
+	if (bFaces)
+	{
+		Directions.Add(FVector3d::UnitX());
+		Directions.Add(-FVector3d::UnitX());
+		Directions.Add(FVector3d::UnitY());
+		Directions.Add(-FVector3d::UnitY());
+		Directions.Add(FVector3d::UnitZ());
+		Directions.Add(-FVector3d::UnitZ());
+	}
+	if (bUpperCorners)
+	{
+		Directions.Add(Normalized(FVector3d(1, 1, -1)));
+		Directions.Add(Normalized(FVector3d(-1, 1, -1)));
+		Directions.Add(Normalized(FVector3d(1, -1, -1)));
+		Directions.Add(Normalized(FVector3d(-1, -1, -1)));
+	}
+	if (bLowerCorners)
+	{
+		Directions.Add(Normalized(FVector3d(1, 1, 1)));
+		Directions.Add(Normalized(FVector3d(-1, 1, 1)));
+		Directions.Add(Normalized(FVector3d(1, -1, 1)));
+		Directions.Add(Normalized(FVector3d(-1, -1, 1)));
+	}
+	if (bUpperEdges)
+	{
+		Directions.Add(Normalized(FVector3d(-1, 0, -1)));
+		Directions.Add(Normalized(FVector3d(1, 0, -1)));
+		Directions.Add(Normalized(FVector3d(0, -1, -1)));
+		Directions.Add(Normalized(FVector3d(0, 1, -1)));
+	}
+	// TODO We are missing bLowerEdges!
+	if (bSideEdges)
+	{
+		Directions.Add(Normalized(FVector3d(1, 1, 0)));
+		Directions.Add(Normalized(FVector3d(-1, 1, 0)));
+		Directions.Add(Normalized(FVector3d(1, -1, 0)));
+		Directions.Add(Normalized(FVector3d(-1, -1, 0)));
+	}
+
+	// Compute a sphere bounding the give actors so we can use it to position the render capture viewpoints
+	// Note: We use FWorldRenderCapture to do this but we are not going to render anything in this function
+	FSphere RenderSphere;
+	{
+		FWorldRenderCapture RenderCapture;
+		RenderCapture.SetVisibleActors(Actors);
+		RenderSphere = RenderCapture.ComputeContainingRenderSphere(HorizontalFOVDegrees);
+	}
+
+	TArray<FSpatialPhotoParams> Result;
+
+	int32 NumDirections = Directions.Num();
+	for (int32 di = 0; di < NumDirections; ++di)
+	{
+		FVector3d ViewDirection = Directions[di];
+		ViewDirection.Normalize();
+
+		FSpatialPhotoParams Params;
+		Params.NearPlaneDist = NearPlaneDist;
+		Params.HorzFOVDegrees = HorizontalFOVDegrees;
+		Params.Dimensions = PhotoDimensions;
+		// TODO Align the frame with the renderer coordinate system then remove the axis swapping in WorldRenderCapture.cpp
+		Params.Frame.AlignAxis(0, ViewDirection);
+		Params.Frame.ConstrainedAlignAxis(2, FVector3d::UnitZ(), Params.Frame.X());
+		Params.Frame.Origin = RenderSphere.Center;
+		Params.Frame.Origin -= RenderSphere.W * Params.Frame.X();
+
+		Result.Add(Params);
+	}
+
+	return Result;
+}
+
+
 #undef LOCTEXT_NAMESPACE
\ No newline at end of file
diff --git a/Engine/Plugins/Runtime/MeshModelingToolset/Source/ModelingComponents/Public/Baking/RenderCaptureFunctions.h b/Engine/Plugins/Runtime/MeshModelingToolset/Source/ModelingComponents/Public/Baking/RenderCaptureFunctions.h
--- a/Engine/Plugins/Runtime/MeshModelingToolset/Source/ModelingComponents/Public/Baking/RenderCaptureFunctions.h
+++ b/Engine/Plugins/Runtime/MeshModelingToolset/Source/ModelingComponents/Public/Baking/RenderCaptureFunctions.h
@@ -105,7 +105,6 @@ struct MODELINGCOMPONENTS_API FRenderCaptureOptions
 	// render capture parameters
 	double FieldOfViewDegrees = 45.0;
 	double NearPlaneDist = 1.0;
-	double ValidSampleDepthThreshold = 0;
 
 	//
 	// Material output settings
@@ -122,38 +121,49 @@ struct MODELINGCOMPONENTS_API FRenderCaptureOptions
 	bool bBakeDeviceDepth = true;
 	
 	bool bUsePackedMRS = true;
-
-	//
-	// Mesh settings
-	//
-
-	//  Which UV layer of the Target mesh (the one we're baking to) should be used
-	int32 TargetUVLayer = 0;
-
-	//
-	// For internal use only
-	//
-
-	// A new MIC derived from this material will be created and assigned to the generated mesh
-	// if null, will use /MeshModelingToolsetExp/Materials/FullMaterialBakePreviewMaterial_PackedMRS instead
-	UMaterialInterface* BakeMaterial = nullptr;
 };
 
+struct MODELINGCOMPONENTS_API FRenderCaptureUpdate
+{
+	// Default to true so that we can use the default value to trigger all update code paths
+	bool bUpdatedBaseColor = true;
+	bool bUpdatedRoughness = true;
+	bool bUpdatedMetallic = true;
+	bool bUpdatedSpecular = true;
+	bool bUpdatedEmissive = true;
+	bool bUpdatedNormalMap = true;
+	bool bUpdatedOpacity = true;
+	bool bUpdatedSubsurfaceColor = true;
+	bool bUpdatedDeviceDepth = true;
+	bool bUpdatedPackedMRS = true;
+};
 
 
 MODELINGCOMPONENTS_API
 TUniquePtr<FSceneCapturePhotoSet> CapturePhotoSet(
 	const TArray<TObjectPtr<AActor>>& Actors,
 	const FRenderCaptureOptions& Options,
+	FRenderCaptureUpdate& Update,
 	bool bAllowCancel);
 
+/**
+* This function efficiently updates the given SceneCapture and returns a struct indicating which channels were updated
+* 
+* - If the requested Options and Actors have already been captured the call is a no-op
+* - If the given Options disable existing capture channels then the call just clears the corresponding photo sets
+* - If the given Options enable a new capture channel then the new photos set are captured and added to the SceneCapture
+* - If the given Actors are different the ones cached in the SceneCapture, or if the given Options changes a parameter
+*   affecting all photo sets (e.g., photo resolution), then all the photo sets are cleared and recomputed
+*/
 MODELINGCOMPONENTS_API
-void UpdatePhotoSet(
-	TUniquePtr<FSceneCapturePhotoSet>& SceneCapture,
+FRenderCaptureUpdate UpdatePhotoSets(
+	const TUniquePtr<FSceneCapturePhotoSet>& SceneCapture,
 	const TArray<TObjectPtr<AActor>>& Actors,
 	const FRenderCaptureOptions& Options,
 	bool bAllowCancel);
 
+MODELINGCOMPONENTS_API
+FRenderCaptureOptions GetComputedPhotoSetOptions(const TUniquePtr<FSceneCapturePhotoSet>& SceneCapture);
 
 
 // Return a render capture baker, note the lifetime of all arguments such match the lifetime of the returned baker
@@ -163,7 +173,8 @@ TUniquePtr<FMeshMapBaker> MakeRenderCaptureBaker(
 	TSharedPtr<FMeshTangentsd, ESPMode::ThreadSafe> BaseMeshTangents,
 	FSceneCapturePhotoSet* SceneCapture,
 	FSceneCapturePhotoSetSampler* Sampler,
-	FRenderCaptureOptions Options,
+	FRenderCaptureOptions PendingBake,
+	int32 TargetUVLayer,
 	EBakeTextureResolution TextureImageSize,
 	EBakeTextureSamplesPerPixel SamplesPerPixel,
 	FRenderCaptureOcclusionHandler* OcclusionHandler);
diff --git a/Engine/Plugins/Runtime/MeshModelingToolset/Source/ModelingComponents/Public/Scene/SceneCapturePhotoSet.h b/Engine/Plugins/Runtime/MeshModelingToolset/Source/ModelingComponents/Public/Scene/SceneCapturePhotoSet.h
--- a/Engine/Plugins/Runtime/MeshModelingToolset/Source/ModelingComponents/Public/Scene/SceneCapturePhotoSet.h
+++ b/Engine/Plugins/Runtime/MeshModelingToolset/Source/ModelingComponents/Public/Scene/SceneCapturePhotoSet.h
@@ -49,8 +49,18 @@ public:
 
 	/**
 	 * Set the target World and set of Actors
+	 * If World != this->World or Actors != this->Actors all existing photo sets are cleared
 	 */
 	void SetCaptureSceneActors(UWorld* World, const TArray<AActor*>& Actors);
+	TArray<AActor*> GetCaptureSceneActors();
+	UWorld* GetCaptureTargetWorld();
+
+	/**
+	 * Set the parameters positioning the virtual cameras used to capture the photo sets
+	 * If SpatialParams != this->PhotoSetParams all existing photo sets are cleared (TODO clear on a more granular level)
+	 */
+	void SetSpatialPhotoParams(const TArray<FSpatialPhotoParams>& SpatialParams);
+	const TArray<FSpatialPhotoParams>& GetSpatialPhotoParams() const;
 
 	/**
 	 * Disable all capture types. By default a standard set of capture types is enabled but, for performance reasons,
@@ -60,20 +70,29 @@ public:
 
 	/**
 	 * Enable/Disable a particular capture type
+	 * If bEnabled == false the corresponding photo set will be cleared
 	 */
 	void SetCaptureTypeEnabled(ERenderCaptureType CaptureType, bool bEnabled);
 	bool GetCaptureTypeEnabled(ERenderCaptureType CaptureType) const;
 
 	/**
 	 * Configure the given capture type
+	 * If Config != the existing capture config the corresponding photo set will be cleared
 	 */
 	void SetCaptureConfig(ERenderCaptureType CaptureType, const FRenderCaptureConfig& Config);
 	FRenderCaptureConfig GetCaptureConfig(ERenderCaptureType CaptureType) const;
 
+	/**
+	 * Render the configured scene, for the configured capture types from the configured viewpoints.
+	 * This function the does most of the work.
+	 */
+	void Compute();
+
 	/**
 	 * Add captures at the corners and face centers of the "view box",
 	 * ie the bounding box that contains the view sphere (see AddExteriorCaptures)
 	 */
+	UE_DEPRECATED(5.2, "AddStandardExteriorCapturesFromBoundingBox is deprecated, please use SetSpatialPhotoParams and Compute instead.")
 	void AddStandardExteriorCapturesFromBoundingBox(
 		FImageDimensions PhotoDimensions,
 		double HorizontalFOVDegrees,
@@ -89,6 +108,7 @@ public:
 	 * will be fully contained inside a square image rendered from locations on the sphere, where
 	 * the view direction is towards the sphere center. The Directions array defines the directions.
 	 */
+	UE_DEPRECATED(5.2, "AddExteriorCaptures is deprecated, please use SetSpatialPhotoParams and Compute instead.")
 	void AddExteriorCaptures(
 		FImageDimensions PhotoDimensions,
 		double HorizontalFOVDegrees,
@@ -242,11 +262,17 @@ protected:
 
 	TArray<FSpatialPhotoParams> PhotoSetParams;
 
+	// This used to unproject the DeviceDepth render capture
+	TArray<FViewMatrices> PhotoViewMatricies;
+
 	bool bWriteDebugImages = false;
 	FString DebugImagesFolderName = TEXT("SceneCapturePhotoSet");
 
 	bool bAllowCancel = false;
 	bool bWasCancelled = false;
+
+private:
+	void EmptyAllPhotoSets();
 };
 
 
@@ -321,6 +347,18 @@ FVector4f FSceneCapturePhotoSet::ComputeSampleNearest(
 	return FVector4f::Zero();
 }
 
+MODELINGCOMPONENTS_API
+TArray<FSpatialPhotoParams> ComputeStandardExteriorSpatialPhotoParameters(
+	UWorld* World,
+	const TArray<AActor*>& Actors,
+	FImageDimensions PhotoDimensions,
+	double HorizontalFOVDegrees,
+	double NearPlaneDist,
+	bool bFaces,
+	bool bUpperCorners,
+	bool bLowerCorners,
+	bool bUpperEdges,
+	bool bSideEdges);
 
 } // end namespace UE::Geometry
 } // end namespace UE
diff --git a/Engine/Source/Runtime/GeometryCore/Public/Image/SpatialPhotoSet.h b/Engine/Source/Runtime/GeometryCore/Public/Image/SpatialPhotoSet.h
--- a/Engine/Source/Runtime/GeometryCore/Public/Image/SpatialPhotoSet.h
+++ b/Engine/Source/Runtime/GeometryCore/Public/Image/SpatialPhotoSet.h
@@ -30,8 +30,14 @@ struct FSpatialPhotoParams
 	/** Pixel dimensions of the photo image */
 	FImageDimensions Dimensions;
 
-	/** Useful to to unproject the DeviceDepth render capture */
-	FViewMatrices ViewMatrices;
+	bool operator==(const FSpatialPhotoParams& Other) const
+	{
+		return Frame.Origin == Other.Frame.Origin &&
+			static_cast<const FVector4d>(Frame.Rotation) == static_cast<const FVector4d>(Other.Frame.Rotation) &&
+			NearPlaneDist == Other.NearPlaneDist &&
+			HorzFOVDegrees == Other.HorzFOVDegrees &&
+			Dimensions == Other.Dimensions;
+	}
 };
 
 /**
