Fixed PatternTool crash on accept in restricted mode.

diff --git a/Engine/Plugins/Experimental/MeshModelingToolsetExp/Source/MeshModelingToolsExp/Private/PatternTool.cpp b/Engine/Plugins/Experimental/MeshModelingToolsetExp/Source/MeshModelingToolsExp/Private/PatternTool.cpp
--- a/Engine/Plugins/Experimental/MeshModelingToolsetExp/Source/MeshModelingToolsExp/Private/PatternTool.cpp
+++ b/Engine/Plugins/Experimental/MeshModelingToolsetExp/Source/MeshModelingToolsExp/Private/PatternTool.cpp
@@ -23,7 +23,6 @@
 #include "Components/StaticMeshComponent.h"
 #include "Components/DynamicMeshComponent.h"
 #include "Components/InstancedStaticMeshComponent.h"
-#include "Engine/StaticMeshActor.h"
 #include "Engine/StaticMesh.h"
 
 #include "ContextObjectStore.h"
@@ -1704,11 +1703,6 @@ void UPatternTool::EmitResults()
 	constexpr bool bPropagateToChildren = true;
 	auto SetActorComponentsVisibility = [bPropagateToChildren](AActor* InActor, const bool bNewVisibility)
 	{
-		if (const AStaticMeshActor* ActorAsStaticMeshActor = Cast<AStaticMeshActor, AActor>(InActor))
-		{
-			ActorAsStaticMeshActor->GetStaticMeshComponent()->SetVisibility(bNewVisibility, bPropagateToChildren);
-		}
-		
 		for (UActorComponent* Component : InActor->GetComponents())
 		{
 			if (UPrimitiveComponent* ActorComponentAsPrimitiveComponent = Cast<UPrimitiveComponent, UActorComponent>(Component))
@@ -1813,10 +1807,8 @@ void UPatternTool::EmitResults()
 		}
 		else if (Element.SourceStaticMesh != nullptr)
 		{
-			UStaticMesh* SetStaticMesh = Element.SourceStaticMesh;
 			UStaticMeshComponent* SourceComponent = Cast<UStaticMeshComponent>(Elements[ElemIdx].SourceComponent);
-			AActor* SourceActor = (SourceComponent != nullptr) ? SourceComponent->GetOwner() : nullptr;
-
+			
 			if (bSeparateActorPerItem)
 			{
 				for (int32 k = 0; k < NumPatternItems; ++k)
@@ -1826,13 +1818,18 @@ void UPatternTool::EmitResults()
 					ComputeWorldTransform(WorldTransform, (FTransform)ElementTransform, (FTransform)PatternTransform);
 					
 					FActorSpawnParameters SpawnInfo;
-					SpawnInfo.Template = SourceActor;
-					AStaticMeshActor* NewActor = GetTargetWorld()->SpawnActor<AStaticMeshActor>(SpawnInfo);
+					AActor* NewActor = GetTargetWorld()->SpawnActor<AActor>(SpawnInfo);
 					if (NewActor != nullptr)
 					{
-						NewActor->SetActorTransform(WorldTransform);
 						NewActors.Add(NewActor);
-						SetActorComponentsVisibility(NewActor, false);
+
+						UStaticMeshComponent* TemplateComponent = DuplicateObject<UStaticMeshComponent>(SourceComponent, NewActor);
+						TemplateComponent->ClearFlags(RF_DefaultSubObject);
+						TemplateComponent->OnComponentCreated();
+						NewActor->SetRootComponent(TemplateComponent);
+						TemplateComponent->RegisterComponent();
+						TemplateComponent->SetWorldTransform(WorldTransform);
+						TemplateComponent->SetVisibility(false, bPropagateToChildren);
 					}
 				}
 			}
@@ -1876,49 +1873,42 @@ void UPatternTool::EmitResults()
 			}
 			else
 			{
-				// Emit a single StaticMeshActor with multiple StaticMeshComponents
-
-				AStaticMeshActor* ParentActor = nullptr;
-				UStaticMeshComponent* TemplateComponent = nullptr;
-				for (int32 k = 0; k < NumPatternItems; ++k)
+				// Emit a single actor with multiple StaticMeshComponents
+				FActorSpawnParameters SpawnInfo;
+				AActor* ParentActor = GetTargetWorld()->SpawnActor<AActor>(SpawnInfo);
+				if (ParentActor != nullptr)
 				{
-					FTransformSRT3d PatternTransform = CurrentPattern[k];
-					FTransform WorldTransform;
-					ComputeWorldTransform(WorldTransform, (FTransform)ElementTransform, (FTransform)PatternTransform);
+					NewActors.Add(ParentActor);
 					
-					if (k == 0)
-					{
-						FActorSpawnParameters SpawnInfo;
-						SpawnInfo.Template = SourceActor;
-						ParentActor = GetTargetWorld()->SpawnActor<AStaticMeshActor>(SpawnInfo);
-						NewActors.Add(ParentActor);
-						ParentActor->SetActorTransform(CurrentStartFrameWorld.ToFTransform());
-						ParentActor->GetStaticMeshComponent()->SetWorldTransform(WorldTransform);
-						ParentActor->GetStaticMeshComponent()->SetVisibility(false, bPropagateToChildren);
-					}
-					else if (k == 1)
-					{
-						TemplateComponent = DuplicateObject<UStaticMeshComponent>(ParentActor->GetStaticMeshComponent(), ParentActor);
-						TemplateComponent->ClearFlags(RF_DefaultSubObject);
-						TemplateComponent->SetupAttachment(ParentActor->GetRootComponent());
-						TemplateComponent->OnComponentCreated();
-						ParentActor->AddInstanceComponent(TemplateComponent);
-						TemplateComponent->RegisterComponent();
-						TemplateComponent->SetWorldTransform( WorldTransform );
-						TemplateComponent->SetVisibility(false, bPropagateToChildren);
-					}
-					else
+					UStaticMeshComponent* TemplateComponent = nullptr;
+					for (int32 k = 0; k < NumPatternItems; ++k)
 					{
-						UStaticMeshComponent* NewCloneComponent = DuplicateObject<UStaticMeshComponent>(TemplateComponent, ParentActor);
-						NewCloneComponent->ClearFlags(RF_DefaultSubObject);
-						NewCloneComponent->SetupAttachment(ParentActor->GetRootComponent());
-						NewCloneComponent->OnComponentCreated();
-						ParentActor->AddInstanceComponent(NewCloneComponent);
-						NewCloneComponent->RegisterComponent();
-						NewCloneComponent->SetWorldTransform( WorldTransform );
-						NewCloneComponent->SetVisibility(false, bPropagateToChildren);
+						FTransformSRT3d PatternTransform = CurrentPattern[k];
+                        FTransform WorldTransform;
+                        ComputeWorldTransform(WorldTransform, (FTransform)ElementTransform, (FTransform)PatternTransform);
+						if (k == 0)
+						{
+							TemplateComponent = DuplicateObject<UStaticMeshComponent>(SourceComponent, ParentActor);
+							TemplateComponent->ClearFlags(RF_DefaultSubObject);
+							TemplateComponent->OnComponentCreated();
+							ParentActor->SetRootComponent(TemplateComponent);
+							TemplateComponent->RegisterComponent();
+							TemplateComponent->SetWorldTransform( WorldTransform );
+                            TemplateComponent->SetVisibility(false, bPropagateToChildren);
+						}
+						else
+						{
+							UStaticMeshComponent* NewCloneComponent = DuplicateObject<UStaticMeshComponent>(TemplateComponent, ParentActor);
+							NewCloneComponent->ClearFlags(RF_DefaultSubObject);
+							NewCloneComponent->SetupAttachment(ParentActor->GetRootComponent());
+							NewCloneComponent->OnComponentCreated();
+							ParentActor->AddInstanceComponent(NewCloneComponent);
+							NewCloneComponent->RegisterComponent();
+							NewCloneComponent->SetWorldTransform( WorldTransform );
+                            NewCloneComponent->SetVisibility(false, bPropagateToChildren);
+						}
 					}
-				}			
+				}
 			}
 		}
 	}
