ModelingTools: handle Mode-Level Mesh Selection input/output in MeshSelectionTool, and add a button-operation to select all triangles w/ the same MaterialID as current selection

diff --git a/Engine/Plugins/Experimental/MeshModelingToolsetExp/Source/MeshModelingToolsExp/Private/MeshSelectionTool.cpp b/Engine/Plugins/Experimental/MeshModelingToolsetExp/Source/MeshModelingToolsExp/Private/MeshSelectionTool.cpp
--- a/Engine/Plugins/Experimental/MeshModelingToolsetExp/Source/MeshModelingToolsExp/Private/MeshSelectionTool.cpp
+++ b/Engine/Plugins/Experimental/MeshModelingToolsetExp/Source/MeshModelingToolsExp/Private/MeshSelectionTool.cpp
@@ -22,6 +22,8 @@
 #include "Properties/MeshUVChannelProperties.h"
 #include "PropertySets/PolygroupLayersProperties.h"
 #include "Polygroups/PolygroupUtil.h"
+#include "Selection/StoredMeshSelectionUtil.h"
+#include "Selections/GeometrySelectionUtil.h"
 
 #include "Algo/MaxElement.h"
 
@@ -38,9 +40,23 @@ UMeshSurfacePointTool* UMeshSelectionToolBuilder::CreateNewTool(const FToolBuild
 {
 	UMeshSelectionTool* SelectionTool = NewObject<UMeshSelectionTool>(SceneState.ToolManager);
 	SelectionTool->SetWorld(SceneState.World);
+
 	return SelectionTool;
 }
 
+void UMeshSelectionToolBuilder::InitializeNewTool(UMeshSurfacePointTool* Tool, const FToolBuilderState& SceneState) const
+{
+	UMeshSurfacePointMeshEditingToolBuilder::InitializeNewTool(Tool, SceneState);
+	if (UMeshSelectionTool* SelectionTool = Cast<UMeshSelectionTool>(Tool))
+	{
+		UE::Geometry::FGeometrySelection Selection;
+		if (UE::Geometry::GetCurrentGeometrySelectionForTarget(SceneState, SelectionTool->GetTarget(), Selection))
+		{
+			SelectionTool->SetGeometrySelection(MoveTemp(Selection));
+		}
+	}
+}
+
 /*
  * Properties
  */
@@ -169,6 +185,16 @@ void UMeshSelectionTool::Setup()
 	MeshStatisticsProperties = NewObject<UMeshStatisticsProperties>(this);
 	AddToolPropertySource(MeshStatisticsProperties);
 
+	// initialize our selection from input selection, if available 
+	if (InputGeometrySelection.IsEmpty() == false)
+	{
+		TArray<int32> AddFaces;
+		UE::Geometry::EnumerateSelectionTriangles(InputGeometrySelection, *Mesh, 
+			[&](int32 TriangleID) { AddFaces.Add(TriangleID); });
+		Selection->AddIndices(EMeshSelectionElementType::Face, AddFaces);
+		OnExternalSelectionChange();
+	}
+
 	RecalculateBrushRadius();
 	UpdateVisualization(true);
 
@@ -193,10 +219,41 @@ UMeshSelectionToolActionPropertySet* UMeshSelectionTool::CreateEditActions()
 
 void UMeshSelectionTool::ApplyShutdownAction(EToolShutdownType ShutdownType)
 {
-	if (bHaveModifiedMesh && ShutdownType == EToolShutdownType::Accept)
+	if (ShutdownType == EToolShutdownType::Accept)
 	{
 		GetToolManager()->BeginUndoTransaction(LOCTEXT("MeshSelectionToolTransactionName", "Edit Mesh"));
-		UE::ToolTarget::CommitDynamicMeshUpdate(Target, *PreviewMesh->GetMesh(), true);
+
+		TArray<int32> SelectedFaces = Selection->GetElements(EMeshSelectionElementType::Face);
+
+		if (bHaveModifiedMesh )
+		{
+			FDynamicMesh3 ResultMesh(*PreviewMesh->GetMesh());
+			if (ResultMesh.IsCompact() == false)
+			{
+				FCompactMaps CompactMaps;
+				ResultMesh.CompactInPlace(&CompactMaps);
+				for (int32& tid : SelectedFaces)
+				{
+					tid = CompactMaps.GetTriangleMapping(tid);
+				}
+			}
+
+			UE::ToolTarget::CommitDynamicMeshUpdate(Target, ResultMesh, true);
+		}
+
+		if (SelectionType == EMeshSelectionElementType::Face)
+		{
+			if ( SelectedFaces.Num() > 0 )
+			{
+				FGeometrySelection OutputSelection;
+				for ( int32 tid : SelectedFaces )
+				{
+					OutputSelection.Selection.Add( FGeoSelectionID::MeshTriangle(tid).Encoded() );
+				}
+				UE::Geometry::SetToolOutputGeometrySelectionForTarget(this, GetTarget(), OutputSelection);
+			}
+		}
+
 		GetToolManager()->EndUndoTransaction();
 	}
 	else if (ShutdownType == EToolShutdownType::Cancel)
@@ -227,6 +284,11 @@ void UMeshSelectionTool::OnShutdown(EToolShutdownType ShutdownType)
 
 
 
+void UMeshSelectionTool::SetGeometrySelection(UE::Geometry::FGeometrySelection&& SelectionIn)
+{
+	InputGeometrySelection = SelectionIn;
+}
+
 
 void UMeshSelectionTool::RegisterActions(FInteractiveToolActionSet& ActionSet)
 {
@@ -1007,6 +1069,10 @@ void UMeshSelectionTool::ApplyAction(EMeshSelectionToolActions ActionType)
 			SelectAll();
 			break;
 	
+		case EMeshSelectionToolActions::SelectAllByMaterial:
+			SelectAllByMaterial();
+			break;
+
 		case EMeshSelectionToolActions::ClearSelection:
 			ClearSelection();
 			break;
@@ -1097,6 +1163,48 @@ void UMeshSelectionTool::SelectAll()
 }
 
 
+void UMeshSelectionTool::SelectAllByMaterial()
+{
+	const FDynamicMesh3* Mesh = PreviewMesh->GetPreviewDynamicMesh();
+
+	if (Selection->Faces.Num() == 0|| Mesh->HasAttributes() == false || Mesh->Attributes()->HasMaterialID() == false)
+	{
+		SelectAll();
+		return;
+	}
+
+	BeginChange(true);
+	
+	TSet<int32> SelectedMaterialIDs;
+	const FDynamicMeshMaterialAttribute* MaterialIDs = Mesh->Attributes()->GetMaterialID();
+	for (int tid : Mesh->TriangleIndicesItr())
+	{
+		if ( SelectedTriangles[tid] )
+		{
+			SelectedMaterialIDs.Add( MaterialIDs->GetValue(tid) );
+		}
+	}
+
+	TArray<int32> AddFaces;
+	for (int tid : Mesh->TriangleIndicesItr())
+	{
+		if ( SelectedTriangles[tid] == false && SelectedMaterialIDs.Contains(MaterialIDs->GetValue(tid)) )
+		{
+			AddFaces.Add(tid);
+		}
+	}
+	
+	ActiveSelectionChange->Add(AddFaces);
+	Selection->AddIndices(EMeshSelectionElementType::Face, AddFaces);
+	
+	TUniquePtr<FToolCommandChange> SelectionChange = EndChange();
+
+	GetToolManager()->EmitObjectChange(Selection, MoveTemp(SelectionChange), LOCTEXT("ExpandToMaterial", "Select Materials"));
+
+	OnExternalSelectionChange();
+}
+
+
 
 void UMeshSelectionTool::ClearSelection()
 {
diff --git a/Engine/Plugins/Experimental/MeshModelingToolsetExp/Source/MeshModelingToolsExp/Public/MeshSelectionTool.h b/Engine/Plugins/Experimental/MeshModelingToolsetExp/Source/MeshModelingToolsExp/Public/MeshSelectionTool.h
--- a/Engine/Plugins/Experimental/MeshModelingToolsetExp/Source/MeshModelingToolsExp/Public/MeshSelectionTool.h
+++ b/Engine/Plugins/Experimental/MeshModelingToolsetExp/Source/MeshModelingToolsExp/Public/MeshSelectionTool.h
@@ -10,6 +10,7 @@
 #include "Changes/MeshSelectionChange.h"
 #include "DynamicMesh/DynamicMeshOctree3.h"
 #include "Polygroups/PolygroupSet.h"
+#include "Selections/GeometrySelection.h"
 #include "MeshSelectionTool.generated.h"
 
 class UMeshStatisticsProperties;
@@ -28,6 +29,7 @@ class MESHMODELINGTOOLSEXP_API UMeshSelectionToolBuilder : public UMeshSurfacePo
 
 public:
 	virtual UMeshSurfacePointTool* CreateNewTool(const FToolBuilderState& SceneState) const override;
+	virtual void InitializeNewTool(UMeshSurfacePointTool* Tool, const FToolBuilderState& SceneState) const override;
 };
 
 
@@ -39,6 +41,7 @@ enum class EMeshSelectionToolActions
 	NoAction,
 
 	SelectAll,
+	SelectAllByMaterial,
 	ClearSelection,
 	InvertSelection,
 	GrowSelection,
@@ -99,6 +102,7 @@ public:
 		PostAction(EMeshSelectionToolActions::SelectAll);
 	}
 
+
 	/** Invert the active triangle selection */
 	UFUNCTION(CallInEditor, Category = SelectionEdits, meta = (DisplayPriority = 2))
 	void Invert()
@@ -141,12 +145,19 @@ public:
 		PostAction(EMeshSelectionToolActions::SelectLargestComponentByArea);
 	}
 
-	/** Optimize the selection border */
+	/** Optimize the selection border by removing "fin" triangles and including "ear" triangles */
 	UFUNCTION(CallInEditor, Category = SelectionEdits, meta = (DisplayPriority = 8))
 	void OptimizeBorder()
 	{
 		PostAction(EMeshSelectionToolActions::OptimizeSelection);
 	}
+
+	/** Expand the selection to include all triangles with Materials matching the Materials on the currently selected triangles */
+	UFUNCTION(CallInEditor, Category = SelectionEdits, meta = (DisplayPriority = 9))
+	void ExpandToMaterials()
+	{
+		PostAction(EMeshSelectionToolActions::SelectAllByMaterial);
+	}
 };
 
 
@@ -315,7 +326,8 @@ public:
 
 	virtual bool HasCancel() const override { return true; }
 	virtual bool HasAccept() const override { return true; }
-	virtual bool CanAccept() const override { return Super::CanAccept() && bHaveModifiedMesh; }
+	//virtual bool CanAccept() const override { return Super::CanAccept() && bHaveModifiedMesh; }
+	virtual bool CanAccept() const override { return Super::CanAccept(); }		// allow selection w/o modified mesh to allow for use as just a selection tool
 
 	// UBaseBrushTool overrides
 	virtual bool HitTest(const FRay& Ray, FHitResult& OutHit) override;
@@ -332,6 +344,10 @@ public:
 	virtual bool CanCurrentlyNestedCancel() override;
 	virtual bool ExecuteNestedCancelCommand() override;
 
+	// input selection support
+	virtual void SetGeometrySelection(UE::Geometry::FGeometrySelection&& SelectionIn);
+
+
 public:
 
 	virtual void RequestAction(EMeshSelectionToolActions ActionType);
@@ -368,6 +384,9 @@ protected:
 	virtual void OnShutdown(EToolShutdownType ShutdownType) override;
 
 protected:
+
+	UE::Geometry::FGeometrySelection InputGeometrySelection;
+
 	UPROPERTY()
 	TObjectPtr<UMeshSelectionSet> Selection;
 
@@ -430,6 +449,7 @@ protected:
 	void InvertSelection();
 	void GrowShrinkSelection(bool bGrow);
 	void ExpandToConnected();
+	void SelectAllByMaterial();
 
 	void SelectLargestComponent(bool bWeightByArea);
 	void OptimizeSelection();
