Remove un-needed 'handle inputs' options on Msh2Vol tool, since it's a single-input tool

diff --git a/Engine/Plugins/Experimental/MeshModelingToolsetExp/Source/MeshModelingToolsEditorOnlyExp/Private/MeshToVolumeTool.cpp b/Engine/Plugins/Experimental/MeshModelingToolsetExp/Source/MeshModelingToolsEditorOnlyExp/Private/MeshToVolumeTool.cpp
--- a/Engine/Plugins/Experimental/MeshModelingToolsetExp/Source/MeshModelingToolsEditorOnlyExp/Private/MeshToVolumeTool.cpp
+++ b/Engine/Plugins/Experimental/MeshModelingToolsetExp/Source/MeshModelingToolsEditorOnlyExp/Private/MeshToVolumeTool.cpp
@@ -82,7 +82,7 @@ void UMeshToVolumeTool::Setup()
 							{ bVolumeValid = false; });
 
 
-	HandleSourcesProperties = NewObject<UOnAcceptHandleSourcesProperties>(this);
+	HandleSourcesProperties = NewObject<UOnAcceptHandleSourcesPropertiesSingle>(this);
 	HandleSourcesProperties->RestoreProperties(this);
 	AddToolPropertySource(HandleSourcesProperties);
 
diff --git a/Engine/Plugins/Experimental/MeshModelingToolsetExp/Source/MeshModelingToolsEditorOnlyExp/Public/MeshToVolumeTool.h b/Engine/Plugins/Experimental/MeshModelingToolsetExp/Source/MeshModelingToolsEditorOnlyExp/Public/MeshToVolumeTool.h
--- a/Engine/Plugins/Experimental/MeshModelingToolsetExp/Source/MeshModelingToolsEditorOnlyExp/Public/MeshToVolumeTool.h
+++ b/Engine/Plugins/Experimental/MeshModelingToolsetExp/Source/MeshModelingToolsEditorOnlyExp/Public/MeshToVolumeTool.h
@@ -90,7 +90,7 @@ protected:
 	TObjectPtr<UMeshToVolumeToolProperties> Settings;
 
 	UPROPERTY()
-	TObjectPtr<UOnAcceptHandleSourcesProperties> HandleSourcesProperties;
+	TObjectPtr<UOnAcceptHandleSourcesPropertiesSingle> HandleSourcesProperties;
 
 	UPROPERTY()
 	TObjectPtr<UPreviewMesh> PreviewMesh;
