ModelingTools: Add OffsetMeshSelectionTool and RetriangulateGeometrySelectionCommand

diff --git a/Engine/Plugins/Editor/ModelingToolsEditorMode/Source/ModelingToolsEditorMode/Private/ModelingToolsEditorMode.cpp b/Engine/Plugins/Editor/ModelingToolsEditorMode/Source/ModelingToolsEditorMode/Private/ModelingToolsEditorMode.cpp
--- a/Engine/Plugins/Editor/ModelingToolsEditorMode/Source/ModelingToolsEditorMode/Private/ModelingToolsEditorMode.cpp
+++ b/Engine/Plugins/Editor/ModelingToolsEditorMode/Source/ModelingToolsEditorMode/Private/ModelingToolsEditorMode.cpp
@@ -92,6 +92,7 @@
 #include "PatternTool.h"
 
 #include "Polymodeling/ExtrudeMeshSelectionTool.h"
+#include "Polymodeling/OffsetMeshSelectionTool.h"
 
 #include "Physics/PhysicsInspectorTool.h"
 #include "Physics/SetCollisionGeometryTool.h"
@@ -106,6 +107,7 @@
 // commands
 #include "Commands/DeleteGeometrySelectionCommand.h"
 #include "Commands/DisconnectGeometrySelectionCommand.h"
+#include "Commands/RetriangulateGeometrySelectionCommand.h"
 #include "Commands/ModifyGeometrySelectionCommand.h"
 
 
@@ -717,39 +719,43 @@ void UModelingToolsEditorMode::Enter()
 
 
 	// this function registers and tracks an active UGeometrySelectionEditCommand and it's associated UICommand
-	auto RegisterSelectionTool = [&](TSharedPtr<FUICommandInfo> UICommand, FString ToolIdentifier, UInteractiveToolBuilder* Builder)
+	auto RegisterSelectionTool = [&](TSharedPtr<FUICommandInfo> UICommand, FString ToolIdentifier, UInteractiveToolBuilder* Builder, bool bRequiresActiveTarget, bool bRequiresSelection)
 	{
 		UEditorInteractiveToolsContext* UseToolsContext = GetInteractiveToolsContext(EToolsContextScope::EdMode);
 		const TSharedRef<FUICommandList>& CommandList = Toolkit->GetToolkitCommands();
 		UseToolsContext->ToolManager->RegisterToolType(ToolIdentifier, Builder);
 		CommandList->MapAction(UICommand,
 			FExecuteAction::CreateUObject(UseToolsContext, &UEdModeInteractiveToolsContext::StartTool, ToolIdentifier),
-			FCanExecuteAction::CreateWeakLambda(UseToolsContext, [this, ToolIdentifier, UseToolsContext]() {
+			FCanExecuteAction::CreateWeakLambda(UseToolsContext, [this, ToolIdentifier, UseToolsContext, bRequiresActiveTarget, bRequiresSelection]() {
 				return ShouldToolStartBeAllowed(ToolIdentifier) &&
-					GetSelectionManager()->HasSelection() &&
+					( GetSelectionManager()->HasActiveTargets() || bRequiresActiveTarget == false) &&
+					( GetSelectionManager()->HasSelection() || bRequiresSelection == false ) &&
+					( GetSelectionManager()->GetMeshTopologyMode() != UGeometrySelectionManager::EMeshTopologyMode::None ) &&
 					UseToolsContext->ToolManager->CanActivateTool(EToolSide::Mouse, ToolIdentifier);
 			}),
 			FIsActionChecked::CreateUObject(UseToolsContext, &UEdModeInteractiveToolsContext::IsToolActive, EToolSide::Mouse, ToolIdentifier),
-			FIsActionButtonVisible::CreateUObject(GetSelectionManager(), &UGeometrySelectionManager::HasSelection),
+			//FIsActionButtonVisible::CreateUObject(GetSelectionManager(), &UGeometrySelectionManager::HasSelection),
+			FIsActionButtonVisible::CreateWeakLambda(UseToolsContext, [this, bRequiresActiveTarget, bRequiresSelection]() {
+				return ( GetSelectionManager()->HasActiveTargets() || bRequiresActiveTarget == false) &&
+					( GetSelectionManager()->HasSelection() || bRequiresSelection == false ) &&
+					( GetSelectionManager()->GetMeshTopologyMode() != UGeometrySelectionManager::EMeshTopologyMode::None );
+			}),
 			EUIActionRepeatMode::RepeatDisabled);
 	};
 
 	// register mesh-selection-driven tools
-	RegisterSelectionTool(ToolManagerCommands.BeginSelectionAction_Extrude, TEXT("BeginSelectionExtrudeTool"), NewObject<UExtrudeMeshSelectionToolBuilder>());
+	RegisterSelectionTool(ToolManagerCommands.BeginSelectionAction_Extrude, TEXT("BeginSelectionExtrudeTool"), NewObject<UExtrudeMeshSelectionToolBuilder>(), true, true);
+	RegisterSelectionTool(ToolManagerCommands.BeginSelectionAction_Offset, TEXT("BeginSelectionOffsetTool"), NewObject<UOffsetMeshSelectionToolBuilder>(), true, true);
+
+	RegisterSelectionTool(ToolManagerCommands.BeginPolyModelTool_PolyEd, TEXT("BeginSelectionPolyEdTool"), NewObject<UMeshSelectionToolBuilder>(), true, false);
+	RegisterSelectionTool(ToolManagerCommands.BeginPolyModelTool_TriSel, TEXT("BeginSelectionTriEdTool"), NewObject<UOffsetMeshSelectionToolBuilder>(), true, false);
 
 
 	auto RegisterPolyModelActionTool = [&](EEditMeshPolygonsToolActions Action, TSharedPtr<FUICommandInfo> UICommand, FString StringName, bool bRequiresSelection)
 	{
 		UEditMeshPolygonsActionModeToolBuilder* ActionModeBuilder = NewObject<UEditMeshPolygonsActionModeToolBuilder>();
 		ActionModeBuilder->StartupAction = Action;
-		if ( bRequiresSelection )
-		{
-			RegisterSelectionTool(UICommand, StringName, ActionModeBuilder);
-		}
-		else
-		{
-			RegisterTool(UICommand, StringName, ActionModeBuilder);
-		}
+		RegisterSelectionTool(UICommand, StringName, ActionModeBuilder, true, bRequiresSelection);
 	};
 	RegisterPolyModelActionTool(EEditMeshPolygonsToolActions::Inset, ToolManagerCommands.BeginPolyModelTool_Inset, TEXT("PolyEdit_Inset"), true);
 	RegisterPolyModelActionTool(EEditMeshPolygonsToolActions::Outset, ToolManagerCommands.BeginPolyModelTool_Outset, TEXT("PolyEdit_Outset"), true);
@@ -808,6 +814,7 @@ void UModelingToolsEditorMode::Enter()
 	// create and register InteractiveCommands for mesh selections
 	RegisterSelectionCommand(NewObject<UDeleteGeometrySelectionCommand>(), ToolManagerCommands.BeginSelectionAction_Delete, false);
 	RegisterSelectionCommand(NewObject<UDisconnectGeometrySelectionCommand>(), ToolManagerCommands.BeginSelectionAction_Disconnect, false);
+	RegisterSelectionCommand(NewObject<URetriangulateGeometrySelectionCommand>(), ToolManagerCommands.BeginSelectionAction_Retriangulate, false);
 	RegisterSelectionCommand(NewObject<UModifyGeometrySelectionCommand>(), ToolManagerCommands.BeginSelectionAction_SelectAll, true);
 	RegisterSelectionCommand(NewObject<UModifyGeometrySelectionCommand_Invert>(), ToolManagerCommands.BeginSelectionAction_Invert, true);
 	RegisterSelectionCommand(NewObject<UModifyGeometrySelectionCommand_ExpandToConnected>(), ToolManagerCommands.BeginSelectionAction_ExpandToConnected, true);
diff --git a/Engine/Plugins/Editor/ModelingToolsEditorMode/Source/ModelingToolsEditorMode/Private/ModelingToolsEditorModeStyle.cpp b/Engine/Plugins/Editor/ModelingToolsEditorMode/Source/ModelingToolsEditorMode/Private/ModelingToolsEditorModeStyle.cpp
--- a/Engine/Plugins/Editor/ModelingToolsEditorMode/Source/ModelingToolsEditorMode/Private/ModelingToolsEditorModeStyle.cpp
+++ b/Engine/Plugins/Editor/ModelingToolsEditorMode/Source/ModelingToolsEditorMode/Private/ModelingToolsEditorModeStyle.cpp
@@ -331,6 +331,12 @@ void FModelingToolsEditorModeStyle::Initialize()
 
 		StyleSet->Set("ModelingToolsManagerCommands.BeginSelectionAction_Delete", new FSlateImageBrush(StyleSet->RootToCoreContentDir(TEXT("../Editor/Slate/Icons/GeneralTools/Delete_40x.png")), Icon20x20));
 
+		StyleSet->Set("ModelingToolsManagerCommands.BeginPolyModelTool_PolyEd", 				new IMAGE_PLUGIN_BRUSH("Icons/PolyEdit_40x", 		Icon20x20));
+		StyleSet->Set("ModelingToolsManagerCommands.BeginPolyModelTool_PolyEd.Small", 			new IMAGE_PLUGIN_BRUSH("Icons/PolyEdit_40x", 		Icon20x20));
+		StyleSet->Set("ModelingToolsManagerCommands.BeginPolyModelTool_TriSel", 				new IMAGE_PLUGIN_BRUSH("Icons/MeshSelect_40x",		Icon20x20));
+		StyleSet->Set("ModelingToolsManagerCommands.BeginPolyModelTool_TriSel.Small", 			new IMAGE_PLUGIN_BRUSH("Icons/MeshSelect_40x",		Icon20x20));
+
+
 		StyleSet->Set("ModelingModeSelection.More_Right",  new IMAGE_BRUSH_SVG("Icons/SelectionToolbar_More", Icon20x20));
 		StyleSet->Set("ModelingModeSelection.Edits_Right",  new IMAGE_BRUSH_SVG("Icons/SelectionToolbar_Edits", Icon20x20));
 
diff --git a/Engine/Plugins/Editor/ModelingToolsEditorMode/Source/ModelingToolsEditorMode/Private/ModelingToolsEditorModeToolkit.cpp b/Engine/Plugins/Editor/ModelingToolsEditorMode/Source/ModelingToolsEditorMode/Private/ModelingToolsEditorModeToolkit.cpp
--- a/Engine/Plugins/Editor/ModelingToolsEditorMode/Source/ModelingToolsEditorMode/Private/ModelingToolsEditorModeToolkit.cpp
+++ b/Engine/Plugins/Editor/ModelingToolsEditorMode/Source/ModelingToolsEditorMode/Private/ModelingToolsEditorModeToolkit.cpp
@@ -1011,15 +1011,21 @@ void FModelingToolsEditorModeToolkit::BuildToolPalette(FName PaletteIndex, class
 	else if (PaletteIndex == SelectionActionsTabName)
 	{
 		ToolbarBuilder.AddToolBarButton(Commands.BeginSelectionAction_Delete);
-		ToolbarBuilder.AddToolBarButton(Commands.BeginSelectionAction_Disconnect);
+		//ToolbarBuilder.AddToolBarButton(Commands.BeginSelectionAction_Disconnect);		// disabled for 5.2, available via TriSel Tool
 		ToolbarBuilder.AddToolBarButton(Commands.BeginSelectionAction_Extrude);
+		ToolbarBuilder.AddToolBarButton(Commands.BeginSelectionAction_Offset);
 
 		ToolbarBuilder.AddToolBarButton(Commands.BeginPolyModelTool_PushPull);
-		ToolbarBuilder.AddToolBarButton(Commands.BeginPolyModelTool_InsertEdgeLoop);
-		ToolbarBuilder.AddToolBarButton(Commands.BeginPolyModelTool_CutFaces);
 		ToolbarBuilder.AddToolBarButton(Commands.BeginPolyModelTool_Inset);
 		ToolbarBuilder.AddToolBarButton(Commands.BeginPolyModelTool_Outset);
+		ToolbarBuilder.AddToolBarButton(Commands.BeginPolyModelTool_CutFaces);
 		ToolbarBuilder.AddToolBarButton(Commands.BeginPolyModelTool_Bevel);
+
+		ToolbarBuilder.AddToolBarButton(Commands.BeginPolyModelTool_InsertEdgeLoop);
+		ToolbarBuilder.AddToolBarButton(Commands.BeginSelectionAction_Retriangulate);
+
+		ToolbarBuilder.AddToolBarButton(Commands.BeginPolyModelTool_PolyEd);
+		ToolbarBuilder.AddToolBarButton(Commands.BeginPolyModelTool_TriSel);
 		
 	}
 	else if (PaletteIndex == PrimitiveTabName)
diff --git a/Engine/Plugins/Editor/ModelingToolsEditorMode/Source/ModelingToolsEditorMode/Private/ModelingToolsManagerActions.cpp b/Engine/Plugins/Editor/ModelingToolsEditorMode/Source/ModelingToolsEditorMode/Private/ModelingToolsManagerActions.cpp
--- a/Engine/Plugins/Editor/ModelingToolsEditorMode/Source/ModelingToolsEditorMode/Private/ModelingToolsManagerActions.cpp
+++ b/Engine/Plugins/Editor/ModelingToolsEditorMode/Source/ModelingToolsEditorMode/Private/ModelingToolsManagerActions.cpp
@@ -182,6 +182,9 @@ void FModelingToolsManagerCommands::RegisterCommands()
 	REGISTER_MODELING_TOOL_COMMAND(BeginPolyModelTool_PushPull, "PushPull", "PolyGroup Push Pull Faces Tool");
 	REGISTER_MODELING_TOOL_COMMAND(BeginPolyModelTool_Bevel, "Bevel", "PolyGroup Bevel Tool");
 
+	REGISTER_MODELING_TOOL_COMMAND(BeginPolyModelTool_PolyEd, "PolyEd", "Edit meshes via PolyGroups");
+	REGISTER_MODELING_TOOL_COMMAND(BeginPolyModelTool_TriSel, "TriSel", "Select and edit mesh triangles");
+
 	REGISTER_MODELING_TOOL_COMMAND(MeshSelectionModeAction_NoSelection, "None", "Disable Geometry Selection");
 	REGISTER_MODELING_TOOL_COMMAND(MeshSelectionModeAction_MeshTriangles, "Tris", "Select Mesh Triangles");
 	REGISTER_MODELING_TOOL_COMMAND(MeshSelectionModeAction_MeshVertices, "Verts", "Select Mesh Vertices");
@@ -192,8 +195,10 @@ void FModelingToolsManagerCommands::RegisterCommands()
 
 	REGISTER_MODELING_ACTION_COMMAND(BeginSelectionAction_Delete, "Delete", "Delete Selection");
 	REGISTER_MODELING_ACTION_COMMAND(BeginSelectionAction_Disconnect, "Discon", "Disconnect Selection");
+	REGISTER_MODELING_ACTION_COMMAND(BeginSelectionAction_Retriangulate, "Clean", "Retriangulate Selection");
 
 	REGISTER_MODELING_TOOL_COMMAND(BeginSelectionAction_Extrude, "Extrude", "Extrude Selection");
+	REGISTER_MODELING_TOOL_COMMAND(BeginSelectionAction_Offset, "Offset", "Offset Selection");
 
 	REGISTER_MODELING_ACTION_COMMAND(BeginSelectionAction_SelectAll, "Select All", "Select All Elements");
 	REGISTER_MODELING_ACTION_COMMAND(BeginSelectionAction_ExpandToConnected, "Expand To Connected", "Expand Selection to Geometrically-Connected Elements");
diff --git a/Engine/Plugins/Editor/ModelingToolsEditorMode/Source/ModelingToolsEditorMode/Public/ModelingToolsManagerActions.h b/Engine/Plugins/Editor/ModelingToolsEditorMode/Source/ModelingToolsEditorMode/Public/ModelingToolsManagerActions.h
--- a/Engine/Plugins/Editor/ModelingToolsEditorMode/Source/ModelingToolsEditorMode/Public/ModelingToolsManagerActions.h
+++ b/Engine/Plugins/Editor/ModelingToolsEditorMode/Source/ModelingToolsEditorMode/Public/ModelingToolsManagerActions.h
@@ -139,6 +139,8 @@ public:
 	TSharedPtr<FUICommandInfo> BeginPolyModelTool_PushPull;
 	TSharedPtr<FUICommandInfo> BeginPolyModelTool_Bevel;
 	
+	TSharedPtr<FUICommandInfo> BeginPolyModelTool_PolyEd;
+	TSharedPtr<FUICommandInfo> BeginPolyModelTool_TriSel;
 
 
 	// selection state actions
@@ -153,6 +155,7 @@ public:
 
 	TSharedPtr<FUICommandInfo> BeginSelectionAction_Delete;
 	TSharedPtr<FUICommandInfo> BeginSelectionAction_Disconnect;
+	TSharedPtr<FUICommandInfo> BeginSelectionAction_Retriangulate;
 
 	TSharedPtr<FUICommandInfo> BeginSelectionAction_SelectAll;
 	TSharedPtr<FUICommandInfo> BeginSelectionAction_ExpandToConnected;
@@ -162,6 +165,7 @@ public:
 	TSharedPtr<FUICommandInfo> BeginSelectionAction_Contract;
 
 	TSharedPtr<FUICommandInfo> BeginSelectionAction_Extrude;
+	TSharedPtr<FUICommandInfo> BeginSelectionAction_Offset;
 
 	// For connecting ModelingMode with UVEditor plugin
 	TSharedPtr<FUICommandInfo> LaunchUVEditor;
diff --git a/Engine/Plugins/Experimental/MeshModelingToolsetExp/Source/MeshModelingToolsExp/Private/Polymodeling/ExtrudeMeshSelectionTool.cpp b/Engine/Plugins/Experimental/MeshModelingToolsetExp/Source/MeshModelingToolsExp/Private/Polymodeling/ExtrudeMeshSelectionTool.cpp
--- a/Engine/Plugins/Experimental/MeshModelingToolsetExp/Source/MeshModelingToolsExp/Private/Polymodeling/ExtrudeMeshSelectionTool.cpp
+++ b/Engine/Plugins/Experimental/MeshModelingToolsetExp/Source/MeshModelingToolsExp/Private/Polymodeling/ExtrudeMeshSelectionTool.cpp
@@ -101,6 +101,11 @@ public:
 };
 
 
+UExtrudeMeshSelectionTool::UExtrudeMeshSelectionTool()
+{
+	SetToolDisplayName(LOCTEXT("ToolName", "Extrude Selection"));
+}
+
 
 void UExtrudeMeshSelectionTool::Setup()
 {
diff --git a/Engine/Plugins/Experimental/MeshModelingToolsetExp/Source/MeshModelingToolsExp/Private/Polymodeling/OffsetMeshSelectionTool.cpp b/Engine/Plugins/Experimental/MeshModelingToolsetExp/Source/MeshModelingToolsExp/Private/Polymodeling/OffsetMeshSelectionTool.cpp
new file mode 100644
--- /dev/null
+++ b/Engine/Plugins/Experimental/MeshModelingToolsetExp/Source/MeshModelingToolsExp/Private/Polymodeling/OffsetMeshSelectionTool.cpp
@@ -0,0 +1,394 @@
+// Copyright Epic Games, Inc. All Rights Reserved.
+
+#include "Polymodeling/OffsetMeshSelectionTool.h"
+#include "InteractiveToolManager.h"
+#include "MeshOpPreviewHelpers.h" // UMeshOpPreviewWithBackgroundCompute
+#include "ToolSetupUtil.h"
+#include "ModelingToolTargetUtil.h"
+#include "Selections/GeometrySelectionUtil.h"
+#include "Selection/StoredMeshSelectionUtil.h"
+#include "ToolSceneQueriesUtil.h"
+#include "MeshQueries.h"
+
+#include "Selections/MeshFaceSelection.h"
+#include "DynamicSubmesh3.h"
+#include "Operations/MeshRegionOperator.h"
+#include "Operations/OffsetMeshRegion.h"
+
+#include "ModelingOperators.h"
+#include "PolyModelingOps/RegionOffsetOp.h"
+#include "Operations/PolyModeling/PolyModelingMaterialUtil.h"
+#include "Operations/PolyModeling/PolyModelingFaceUtil.h"
+
+using namespace UE::Geometry;
+
+#define LOCTEXT_NAMESPACE "UOffsetMeshSelectionTool"
+
+
+
+USingleTargetWithSelectionTool* UOffsetMeshSelectionToolBuilder::CreateNewTool(const FToolBuilderState& SceneState) const
+{
+	return NewObject<UOffsetMeshSelectionTool>(SceneState.ToolManager);
+}
+
+
+
+class FOffsetMeshSelectionOpFactory : public IDynamicMeshOperatorFactory
+{
+public:
+	UOffsetMeshSelectionTool* SourceTool;
+
+	FOffsetMeshSelectionOpFactory(UOffsetMeshSelectionTool* Tool)
+	{
+		SourceTool = Tool;
+	}
+
+	TUniquePtr<FRegionOffsetOp> MakeNewOffsetOp(
+		TSharedPtr<FSharedConstDynamicMesh3> SourceMesh,
+		const TArray<int32>& OffsetROI	)
+	{
+		TUniquePtr<FRegionOffsetOp> Op = MakeUnique<FRegionOffsetOp>();
+		Op->OriginalMeshShared = SourceTool->EditRegionSharedMesh;
+		Op->TriangleSelection = OffsetROI;
+
+		Op->bShellsToSolids = SourceTool->OffsetProperties->bShellsToSolids;
+
+		Op->bInferGroupsFromNeighbours = SourceTool->OffsetProperties->bInferGroupsFromNbrs;
+		Op->bNewGroupPerSubdivision = SourceTool->OffsetProperties->bGroupPerSubdivision;
+		Op->bUseColinearityForSettingBorderGroups = false;		// what is this??
+		Op->bRemapExtrudeGroups = ! SourceTool->OffsetProperties->bReplaceSelectionGroups;
+
+		Op->OffsetMode = FRegionOffsetOp::EOffsetComputationMode::VertexNormals;
+		if ( SourceTool->OffsetProperties->Direction == EOffsetMeshSelectionDirectionMode::ConstantWidth )
+		{
+			Op->OffsetMode = FRegionOffsetOp::EOffsetComputationMode::ApproximateConstantThickness;
+		}
+		else if (SourceTool->OffsetProperties->Direction == EOffsetMeshSelectionDirectionMode::FaceNormals)
+		{
+			Op->OffsetMode = FRegionOffsetOp::EOffsetComputationMode::FaceNormals;
+		}
+
+		Op->OffsetDistance = SourceTool->OffsetProperties->OffsetDistance;
+
+		//Op->bUseColinearityForSettingBorderGroups = OffsetProperties->bUseColinearityForSettingBorderGroups;
+		//Op->MaxScaleForAdjustingTriNormalsOffset = OffsetProperties->MaxDistanceScaleFactor;
+
+		Op->CreaseAngleThresholdDeg = SourceTool->OffsetProperties->CreaseAngle;
+
+		Op->UVScaleFactor = SourceTool->OffsetProperties->UVScale;
+		Op->bUVIslandPerGroup = SourceTool->OffsetProperties->bUVIslandPerGroup;
+
+		Op->bInferMaterialID = SourceTool->OffsetProperties->bInferMaterialID;
+		Op->SetMaterialID = FMath::Min(SourceTool->OffsetProperties->SetMaterialID, SourceTool->MaxMaterialID);
+
+		Op->NumSubdivisions = SourceTool->OffsetProperties->NumSubdivisions;
+
+		Op->SetResultTransform(SourceTool->WorldTransform);
+
+		return Op;
+	}
+
+	virtual TUniquePtr<FDynamicMeshOperator> MakeNewOperator() override
+	{
+		return MakeNewOffsetOp(SourceTool->EditRegionSharedMesh, SourceTool->RegionOffsetROI.Array());
+	}
+};
+
+
+UOffsetMeshSelectionTool::UOffsetMeshSelectionTool()
+{
+	SetToolDisplayName(LOCTEXT("ToolName", "Offset Selection"));
+}
+
+
+void UOffsetMeshSelectionTool::Setup()
+{
+	UToolTarget* UseTarget = Super::GetTarget();
+	CurrentMesh = UE::ToolTarget::GetDynamicMeshCopy(UseTarget);
+
+	// need valid MaterialIDs for MaterialID parameter
+	FInterval1i MatIDRange;
+	UE::Geometry::ComputeMaterialIDRange(CurrentMesh, MatIDRange);
+	MaxMaterialID = MatIDRange.Max;
+
+	WorldTransform = UE::ToolTarget::GetLocalToWorldTransform(UseTarget);
+	FComponentMaterialSet MaterialSet = UE::ToolTarget::GetMaterialSet(UseTarget);
+
+	OffsetProperties = NewObject<UOffsetMeshSelectionToolProperties>(this);
+	OffsetProperties->RestoreProperties(this);
+	OffsetProperties->WatchProperty(OffsetProperties->OffsetDistance,
+		[this](double) { EditCompute->InvalidateResult(); });
+	OffsetProperties->WatchProperty(OffsetProperties->Direction,
+		[this](EOffsetMeshSelectionDirectionMode) { EditCompute->InvalidateResult();});
+	OffsetProperties->WatchProperty(OffsetProperties->NumSubdivisions,
+		[this](int) { EditCompute->InvalidateResult();});
+	OffsetProperties->WatchProperty(OffsetProperties->bShellsToSolids,
+		[this](bool) { EditCompute->InvalidateResult(); });
+	OffsetProperties->WatchProperty(OffsetProperties->UVScale,
+		[this](double) { EditCompute->InvalidateResult(); });
+	OffsetProperties->WatchProperty(OffsetProperties->bInferGroupsFromNbrs,
+		[this](bool) { EditCompute->InvalidateResult(); });
+	OffsetProperties->WatchProperty(OffsetProperties->bGroupPerSubdivision,
+		[this](bool) { EditCompute->InvalidateResult(); });
+	OffsetProperties->WatchProperty(OffsetProperties->bUVIslandPerGroup,
+		[this](bool) { EditCompute->InvalidateResult(); });
+	OffsetProperties->WatchProperty(OffsetProperties->bReplaceSelectionGroups,
+		[this](bool) { EditCompute->InvalidateResult(); });
+	OffsetProperties->WatchProperty(OffsetProperties->CreaseAngle,
+		[this](double) { EditCompute->InvalidateResult(); });
+	OffsetProperties->WatchProperty(OffsetProperties->bInferMaterialID,
+		[this](bool) { EditCompute->InvalidateResult(); });
+	OffsetProperties->WatchProperty(OffsetProperties->SetMaterialID,
+		[this](int) { EditCompute->InvalidateResult(); });
+
+	OffsetProperties->WatchProperty(OffsetProperties->bShowInputMaterials,
+		[this](bool) { UpdateVisualizationSettings(); });
+
+	AddToolPropertySource(OffsetProperties);
+
+	// extract selection
+	FMeshFaceSelection TriSelection(&CurrentMesh);
+	const FGeometrySelection& Selection = GetGeometrySelection();
+	UE::Geometry::EnumerateSelectionTriangles(Selection, CurrentMesh,
+		[&](int32 tid) { TriSelection.Select(tid); },
+		nullptr  /* no support for group layers yet */);
+	OffsetROI = TriSelection.AsArray();
+
+	// if we have an empty selection, just abort here
+	if (OffsetROI.Num() == 0)
+	{
+		GetToolManager()->PostActiveToolShutdownRequest(this, EToolShutdownType::Cancel,
+			true, LOCTEXT("InvalidSelectionMessage", "OffsetMeshSelectionTool: input face selection was empty, cannot Offset"));
+		return;
+	}
+
+	TriSelection.ExpandToOneRingNeighbours(2);
+	ModifiedROI = TriSelection.AsSet();		// could steal here if it was possible
+
+	SelectionBoundsWorld = (FBox)TMeshQueries<FDynamicMesh3>::GetTrianglesBounds(CurrentMesh, OffsetROI, WorldTransform);
+
+	// create the preview object for the unmodified area
+	// (if input was a UDynamicMesh we could do this w/o making a copy...)
+	SourcePreview = NewObject<UPreviewMesh>();
+	SourcePreview->CreateInWorld(GetTargetWorld(), WorldTransform);
+	ToolSetupUtil::ApplyRenderingConfigurationToPreview(SourcePreview, UseTarget);
+	SourcePreview->SetMaterials(MaterialSet.Materials);
+	SourcePreview->EnableSecondaryTriangleBuffers(
+		[this](const FDynamicMesh3* Mesh, int32 TriangleID)
+		{
+			return OffsetROI.Contains(TriangleID);		// hide triangles in Offset area in base mesh
+		});
+	SourcePreview->SetSecondaryBuffersVisibility(false);
+	SourcePreview->SetTangentsMode(EDynamicMeshComponentTangentsMode::AutoCalculated);
+	SourcePreview->UpdatePreview(&CurrentMesh);
+
+	// initialize a region operator for the modified area
+	RegionOperator = MakePimpl<FMeshRegionOperator>(&CurrentMesh, ModifiedROI.Array());
+	for (int32 tid : OffsetROI)
+	{
+		RegionOffsetROI.Add( RegionOperator->Region.MapTriangleToSubmesh(tid) );
+	}
+	EditRegionMesh = RegionOperator->Region.GetSubmesh();
+	for (int32 tid : EditRegionMesh.TriangleIndicesItr())
+	{
+		if (RegionOffsetROI.Contains(tid) == false)
+		{
+			RegionBorderTris.Add(tid);
+		}
+	}
+	EditRegionSharedMesh = MakeShared<FSharedConstDynamicMesh3>(&EditRegionMesh);
+
+	// try to guess selection frame...
+	SelectionFrameLocal = UE::Geometry::ComputeFaceSelectionFrame(EditRegionMesh, RegionOffsetROI.Array(), false);
+	InitialFrameLocal = SelectionFrameLocal;
+	InitialFrameWorld = InitialFrameLocal;
+	InitialFrameWorld.Transform(WorldTransform);
+
+	this->OperatorFactory = MakePimpl<FOffsetMeshSelectionOpFactory>(this);
+
+	// Create the preview compute for the extrusion operation
+	EditCompute = NewObject<UMeshOpPreviewWithBackgroundCompute>(this);
+	EditCompute->Setup(GetTargetWorld(), this->OperatorFactory.Get());
+	ToolSetupUtil::ApplyRenderingConfigurationToPreview(EditCompute->PreviewMesh, UseTarget);
+	EditCompute->PreviewMesh->SetTransform((FTransform)WorldTransform);
+	EditCompute->ConfigureMaterials(MaterialSet.Materials, nullptr, nullptr);
+
+	// hide the triangles in the 'border' region outside the Offset area, although they are included in
+	// the Offset region operator submesh, they are still visible in the base mesh
+	EditCompute->PreviewMesh->EnableSecondaryTriangleBuffers(
+		[this](const FDynamicMesh3* Mesh, int32 TriangleID) { return RegionBorderTris.Contains(TriangleID); });
+	EditCompute->PreviewMesh->SetSecondaryBuffersVisibility(false);
+
+	UpdateVisualizationSettings();
+
+	// is this the right tangents behavior?
+	EditCompute->PreviewMesh->SetTangentsMode(EDynamicMeshComponentTangentsMode::AutoCalculated);
+	EditCompute->PreviewMesh->UpdatePreview( & RegionOperator->Region.GetSubmesh() );
+	EditCompute->SetVisibility(true);
+
+	// hide input Component
+	UE::ToolTarget::HideSourceObject(Target);
+
+	// are these needed?
+	OffsetFrameWorld = InitialFrameWorld;
+	OffsetFrameLocal = OffsetFrameWorld;
+	OffsetFrameLocal.Transform( WorldTransform.InverseUnsafe() );
+	LocalScale = FVector3d::One();
+
+	EditCompute->InvalidateResult();
+}
+
+
+void UOffsetMeshSelectionTool::OnShutdown(EToolShutdownType ShutdownType)
+{
+	if ( OffsetProperties )
+	{
+		OffsetProperties->SaveProperties(this);
+	}
+
+	if (SourcePreview != nullptr )
+	{
+		SourcePreview->Disconnect();
+		SourcePreview = nullptr;
+	}
+
+	if (EditCompute != nullptr)
+	{
+		// shut down background computation
+		// (we don't actually care about compute result here because it was only for the live preview...)
+		FDynamicMeshOpResult ComputeResult = EditCompute->Shutdown();
+		EditCompute->ClearOpFactory();
+		EditCompute->OnOpCompleted.RemoveAll(this);
+		EditCompute = nullptr;
+
+		UToolTarget* UseTarget = GetTarget();
+		UE::ToolTarget::ShowSourceObject(UseTarget);
+
+		if (ShutdownType == EToolShutdownType::Accept)
+		{
+			GetToolManager()->BeginUndoTransaction(LOCTEXT("OffsetSelection", "Offset Selection"));
+
+			TArray<int32> SelectOutputTriangles;
+
+			// this function computes the final Offset on full Mesh and cleans it up
+			auto ComputeFinalResult = [this, &SelectOutputTriangles](FDynamicMesh3& EditMesh) -> bool
+			{
+				TSharedPtr<FSharedConstDynamicMesh3> FullMeshShared = MakeShared<FSharedConstDynamicMesh3>(&EditMesh);
+				TUniquePtr<FRegionOffsetOp> FinalOp = OperatorFactory->MakeNewOffsetOp(FullMeshShared, this->OffsetROI);
+				SelectOutputTriangles = this->OffsetROI;
+				bool bResult = FinalOp->CalculateResultInPlace(EditMesh, nullptr);
+				if (EditMesh.IsCompact() == false)
+				{
+					FCompactMaps CompactMaps;
+					EditMesh.CompactInPlace(&CompactMaps);
+					if (CompactMaps.TriangleMapIsSet())
+					{
+						for (int32& tid : SelectOutputTriangles)
+						{
+							tid = CompactMaps.GetTriangleMapping(tid);
+						}
+					}
+				}
+
+				FullMeshShared->ReleaseSharedObject();
+				return bResult;
+			};
+
+			// try to emit an incremental FMeshChange, to avoid storing full mesh before/after copies in undo history
+			bool bChangeApplied = false;
+			if (UE::ToolTarget::SupportsIncrementalMeshChanges(UseTarget))
+			{
+				bChangeApplied = UE::ToolTarget::ApplyIncrementalMeshEditChange(UseTarget,
+					[&](FDynamicMesh3& EditMesh, UObject* TransactionTarget)
+					{
+						FDynamicMeshChangeTracker ChangeTracker(&EditMesh);
+						ChangeTracker.BeginChange();
+						ChangeTracker.SaveTriangles(OffsetROI, true);
+						if (ComputeFinalResult(EditMesh))
+						{
+							TUniquePtr<FDynamicMeshChange> MeshEditChange = ChangeTracker.EndChange();
+							GetToolManager()->GetContextTransactionsAPI()->AppendChange( TransactionTarget,
+								MakeUnique<FMeshChange>(MoveTemp(MeshEditChange)), LOCTEXT("OffsetSelection", "Offset Selection") );
+							return true;
+						}
+						return false;
+					});
+
+				ensureMsgf(bChangeApplied, TEXT("UOffsetMeshSelectionTool::OnShutdown : incremental mesh edit failed!"));
+			}
+
+			// if we could not apply incremental change, apply a full-mesh update
+			if ( bChangeApplied == false )
+			{
+				ComputeFinalResult(CurrentMesh);
+				const bool bModifiedTopology = true;
+				UE::ToolTarget::CommitDynamicMeshUpdate(UseTarget, CurrentMesh, bModifiedTopology);
+			}
+
+			// Construct and send output triangle selection, assume selection system will convert to correct type.
+			// (todo: maybe be smarter about this, as for (eg) an input vertex selection, it will grow the selection...)
+			// Note that we cannot do this before we commit/complete the mesh update, as the selection system needs a chance
+			// to respond to the mesh change. Possibly should use UE::Geometry::InitializeSelectionFromTriangles() here...need to ProcessMesh() then though
+			FGeometrySelection OutputSelection;
+			OutputSelection.InitializeTypes(EGeometryElementType::Face, EGeometryTopologyType::Triangle);
+			for (int32 tid : SelectOutputTriangles)
+			{
+				OutputSelection.Selection.Add( FGeoSelectionID::MeshTriangle(tid).Encoded() );
+			}
+			UE::Geometry::SetToolOutputGeometrySelectionForTarget(this, UseTarget, OutputSelection);
+
+
+			GetToolManager()->EndUndoTransaction();
+		}
+	}
+
+	if ( EditRegionSharedMesh.IsValid() )
+	{
+		EditRegionSharedMesh->ReleaseSharedObject();
+	}
+
+	CurrentMesh.Clear();
+}
+
+
+
+void UOffsetMeshSelectionTool::OnTick(float DeltaTime)
+{
+	EditCompute->Tick(DeltaTime);
+}
+
+
+void UOffsetMeshSelectionTool::Render(IToolsContextRenderAPI* RenderAPI)
+{
+}
+
+FBox UOffsetMeshSelectionTool::GetWorldSpaceFocusBox() 
+{ 
+	FAxisAlignedBox3d CurBox = SelectionBoundsWorld;
+	FVector3d Translation = OffsetFrameWorld.Origin - InitialFrameWorld.Origin;
+	CurBox.Min += Translation; CurBox.Max += Translation;
+	CurBox.Contain(SelectionBoundsWorld);
+	return (FBox)CurBox;
+}
+
+
+void UOffsetMeshSelectionTool::UpdateVisualizationSettings()
+{
+	if (OffsetProperties->bShowInputMaterials)
+	{
+		EditCompute->DisablePreviewMaterials();
+		Cast<UDynamicMeshComponent>(EditCompute->PreviewMesh->GetRootComponent())->SetColorOverrideMode(EDynamicMeshComponentColorOverrideMode::None);
+	}
+	else
+	{
+		EditCompute->ConfigurePreviewMaterials(
+			ToolSetupUtil::GetDefaultWorkingMaterial(GetToolManager()),
+			ToolSetupUtil::GetSelectionMaterial(FLinearColor::Yellow, GetToolManager())  );
+		Cast<UDynamicMeshComponent>(EditCompute->PreviewMesh->GetRootComponent())->SetColorOverrideMode(EDynamicMeshComponentColorOverrideMode::Polygroups);
+	}
+}
+
+
+
+
+#undef LOCTEXT_NAMESPACE
\ No newline at end of file
diff --git a/Engine/Plugins/Experimental/MeshModelingToolsetExp/Source/MeshModelingToolsExp/Public/Polymodeling/ExtrudeMeshSelectionTool.h b/Engine/Plugins/Experimental/MeshModelingToolsetExp/Source/MeshModelingToolsExp/Public/Polymodeling/ExtrudeMeshSelectionTool.h
--- a/Engine/Plugins/Experimental/MeshModelingToolsetExp/Source/MeshModelingToolsExp/Public/Polymodeling/ExtrudeMeshSelectionTool.h
+++ b/Engine/Plugins/Experimental/MeshModelingToolsetExp/Source/MeshModelingToolsExp/Public/Polymodeling/ExtrudeMeshSelectionTool.h
@@ -123,6 +123,7 @@ class MESHMODELINGTOOLSEXP_API UExtrudeMeshSelectionTool : public USingleTargetW
 	GENERATED_BODY()
 	using FFrame3d = UE::Geometry::FFrame3d;
 public:
+	UExtrudeMeshSelectionTool();
 
 	virtual void Setup() override;
 	virtual void OnShutdown(EToolShutdownType ShutdownType) override;
diff --git a/Engine/Plugins/Experimental/MeshModelingToolsetExp/Source/MeshModelingToolsExp/Public/Polymodeling/OffsetMeshSelectionTool.h b/Engine/Plugins/Experimental/MeshModelingToolsetExp/Source/MeshModelingToolsExp/Public/Polymodeling/OffsetMeshSelectionTool.h
new file mode 100644
--- /dev/null
+++ b/Engine/Plugins/Experimental/MeshModelingToolsetExp/Source/MeshModelingToolsExp/Public/Polymodeling/OffsetMeshSelectionTool.h
@@ -0,0 +1,172 @@
+// Copyright Epic Games, Inc. All Rights Reserved.
+
+#pragma once
+
+#include "BaseTools/SingleTargetWithSelectionTool.h"
+#include "DynamicMesh/DynamicMesh3.h"
+#include "DynamicMesh/MeshSharingUtil.h"
+#include "Templates/PimplPtr.h"
+#include "OffsetMeshSelectionTool.generated.h"
+
+class UPreviewMesh;
+class UMeshOpPreviewWithBackgroundCompute; 
+class FOffsetMeshSelectionOpFactory;
+namespace UE::Geometry { class FMeshRegionOperator; }
+
+/**
+ * ToolBuilder
+ */
+UCLASS()
+class MESHMODELINGTOOLSEXP_API UOffsetMeshSelectionToolBuilder : public USingleTargetWithSelectionToolBuilder
+{
+	GENERATED_BODY()
+public:
+	virtual USingleTargetWithSelectionTool* CreateNewTool(const FToolBuilderState& SceneState) const override;
+	virtual bool RequiresInputSelection() const override { return true; }
+};
+
+
+UENUM()
+enum class EOffsetMeshSelectionInteractionMode : uint8
+{
+	/** Define the offset distance using a slider in the Settings */
+	Fixed = 0
+};
+
+UENUM()
+enum class EOffsetMeshSelectionDirectionMode : uint8
+{
+	/** */
+	VertexNormals = 0,
+	/** */
+	FaceNormals = 1,
+	/**  */
+	ConstantWidth = 2
+};
+
+
+UCLASS()
+class MESHMODELINGTOOLSEXP_API UOffsetMeshSelectionToolProperties : public UInteractiveToolPropertySet
+{
+	GENERATED_BODY()
+
+public:
+	/** Control how the Offset Area should be Transformed */
+	//UPROPERTY(EditAnywhere, Category = "Offset")
+	//EOffsetMeshSelectionInteractionMode InputMode = EOffsetMeshSelectionInteractionMode::Interactive;
+
+	/** The Extrusion Distance used in Fixed Input Mode*/
+	UPROPERTY(EditAnywhere, Category = "Offset", meta = (DisplayName="Distance", UIMin = -250, UIMax = 250))
+	double OffsetDistance = 10.0;
+
+	/** Control how the Offset Area should be displaced */
+	UPROPERTY(EditAnywhere, Category = "Offset")
+	EOffsetMeshSelectionDirectionMode Direction = EOffsetMeshSelectionDirectionMode::ConstantWidth;
+
+	/** Specify the number of subdivisions along the sides of the Extrusion */
+	UPROPERTY(EditAnywhere, Category = "Offset", meta = (DisplayName="Subdivisions", UIMin = 0, UIMax = 10, ClampMin = 0, ClampMax = 1000))
+	int NumSubdivisions = 0;
+
+	/** Specify the Crease Angle used to split the sides of the Extrusion into separate Groups */
+	UPROPERTY(EditAnywhere, Category = "Offset", meta = (UIMin = 0, UIMax = 180, Min = 0, Max = 180))
+	double CreaseAngle = 60.0;
+
+	/** If the Offset Area has a fully open border, this option determines if Extrusion will create a Solid mesh or leave the base "open" */
+	UPROPERTY(EditAnywhere, Category = "Offset")
+	bool bShellsToSolids = true;			// todo: should only be available if input has shells that could become solid...
+
+	/** Control whether a single Group should be generated along the sides of the Extrusion, or multiple Groups based on the adjacent Groups around the Offset Area border */
+	UPROPERTY(EditAnywhere, Category = "Groups", meta = (DisplayName="Propagate Groups"))
+	bool bInferGroupsFromNbrs = true;
+
+	/** Control whether a new Group is generated for each Subdivision */
+	UPROPERTY(EditAnywhere, Category = "Groups", meta=(DisplayName="Per Subdivision", EditCondition="NumSubdivisions > 0") )
+	bool bGroupPerSubdivision = true;
+
+	/** Control whether groups in the Offset Area are mapped to new Groups, or replaced with a single new Group */
+	UPROPERTY(EditAnywhere, Category = "Groups", meta=(DisplayName="Replace Cap Groups"))
+	bool bReplaceSelectionGroups = false;
+
+	/** The automatically-generated UVs on the sides of the Extrusion are scaled by this value */
+	UPROPERTY(EditAnywhere, Category = "UVs", meta = (DisplayName="UV Scale", UIMin = 0.001, UIMax = 10.0, ClampMin = 0.000001))
+	double UVScale = 1.0f;
+
+	/** Control whether a separate UV island should be generated for each output Group on the sides of the Extrusion, or a single UV island wrapping around the entire "tube" */
+	UPROPERTY(EditAnywhere, Category = "UVs", meta = (DisplayName="Island Per Group"))
+	bool bUVIslandPerGroup = true;
+
+	/** Control whether SetMaterialID is assigned to all triangles along the sides of the Extrusion, or if MaterialIDs should be inferred from the Offset Area */
+	UPROPERTY(EditAnywhere, Category = "Material", meta=(DisplayName="Infer Materials"))
+	bool bInferMaterialID = true;
+
+	/** Constant Material ID used when MaterialIDs are not being inferred, or no adjacent MaterialID exists */
+	UPROPERTY(EditAnywhere, Category = "Material", meta=(DisplayName="Material ID", EditCondition="bInferMaterialID == false"))
+	int SetMaterialID = 0;
+
+	/** Control whether the original Mesh Materials should be shown, or a visualization of the Offset Groups */
+	UPROPERTY(EditAnywhere, Category = "Visualization", meta=(DisplayName="Show Materials"))
+	bool bShowInputMaterials = false;
+};
+
+
+
+UCLASS()
+class MESHMODELINGTOOLSEXP_API UOffsetMeshSelectionTool : public USingleTargetWithSelectionTool
+{
+	GENERATED_BODY()
+	using FFrame3d = UE::Geometry::FFrame3d;
+public:
+	UOffsetMeshSelectionTool();
+
+	virtual void Setup() override;
+	virtual void OnShutdown(EToolShutdownType ShutdownType) override;
+
+	virtual void OnTick(float DeltaTime) override;
+	virtual void Render(IToolsContextRenderAPI* RenderAPI) override;
+
+	virtual bool HasCancel() const override { return true; }
+	virtual bool HasAccept() const override { return true; }
+
+	// IInteractiveToolCameraFocusAPI implementation
+	virtual FBox GetWorldSpaceFocusBox() override;
+
+
+protected:
+	UPROPERTY()
+	TObjectPtr<UOffsetMeshSelectionToolProperties> OffsetProperties = nullptr;
+
+
+protected:
+	FBox SelectionBoundsWorld;
+	UE::Geometry::FTransformSRT3d WorldTransform;
+
+	UE::Geometry::FFrame3d SelectionFrameLocal;
+	UE::Geometry::FFrame3d InitialFrameLocal;
+	UE::Geometry::FFrame3d InitialFrameWorld;
+	UE::Geometry::FFrame3d OffsetFrameWorld;
+	UE::Geometry::FFrame3d OffsetFrameLocal;
+	FVector3d LocalScale;
+
+	UE::Geometry::FDynamicMesh3 CurrentMesh;
+	TArray<int32> OffsetROI;
+	TSet<int32> ModifiedROI;
+	int32 MaxMaterialID = 0;		// maximum material ID on mesh (inclusive)
+
+	TPimplPtr<UE::Geometry::FMeshRegionOperator> RegionOperator;
+	UE::Geometry::FDynamicMesh3 EditRegionMesh;
+	TSharedPtr<UE::Geometry::FSharedConstDynamicMesh3> EditRegionSharedMesh;
+
+	TSet<int32> RegionOffsetROI;
+	TSet<int32> RegionBorderTris;
+
+	TPimplPtr<FOffsetMeshSelectionOpFactory> OperatorFactory;
+	friend class FOffsetMeshSelectionOpFactory;
+
+	UPROPERTY()
+	TObjectPtr<UPreviewMesh> SourcePreview = nullptr;
+
+	UPROPERTY()
+	TObjectPtr<UMeshOpPreviewWithBackgroundCompute> EditCompute = nullptr;
+
+	void UpdateVisualizationSettings();
+};
\ No newline at end of file
diff --git a/Engine/Plugins/Runtime/MeshModelingToolset/Source/MeshModelingTools/Private/Commands/RetriangulateGeometrySelectionCommand.cpp b/Engine/Plugins/Runtime/MeshModelingToolset/Source/MeshModelingTools/Private/Commands/RetriangulateGeometrySelectionCommand.cpp
new file mode 100644
--- /dev/null
+++ b/Engine/Plugins/Runtime/MeshModelingToolset/Source/MeshModelingTools/Private/Commands/RetriangulateGeometrySelectionCommand.cpp
@@ -0,0 +1,190 @@
+// Copyright Epic Games, Inc. All Rights Reserved.
+
+#include "Commands/RetriangulateGeometrySelectionCommand.h"
+#include "ToolContextInterfaces.h"
+#include "UDynamicMesh.h"
+#include "DynamicMesh/DynamicMesh3.h"
+#include "DynamicMesh/DynamicMeshChangeTracker.h"
+#include "DynamicMeshEditor.h"
+#include "Changes/MeshChange.h"
+#include "Selections/GeometrySelectionUtil.h"
+#include "Operations/PolygroupRemesh.h"
+#include "ConstrainedDelaunay2.h"
+#include "Operations/SimpleHoleFiller.h"
+#include "MeshRegionBoundaryLoops.h"
+#include "DynamicMesh/MeshIndexUtil.h"
+
+#include UE_INLINE_GENERATED_CPP_BY_NAME(RetriangulateGeometrySelectionCommand)
+
+using namespace UE::Geometry;
+
+
+#define LOCTEXT_NAMESPACE "URetriangulateGeometrySelectionCommand"
+
+
+FText URetriangulateGeometrySelectionCommand::GetCommandShortString() const
+{
+	return LOCTEXT("ShortString", "Retriangulate Selection");
+}
+
+
+bool URetriangulateGeometrySelectionCommand::CanExecuteCommandForSelection(UGeometrySelectionEditCommandArguments* SelectionArgs)
+{
+	return SelectionArgs->IsMatchingType(FGeometryIdentifier::ETargetType::MeshContainer, FGeometryIdentifier::EObjectType::DynamicMesh)
+		&& (SelectionArgs->SelectionHandle.Selection->ElementType == EGeometryElementType::Face)
+		&& (SelectionArgs->SelectionHandle.Selection->TopologyType == EGeometryTopologyType::Polygroup);
+}
+
+
+void URetriangulateGeometrySelectionCommand::ExecuteCommandForSelection(UGeometrySelectionEditCommandArguments* SelectionArgs, UInteractiveCommandResult** Result)
+{
+	if (Result != nullptr)
+	{
+		*Result = nullptr;
+	}
+
+	// should have been verified by CanExecute
+	check(SelectionArgs->IsMatchingType(FGeometryIdentifier::ETargetType::MeshContainer, FGeometryIdentifier::EObjectType::DynamicMesh));
+	
+	// collect up all our inputs
+	UDynamicMesh* MeshObject = SelectionArgs->SelectionHandle.Identifier.GetAsObjectType<UDynamicMesh>();
+	check(MeshObject != nullptr);
+	const FGeometrySelection* Selection = SelectionArgs->SelectionHandle.Selection;
+
+	bool bTrackChanges = SelectionArgs->HasTransactionsAPI();
+	TUniquePtr<FDynamicMeshChange> DynamicMeshChange;	// only initialized if bTrackChanges == true
+
+	// apply the Retriangulate operation
+	MeshObject->EditMesh([&](FDynamicMesh3& EditMesh)
+	{
+		// build list of triangles from whatever the selection contains
+		TSet<int32> TriangleList;
+		//UE::Geometry::FPolygroupSet UsePolygroupSet = ...;		// need to support this eventually
+		UE::Geometry::EnumerateSelectionTriangles(*Selection, EditMesh,
+			[&](int32 TriangleID) { TriangleList.Add(TriangleID); });
+
+		if (TriangleList.Num() == 0)
+		{
+			for (int32 tid : EditMesh.TriangleIndicesItr())
+			{
+				TriangleList.Add(tid);
+			}
+		}
+
+		// mark triangles for change
+		FDynamicMeshChangeTracker ChangeTracker(&EditMesh);
+		FGroupTopology Topology(&EditMesh, true);
+
+		if (bTrackChanges)
+		{
+			ChangeTracker.BeginChange();
+		}
+
+		if ( TriangleList.Num() == EditMesh.TriangleCount() )
+		{
+			if (bTrackChanges)
+			{
+				ChangeTracker.SaveTriangles(TriangleList, true);
+			}
+
+			FPolygroupRemesh Remesh(&EditMesh, &Topology, ConstrainedDelaunayTriangulate<double>);
+			bool bSuccess = Remesh.Compute();
+		}
+		else
+		{
+			// ported from UEditMeshPolygonsTool::ApplyRetriangulate(), should be moved to a standalone geometry op
+
+			const FDynamicMeshAttributeSet* Attributes = EditMesh.HasAttributes() ? EditMesh.Attributes() : nullptr;
+
+			TSet<int32> SelectedGroupIDs;
+			for (int32 tid : TriangleList)
+			{
+				SelectedGroupIDs.Add( Topology.GetGroupID(tid) );
+			}
+			
+			int32 nCompleted = 0;
+			FDynamicMeshEditor Editor(&EditMesh);
+			for (int32 GroupID : SelectedGroupIDs)
+			{
+				const TArray<int32>& Triangles = Topology.GetGroupTriangles(GroupID);
+				ChangeTracker.SaveTriangles(Triangles, true);
+				FMeshRegionBoundaryLoops RegionLoops(&EditMesh, Triangles, true);
+				if (!RegionLoops.bFailed && RegionLoops.Loops.Num() == 1 && Triangles.Num() > 1)
+				{
+					TArray<FMeshRegionBoundaryLoops::VidOverlayMap<FVector2f>> VidUVMaps;
+					if (Attributes)
+					{
+						for (int i = 0; i < Attributes->NumUVLayers(); ++i)
+						{
+							VidUVMaps.Emplace();
+							RegionLoops.GetLoopOverlayMap(RegionLoops.Loops[0], *Attributes->GetUVLayer(i), VidUVMaps.Last());
+						}
+					}
+
+					// We don't want to remove isolated vertices while removing triangles because we don't
+					// want to throw away boundary verts. However, this means that we'll have to go back
+					// through these vertices later to throw away isolated internal verts.
+					TArray<int32> OldVertices;
+					UE::Geometry::TriangleToVertexIDs(&EditMesh, Triangles, OldVertices);
+					Editor.RemoveTriangles(Topology.GetGroupTriangles(GroupID), false);
+
+					RegionLoops.Loops[0].Reverse();
+					FSimpleHoleFiller Filler(&EditMesh, RegionLoops.Loops[0]);
+					Filler.FillType = FSimpleHoleFiller::EFillType::PolygonEarClipping;
+					Filler.Fill(GroupID);
+
+					// Throw away any of the old verts that are still isolated (they were in the interior of the group)
+					for (int32 Vid : OldVertices)
+					{
+						if ( !EditMesh.IsReferencedVertex(Vid) )
+						{
+							constexpr bool bPreserveManifold = false;
+							EditMesh.RemoveVertex(Vid, bPreserveManifold);
+						}
+					}
+
+					if (Attributes)
+					{
+						for (int i = 0; i < Attributes->NumUVLayers(); ++i)
+						{
+							RegionLoops.UpdateLoopOverlayMapValidity(VidUVMaps[i], *Attributes->GetUVLayer(i));
+						}
+						Filler.UpdateAttributes(VidUVMaps);
+					}
+
+					nCompleted++;
+				}
+			}
+
+		}
+
+		// extract the change record
+		if (bTrackChanges)
+		{
+			DynamicMeshChange = ChangeTracker.EndChange();
+		}
+
+	}, EDynamicMeshChangeType::GeneralEdit, EDynamicMeshAttributeChangeFlags::Unknown, false);
+
+	// emit change 
+	if ( bTrackChanges && DynamicMeshChange.IsValid() )
+	{
+		SelectionArgs->GetTransactionsAPI()->BeginUndoTransaction(GetCommandShortString());
+		SelectionArgs->GetTransactionsAPI()->AppendChange(MeshObject, 
+			MakeUnique<FMeshChange>(MoveTemp(DynamicMeshChange)), GetCommandShortString());
+		SelectionArgs->GetTransactionsAPI()->EndUndoTransaction();
+	}
+
+	if (Result != nullptr)
+	{
+		UGeometrySelectionEditCommandResult* NewResult = NewObject<UGeometrySelectionEditCommandResult>();
+		NewResult->SourceHandle = SelectionArgs->SelectionHandle;
+		// todo...
+		//NewResult->OutputSelection = *NewResult->SourceHandle.Selection;
+		*Result = NewResult;
+	}
+
+}
+
+
+#undef LOCTEXT_NAMESPACE
diff --git a/Engine/Plugins/Runtime/MeshModelingToolset/Source/MeshModelingTools/Public/Commands/RetriangulateGeometrySelectionCommand.h b/Engine/Plugins/Runtime/MeshModelingToolset/Source/MeshModelingTools/Public/Commands/RetriangulateGeometrySelectionCommand.h
new file mode 100644
--- /dev/null
+++ b/Engine/Plugins/Runtime/MeshModelingToolset/Source/MeshModelingTools/Public/Commands/RetriangulateGeometrySelectionCommand.h
@@ -0,0 +1,23 @@
+// Copyright Epic Games, Inc. All Rights Reserved.
+
+#pragma once
+
+#include "Selection/SelectionEditInteractiveCommand.h"
+#include "RetriangulateGeometrySelectionCommand.generated.h"
+
+
+/**
+ * URetriangulateGeometrySelectionCommand 
+ */
+UCLASS()
+class MESHMODELINGTOOLS_API URetriangulateGeometrySelectionCommand : public UGeometrySelectionEditCommand
+{
+	GENERATED_BODY()
+public:
+
+	virtual bool AllowEmptySelection() const override { return true; }
+	virtual FText GetCommandShortString() const override;
+
+	virtual bool CanExecuteCommandForSelection(UGeometrySelectionEditCommandArguments* Arguments) override;
+	virtual void ExecuteCommandForSelection(UGeometrySelectionEditCommandArguments* Arguments, UInteractiveCommandResult** Result) override;
+};
\ No newline at end of file
diff --git a/Engine/Plugins/Runtime/MeshModelingToolset/Source/ModelingComponents/Private/Selection/GeometrySelectionManager.cpp b/Engine/Plugins/Runtime/MeshModelingToolset/Source/ModelingComponents/Private/Selection/GeometrySelectionManager.cpp
--- a/Engine/Plugins/Runtime/MeshModelingToolset/Source/ModelingComponents/Private/Selection/GeometrySelectionManager.cpp
+++ b/Engine/Plugins/Runtime/MeshModelingToolset/Source/ModelingComponents/Private/Selection/GeometrySelectionManager.cpp
@@ -935,7 +935,7 @@ bool UGeometrySelectionManager::CanExecuteSelectionCommand(UGeometrySelectionEdi
 		bHaveSelections = true;
 	});
 
-	return bHaveSelections && bCanExecute;
+	return ( bHaveSelections || (Command->AllowEmptySelection() && HasActiveTargets() && MeshTopologyMode != EMeshTopologyMode::None ) ) && bCanExecute;
 }
 
 void UGeometrySelectionManager::ExecuteSelectionCommand(UGeometrySelectionEditCommand* Command)
@@ -951,7 +951,7 @@ void UGeometrySelectionManager::ExecuteSelectionCommand(UGeometrySelectionEditCo
 
 	for (TSharedPtr<FGeometrySelectionTarget> Target : ActiveTargetReferences)
 	{
-		if (Target->Selection.IsEmpty()) continue;
+		if (Target->Selection.IsEmpty() && Command->AllowEmptySelection() == false) continue;
 
 		// TODO: can use Command->IsModifySelectionCommand() to check if this is a command that only affects selection
 		// and not geometry. In that case we can skip the intermediate clear-selection and emit a single change.
diff --git a/Engine/Plugins/Runtime/MeshModelingToolset/Source/ModelingOperators/Private/PolyModelingOps/RegionOffsetOp.cpp b/Engine/Plugins/Runtime/MeshModelingToolset/Source/ModelingOperators/Private/PolyModelingOps/RegionOffsetOp.cpp
new file mode 100644
--- /dev/null
+++ b/Engine/Plugins/Runtime/MeshModelingToolset/Source/ModelingOperators/Private/PolyModelingOps/RegionOffsetOp.cpp
@@ -0,0 +1,98 @@
+// Copyright Epic Games, Inc. All Rights Reserved.
+
+#include "PolyModelingOps/RegionOffsetOp.h"
+#include "Operations/OffsetMeshRegion.h"
+#include "Selections/MeshConnectedComponents.h"
+#include "DynamicMesh/MeshNormals.h"
+#include "Util/ProgressCancel.h"
+
+using namespace UE::Geometry;
+
+
+void FRegionOffsetOp::CalculateResult(FProgressCancel* Progress)
+{
+	bool bMeshValid = OriginalMeshShared->AccessSharedObject([&](const FDynamicMesh3& Mesh)
+	{
+		*ResultMesh = Mesh;
+	});
+
+	if (bMeshValid)
+	{
+		CalculateResultInPlace(*ResultMesh, Progress);
+	}
+}
+
+
+
+bool FRegionOffsetOp::CalculateResultInPlace(FDynamicMesh3& EditMesh, FProgressCancel* Progress)
+{
+	if (Progress && Progress->Cancelled())
+	{
+		return false;
+	}
+
+	if ( FMathd::Abs(OffsetDistance) < FMathd::ZeroTolerance )
+	{
+		return true;
+	}
+
+	FOffsetMeshRegion Extruder(&EditMesh);
+
+	Extruder.Triangles = TriangleSelection;
+	Extruder.NumSubdivisions = NumSubdivisions;
+
+	if ( bInferGroupsFromNeighbours )
+	{
+		Extruder.LoopEdgesShouldHaveSameGroup = [this, &EditMesh](int32 EdgeID1, int32 EdgeID2) 
+		{
+			return FOffsetMeshRegion::EdgesSeparateSameGroupsAndAreColinearAtBorder(
+				&EditMesh, EdgeID1, EdgeID2, bUseColinearityForSettingBorderGroups);
+		};
+	}
+	else
+	{
+		Extruder.LoopEdgesShouldHaveSameGroup = [this, &EditMesh](int32 EdgeID1, int32 EdgeID2) { return true; };
+	}
+	Extruder.bGroupPerSubdivision = bNewGroupPerSubdivision;
+	Extruder.bSingleGroupPerArea = !bRemapExtrudeGroups;
+	Extruder.CreaseAngleThresholdDeg = CreaseAngleThresholdDeg;
+
+	Extruder.UVScaleFactor = UVScaleFactor;
+	Extruder.bUVIslandPerGroup = bUVIslandPerGroup;
+
+	Extruder.SetMaterialID = SetMaterialID;
+	Extruder.bInferMaterialID = bInferMaterialID;
+
+	Extruder.ExtrusionVectorType = FOffsetMeshRegion::EVertexExtrusionVectorType::VertexNormal;
+	if (OffsetMode == EOffsetComputationMode::FaceNormals)
+	{
+		Extruder.ExtrusionVectorType = FOffsetMeshRegion::EVertexExtrusionVectorType::SelectionTriNormalsAngleWeightedAverage;
+	} 
+	else if (OffsetMode == EOffsetComputationMode::ApproximateConstantThickness)
+	{
+		Extruder.ExtrusionVectorType = FOffsetMeshRegion::EVertexExtrusionVectorType::SelectionTriNormalsAngleWeightedAdjusted;
+	} 
+
+	Extruder.OffsetPositionFunc = [this](const FVector3d& Pos, const FVector3d& VertexVector, int32 VertexID) {
+		return Pos + OffsetDistance * VertexVector;
+	};
+
+	Extruder.bIsPositiveOffset = (OffsetDistance > 0);
+	Extruder.bOffsetFullComponentsAsSolids = bShellsToSolids;
+
+	Extruder.Apply();
+
+	if (Progress && Progress->Cancelled())
+	{
+		return false;
+	}
+
+	if (bRecomputeNormals)
+	{
+		FMeshNormals::RecomputeOverlayTriNormals(EditMesh, Extruder.AllModifiedAndNewTriangles);
+	}
+
+
+	return true;
+}
+
diff --git a/Engine/Plugins/Runtime/MeshModelingToolset/Source/ModelingOperators/Public/PolyModelingOps/RegionOffsetOp.h b/Engine/Plugins/Runtime/MeshModelingToolset/Source/ModelingOperators/Public/PolyModelingOps/RegionOffsetOp.h
new file mode 100644
--- /dev/null
+++ b/Engine/Plugins/Runtime/MeshModelingToolset/Source/ModelingOperators/Public/PolyModelingOps/RegionOffsetOp.h
@@ -0,0 +1,117 @@
+// Copyright Epic Games, Inc. All Rights Reserved.
+
+#pragma once
+
+#include "CoreMinimal.h"
+#include "DynamicMesh/MeshSharingUtil.h"
+#include "ModelingOperators.h"
+
+
+namespace UE {
+namespace Geometry {
+
+class FDynamicMesh3;
+
+/**
+ * RegionOffsetOp extrudes a set of triangles of a mesh in a linear direction,
+ * with various options for how to process that extrusion. 
+ */
+class MODELINGOPERATORS_API FRegionOffsetOp : public FDynamicMeshOperator
+{
+public:
+	virtual ~FRegionOffsetOp() {}
+
+	//
+	// Inputs
+	//
+
+	/** Handle to the source DynamicMesh */
+	TSharedPtr<FSharedConstDynamicMesh3> OriginalMeshShared;
+
+	/** Triangles of the OriginalMesh that should be extruded */
+	TArray<int32> TriangleSelection;
+
+	//
+	// Parameters
+	//
+
+	enum class EOffsetComputationMode
+	{
+		/** */
+		VertexNormals = 0,
+		/** */
+		FaceNormals = 1,
+		/**  */
+		ApproximateConstantThickness = 2
+	};
+
+
+	double OffsetDistance = 1.0;
+
+	EOffsetComputationMode OffsetMode = EOffsetComputationMode::ApproximateConstantThickness;
+
+
+	/** If bShellsToSolids is true, patches surrounded by mesh boundaries will be extruded into closed solids instead of open shells */
+	bool bShellsToSolids = true;
+
+	/** Number of subdivisions along the extrusion "tubes" */
+	int32 NumSubdivisions = 0;
+
+	/** 
+	 * Assign new groups to regions of the extrusion "tube" based on the existing adjacent groups before extrusion. Each (disconnected) 
+	 * unique pair of group IDs gets a new unique group ID on the extrusion. If disabled, the entire extrusion tube will be assigned
+	 * a single new group ID
+	 */
+	bool bInferGroupsFromNeighbours = true;
+	
+	/** If true, assign a separate group for each subdivision, otherwise all subdivisions have the same group */
+	bool bNewGroupPerSubdivision = true;
+	
+	/** assign new GroupIDs to each input GroupID in the extrude area. If false, disconnected extrude areas get single new group IDs */
+	bool bRemapExtrudeGroups = true;
+
+	/** if the opening angle at an edge of the extrusion tube is > this threshold, "cut" any group that crosses the edge into two groups */
+	double CreaseAngleThresholdDeg = 180.0;
+
+	/** Scaling applied to the default UV values */
+	double UVScaleFactor = 1.0f;
+
+	/** If true, assign a new UV island for each new group, otherwise each disconnected extrusion tube is unwrapped to a single UV island*/
+	bool bUVIslandPerGroup = true;
+
+	/** Constant Material ID to set along the extrusion tube */
+	int SetMaterialID = 0;
+
+	/** If true, Material IDs around the border of the extrude area are propagated "down" the extrusion tube */
+	bool bInferMaterialID = true;
+
+	/** Determine whether normals should be recomputed */
+	bool bRecomputeNormals = true;
+
+	// (this parameter may not be relevant and perhaps should be removed...)
+	// Used when setting groups for the sides when the extrusion includes a mesh border. When true,
+	// different groups are added based on colinearity of the border edges.
+	bool bUseColinearityForSettingBorderGroups = true;
+
+
+	//
+	// Outputs
+	//
+
+
+	/**
+	 * Compute result applied directly to EditMesh
+	 * @return true on success, false on error or if Progress is cancelled
+	 */
+	virtual bool CalculateResultInPlace(FDynamicMesh3& EditMesh, FProgressCancel* Progress);
+
+
+	//
+	// FDynamicMeshOperator interface 
+	//
+	virtual void CalculateResult(FProgressCancel* Progress) override;
+
+
+};
+
+}} // end UE::Geometry
