ModelingTools: Fix non-unity build of SDraggableBox and STransformGizmoNumericalUIOverlay.

diff --git a/Engine/Plugins/Experimental/MeshModelingToolsetExp/Source/ModelingEditorUI/Private/ModelingWidgets/SDraggableBox.cpp b/Engine/Plugins/Experimental/MeshModelingToolsetExp/Source/ModelingEditorUI/Private/ModelingWidgets/SDraggableBox.cpp
--- a/Engine/Plugins/Experimental/MeshModelingToolsetExp/Source/ModelingEditorUI/Private/ModelingWidgets/SDraggableBox.cpp
+++ b/Engine/Plugins/Experimental/MeshModelingToolsetExp/Source/ModelingEditorUI/Private/ModelingWidgets/SDraggableBox.cpp
@@ -2,8 +2,14 @@
 
 #include "ModelingWidgets/SDraggableBox.h"
 
+#include "Framework/Application/SlateApplication.h"
+#include "Input/DragAndDrop.h" // FDragDropEvent
+#include "Styling/SlateTypes.h" // FWindowStyle
 #include "Viewports/InViewportUIDragOperation.h"
 #include "Widgets/Layout/SBackgroundBlur.h"
+#include "Widgets/Layout/SBorder.h"
+#include "Widgets/SBoxPanel.h" // SVerticalBox
+#include "Widgets/SWindow.h"
 
 
 void SDraggableBoxOverlay::Construct(const FArguments& InArgs)
diff --git a/Engine/Plugins/Experimental/MeshModelingToolsetExp/Source/ModelingEditorUI/Private/STransformGizmoNumericalUIOverlay.cpp b/Engine/Plugins/Experimental/MeshModelingToolsetExp/Source/ModelingEditorUI/Private/STransformGizmoNumericalUIOverlay.cpp
--- a/Engine/Plugins/Experimental/MeshModelingToolsetExp/Source/ModelingEditorUI/Private/STransformGizmoNumericalUIOverlay.cpp
+++ b/Engine/Plugins/Experimental/MeshModelingToolsetExp/Source/ModelingEditorUI/Private/STransformGizmoNumericalUIOverlay.cpp
@@ -5,10 +5,15 @@
 #include "BaseGizmos/CombinedTransformGizmo.h"
 #include "BaseGizmos/TransformGizmoUtil.h" // UCombinedTransformGizmoContextObject
 #include "ContextObjectStore.h"
+#include "Framework/Commands/UIAction.h" // FUIAction
+#include "Framework/Commands/UICommandInfo.h" // EUserInterfaceActionType
+#include "Framework/Multibox/MultiboxBuilder.h" // FMenuBuilder
 #include "InteractiveToolsContext.h"
 #include "ModelingWidgets/SDraggableBox.h" // SDraggableBoxOverlay
 #include "ModelingWidgets/STickDelegateWidget.h"
 #include "Transforms/TransformGizmoDataBinder.h"
+#include "Widgets/Input/SButton.h"
+#include "Widgets/Layout/SBox.h"
 #include "Widgets/SBoxPanel.h" // SHorizontalBox
 #include "Widgets/Images/SImage.h"
 
