MeshModelingTools: Fix casing of MultiBox in include path for non-windows non-unity builds.

diff --git a/Engine/Plugins/Experimental/MeshModelingToolsetExp/Source/ModelingEditorUI/Private/STransformGizmoNumericalUIOverlay.cpp b/Engine/Plugins/Experimental/MeshModelingToolsetExp/Source/ModelingEditorUI/Private/STransformGizmoNumericalUIOverlay.cpp
--- a/Engine/Plugins/Experimental/MeshModelingToolsetExp/Source/ModelingEditorUI/Private/STransformGizmoNumericalUIOverlay.cpp
+++ b/Engine/Plugins/Experimental/MeshModelingToolsetExp/Source/ModelingEditorUI/Private/STransformGizmoNumericalUIOverlay.cpp
@@ -7,7 +7,7 @@
 #include "ContextObjectStore.h"
 #include "Framework/Commands/UIAction.h" // FUIAction
 #include "Framework/Commands/UICommandInfo.h" // EUserInterfaceActionType
-#include "Framework/Multibox/MultiboxBuilder.h" // FMenuBuilder
+#include "Framework/MultiBox/MultiBoxBuilder.h" // FMenuBuilder
 #include "InteractiveToolsContext.h"
 #include "ModelingWidgets/SDraggableBox.h" // SDraggableBoxOverlay
 #include "ModelingWidgets/STickDelegateWidget.h"
