MeshModelingTools: Update UV island transform tool to use normal gizmo instead of the old multitransformer (which will probably be deleted at some point). Also let it set a display transform to enable non-delta mode in the gizmo numerical UI for the tool.

diff --git a/Engine/Plugins/Experimental/MeshModelingToolsetExp/Source/MeshModelingToolsExp/Private/EditUVIslandsTool.cpp b/Engine/Plugins/Experimental/MeshModelingToolsetExp/Source/MeshModelingToolsExp/Private/EditUVIslandsTool.cpp
--- a/Engine/Plugins/Experimental/MeshModelingToolsetExp/Source/MeshModelingToolsExp/Private/EditUVIslandsTool.cpp
+++ b/Engine/Plugins/Experimental/MeshModelingToolsetExp/Source/MeshModelingToolsExp/Private/EditUVIslandsTool.cpp
@@ -4,12 +4,14 @@
 #include "InteractiveToolManager.h"
 #include "ToolBuilderUtil.h"
 
+#include "BaseGizmos/CombinedTransformGizmo.h"
+#include "BaseGizmos/TransformGizmoUtil.h"
+#include "BaseGizmos/TransformProxy.h"
 #include "SegmentTypes.h"
 #include "DynamicMesh/DynamicMeshAttributeSet.h"
 #include "DynamicMesh/MeshNormals.h"
 #include "Engine/World.h"
 #include "Selections/MeshConnectedComponents.h"
-#include "Transforms/MultiTransformer.h"
 #include "BaseBehaviors/SingleClickBehavior.h"
 #include "ToolSetupUtil.h"
 #include "ModelingToolTargetUtil.h"
@@ -106,19 +108,30 @@ void UEditUVIslandsTool::Setup()
 	// init state flags flags
 	bInDrag = false;
 
-	// MultiTransformer abstracts the standard and "quick" Gizmo variants
-	MultiTransformer = NewObject<UMultiTransformer>(this);
-	MultiTransformer->Setup(GetToolManager()->GetPairedGizmoManager(), GetToolManager());
-	MultiTransformer->OnTransformStarted.AddUObject(this, &UEditUVIslandsTool::OnMultiTransformerTransformBegin);
-	MultiTransformer->OnTransformUpdated.AddUObject(this, &UEditUVIslandsTool::OnMultiTransformerTransformUpdate);
-	MultiTransformer->OnTransformCompleted.AddUObject(this, &UEditUVIslandsTool::OnMultiTransformerTransformEnd);
-	MultiTransformer->SetGizmoVisibility(false);
-	MultiTransformer->SetEnabledGizmoSubElements(
+	UInteractiveGizmoManager* GizmoManager = GetToolManager()->GetPairedGizmoManager();
+	TransformGizmo = UE::TransformGizmoUtil::CreateCustomTransformGizmo(GizmoManager,
 		ETransformGizmoSubElements::TranslateAxisX | ETransformGizmoSubElements::TranslateAxisY
 		| ETransformGizmoSubElements::TranslatePlaneXY | ETransformGizmoSubElements::RotateAxisZ
 		| ETransformGizmoSubElements::ScaleAxisX | ETransformGizmoSubElements::ScaleAxisY
-		| ETransformGizmoSubElements::ScalePlaneXY | ETransformGizmoSubElements::ScaleUniform );
-	MultiTransformer->SetOverrideGizmoCoordinateSystem(EToolContextCoordinateSystem::Local);
+		| ETransformGizmoSubElements::ScalePlaneXY | ETransformGizmoSubElements::ScaleUniform,
+		this);
+
+	// Things are probably pretty broken if we're unable to get a transform gizmo... But we do this check in PolyEd,
+	// so might as well be safe here too.
+	if (ensure(TransformGizmo))
+	{
+		// Hook up callbacks
+		TransformProxy = NewObject<UTransformProxy>(this);
+		TransformProxy->OnBeginTransformEdit.AddUObject(this, &UEditUVIslandsTool::OnGizmoTransformBegin);
+		TransformProxy->OnTransformChanged.AddUObject(this, &UEditUVIslandsTool::OnGizmoTransformUpdate);
+		TransformProxy->OnEndTransformEdit.AddUObject(this, &UEditUVIslandsTool::OnGizmoTransformEnd);
+		
+		TransformGizmo->SetActiveTarget(TransformProxy, GetToolManager());
+		TransformGizmo->bUseContextCoordinateSystem = false;
+		TransformGizmo->CurrentCoordinateSystem = EToolContextCoordinateSystem::Local;
+		TransformGizmo->SetVisibility(false);
+	}
+
 
 
 	MaterialSettings = NewObject<UExistingMeshMaterialProperties>(this);
@@ -145,7 +158,6 @@ void UEditUVIslandsTool::Shutdown(EToolShutdownType ShutdownType)
 {
 	MaterialSettings->SaveProperties(this);
 
-	MultiTransformer->Shutdown();
 	SelectionMechanic->Shutdown();
 
 	if (DynamicMeshComponent != nullptr)
@@ -179,6 +191,9 @@ void UEditUVIslandsTool::Shutdown(EToolShutdownType ShutdownType)
 		PreviewMeshActor = nullptr;
 	}
 
+	GetToolManager()->GetPairedGizmoManager()->DestroyAllGizmosByOwner(this);
+	TransformGizmo = nullptr;
+	TransformProxy = nullptr;
 }
 
 
@@ -230,7 +245,12 @@ void UEditUVIslandsTool::OnSelectionModifiedEvent()
 		FFrame3d UseFrame = Topology.GetIslandFrame(
 			SelectionMechanic->GetActiveSelection().GetASelectedGroupID(), GetSpatial());
 		UseFrame.Transform(WorldTransform);
-		MultiTransformer->InitializeGizmoPositionFromWorldFrame(UseFrame, true);
+
+		if (TransformGizmo)
+		{
+			TransformGizmo->ReinitializeGizmoTransform(UseFrame.ToFTransform());
+			TransformGizmo->SetDisplaySpaceTransform(UseFrame.ToFTransform());
+		}
 	}
 }
 
@@ -264,28 +284,33 @@ void UEditUVIslandsTool::OnEndDrag(const FRay& Ray)
 
 
 
-void UEditUVIslandsTool::OnMultiTransformerTransformBegin()
+void UEditUVIslandsTool::OnGizmoTransformBegin(UTransformProxy*)
 {
+	bInDrag = true;
 	SelectionMechanic->ClearHighlight();
 	UpdateUVTransformFromSelection( SelectionMechanic->GetActiveSelection() );
-	InitialGizmoFrame = MultiTransformer->GetCurrentGizmoFrame();
-	InitialGizmoScale = MultiTransformer->GetCurrentGizmoScale();
+	InitialGizmoFrame = FFrame3d(TransformProxy->GetTransform());
+	InitialGizmoScale = TransformProxy->GetTransform().GetScale3D();
 	BeginChange();
 }
 
-void UEditUVIslandsTool::OnMultiTransformerTransformUpdate()
+void UEditUVIslandsTool::OnGizmoTransformUpdate(UTransformProxy*, FTransform)
 {
-	if (MultiTransformer->InGizmoEdit())
+	if (bInDrag)
 	{
 		ComputeUpdate_Gizmo();
 	}
 }
 
-void UEditUVIslandsTool::OnMultiTransformerTransformEnd()
+void UEditUVIslandsTool::OnGizmoTransformEnd(UTransformProxy*)
 {
+	bInDrag = false;
 	SelectionMechanic->NotifyMeshChanged(false);
 
-	MultiTransformer->ResetScale();
+	if (TransformGizmo)
+	{
+		TransformGizmo->SetNewChildScale(FVector::OneVector);
+	}
 
 	// close change record
 	EndChange();
@@ -295,7 +320,7 @@ void UEditUVIslandsTool::OnMultiTransformerTransformEnd()
 
 bool UEditUVIslandsTool::OnUpdateHover(const FInputDeviceRay& DevicePos)
 {
-	if (ActiveVertexChange == nullptr && MultiTransformer->InGizmoEdit() == false )
+	if (ActiveVertexChange == nullptr && !bInDrag)
 	{
 		SelectionMechanic->UpdateHighlight(DevicePos.WorldRay);
 	}
@@ -360,8 +385,8 @@ void UEditUVIslandsTool::ComputeUpdate_Gizmo()
 		return;
 	}
 
-	FFrame3d CurFrame = MultiTransformer->GetCurrentGizmoFrame();
-	FVector3d CurScale = MultiTransformer->GetCurrentGizmoScale();
+	FFrame3d CurFrame = FFrame3d(TransformProxy->GetTransform());
+	FVector3d CurScale = TransformProxy->GetTransform().GetScale3D();
 	FVector3d TranslationDelta = CurFrame.Origin - InitialGizmoFrame.Origin;
 	FQuaterniond RotateDelta = CurFrame.Rotation - InitialGizmoFrame.Rotation;
 	FVector3d CurScaleDelta = CurScale - InitialGizmoScale;
@@ -405,20 +430,14 @@ void UEditUVIslandsTool::ComputeUpdate_Gizmo()
 
 void UEditUVIslandsTool::OnTick(float DeltaTime)
 {
-	MultiTransformer->Tick(DeltaTime);
-
 	if (bSelectionStateDirty)
 	{
 		// update color highlights
 		DynamicMeshComponent->FastNotifySecondaryTrianglesChanged();
 
-		if (SelectionMechanic->HasSelection())
-		{
-			MultiTransformer->SetGizmoVisibility(true);
-		}
-		else
+		if (TransformGizmo)
 		{
-			MultiTransformer->SetGizmoVisibility(false);
+			TransformGizmo->SetVisibility(SelectionMechanic->HasSelection());
 		}
 
 		bSelectionStateDirty = false;
diff --git a/Engine/Plugins/Experimental/MeshModelingToolsetExp/Source/MeshModelingToolsExp/Public/EditUVIslandsTool.h b/Engine/Plugins/Experimental/MeshModelingToolsetExp/Source/MeshModelingToolsExp/Public/EditUVIslandsTool.h
--- a/Engine/Plugins/Experimental/MeshModelingToolsetExp/Source/MeshModelingToolsExp/Public/EditUVIslandsTool.h
+++ b/Engine/Plugins/Experimental/MeshModelingToolsetExp/Source/MeshModelingToolsExp/Public/EditUVIslandsTool.h
@@ -13,10 +13,11 @@
 #include "Selection/GroupTopologySelector.h"
 #include "Selection/PolygonSelectionMechanic.h"
 #include "ToolDataVisualizer.h"
-#include "Transforms/MultiTransformer.h"
 #include "EditUVIslandsTool.generated.h"
 
 class FMeshVertexChangeBuilder;
+class UCombinedTransformGizmo;
+class UTransformProxy;
 using UE::Geometry::FDynamicMeshUVOverlay;
 
 /**
@@ -110,11 +111,14 @@ protected:
 	void OnSelectionModifiedEvent();
 
 	UPROPERTY()
-	TObjectPtr<UMultiTransformer> MultiTransformer = nullptr;
+	TObjectPtr<UCombinedTransformGizmo> TransformGizmo = nullptr;
 
-	void OnMultiTransformerTransformBegin();
-	void OnMultiTransformerTransformUpdate();
-	void OnMultiTransformerTransformEnd();
+	UPROPERTY()
+	TObjectPtr<UTransformProxy> TransformProxy = nullptr;
+
+	void OnGizmoTransformBegin(UTransformProxy* TransformProxy);
+	void OnGizmoTransformUpdate(UTransformProxy* TransformProxy, FTransform Transform);
+	void OnGizmoTransformEnd(UTransformProxy* TransformProxy);
 
 	// realtime visualization
 	void OnDynamicMeshComponentChanged();
