Render Capture Baking: Remove sub-classing of BakeRC from UBakeMeshAttributeMapsToolBase and implement a few fixes

Fixed materials in the BakeRC tool missing UV layer switches or controls for brightness of the Base Color, Subsurface Color or Emissive channels
Removed the unused AO Multiplier setting that appeared in the Preview panel of the tool UI
Improved performance of BakeRC tool baking step by not recomputing the target mesh spatial index and uv charts
Improved the clarity of the analytics code since the sub-classing has enabled us to more easily use custom BakeRC analytics

#rnx

diff --git a/Engine/Plugins/Experimental/GeometryScripting/Source/GeometryScriptingCore/Private/MeshBakeFunctions.cpp b/Engine/Plugins/Experimental/GeometryScripting/Source/GeometryScriptingCore/Private/MeshBakeFunctions.cpp
--- a/Engine/Plugins/Experimental/GeometryScripting/Source/GeometryScriptingCore/Private/MeshBakeFunctions.cpp
+++ b/Engine/Plugins/Experimental/GeometryScripting/Source/GeometryScriptingCore/Private/MeshBakeFunctions.cpp
@@ -1108,6 +1108,8 @@ FGeometryScriptRenderCaptureTextures UGeometryScriptLibrary_MeshBakeFunctions::B
 	const FDynamicMeshAABBTree3 TargetMeshSpatial(TargetMesh->GetMeshPtr());
 	TSharedPtr<FMeshTangentsd, ESPMode::ThreadSafe> TargetMeshTangents = MakeShared<FMeshTangentsd>(TargetMesh->GetMeshPtr());
 	TargetMeshTangents->CopyTriVertexTangents(TargetMesh->GetMeshRef());
+	TSharedPtr<TArray<int32>, ESPMode::ThreadSafe> TargetMeshUVCharts = MakeShared<TArray<int32>>();
+	FMeshMapBaker::ComputeUVCharts(TargetMesh->GetMeshRef(), *TargetMeshUVCharts);
 
 	FSceneCapturePhotoSetSampler Sampler(
 		SceneCapture.Get(),
@@ -1122,6 +1124,7 @@ FGeometryScriptRenderCaptureTextures UGeometryScriptLibrary_MeshBakeFunctions::B
 	TUniquePtr<FMeshMapBaker> Baker = MakeRenderCaptureBaker(
 		TargetMesh->GetMeshPtr(),
 		TargetMeshTangents,
+		TargetMeshUVCharts,
 		SceneCapture.Get(),
 		&Sampler,
 		PendingBake,
diff --git a/Engine/Plugins/Experimental/MeshModelingToolsetExp/Source/MeshModelingToolsExp/Private/BakeMeshAttributeMapsTool.cpp b/Engine/Plugins/Experimental/MeshModelingToolsetExp/Source/MeshModelingToolsExp/Private/BakeMeshAttributeMapsTool.cpp
--- a/Engine/Plugins/Experimental/MeshModelingToolsetExp/Source/MeshModelingToolsExp/Private/BakeMeshAttributeMapsTool.cpp
+++ b/Engine/Plugins/Experimental/MeshModelingToolsetExp/Source/MeshModelingToolsExp/Private/BakeMeshAttributeMapsTool.cpp
@@ -3,6 +3,7 @@
 #include "BakeMeshAttributeMapsTool.h"
 #include "InteractiveToolManager.h"
 #include "ToolBuilderUtil.h"
+#include "ModelingToolTargetUtil.h"
 
 #include "DynamicMesh/DynamicMesh3.h"
 #include "DynamicMesh/DynamicMeshAttributeSet.h"
@@ -313,12 +314,12 @@ void UBakeMeshAttributeMapsTool::Setup()
 	SetToolPropertySourceEnabled(InputMeshSettings, true);
 	InputMeshSettings->bHasTargetUVLayer = true;
 	InputMeshSettings->bHasSourceNormalMap = true;
-	InputMeshSettings->TargetStaticMesh = GetStaticMeshTarget(Target);
-	InputMeshSettings->TargetSkeletalMesh = GetSkeletalMeshTarget(Target);
-	InputMeshSettings->TargetDynamicMesh = GetDynamicMeshTarget(Target);
-	InputMeshSettings->SourceStaticMesh = !bIsBakeToSelf ? GetStaticMeshTarget(DetailTarget) : nullptr;
-	InputMeshSettings->SourceSkeletalMesh = !bIsBakeToSelf ? GetSkeletalMeshTarget(DetailTarget) : nullptr;
-	InputMeshSettings->SourceDynamicMesh = !bIsBakeToSelf ? GetDynamicMeshTarget(DetailTarget) : nullptr;
+	InputMeshSettings->TargetStaticMesh = UE::ToolTarget::GetStaticMeshFromTargetIfAvailable(Target);
+	InputMeshSettings->TargetSkeletalMesh = UE::ToolTarget::GetSkeletalMeshFromTargetIfAvailable(Target);
+	InputMeshSettings->TargetDynamicMesh = GetTargetActorViaIPersistentDynamicMeshSource(Target);
+	InputMeshSettings->SourceStaticMesh = !bIsBakeToSelf ? UE::ToolTarget::GetStaticMeshFromTargetIfAvailable(DetailTarget) : nullptr;
+	InputMeshSettings->SourceSkeletalMesh = !bIsBakeToSelf ? UE::ToolTarget::GetSkeletalMeshFromTargetIfAvailable(DetailTarget) : nullptr;
+	InputMeshSettings->SourceDynamicMesh = !bIsBakeToSelf ? GetTargetActorViaIPersistentDynamicMeshSource(DetailTarget) : nullptr;
 	InputMeshSettings->SourceNormalMap = nullptr;
 	UpdateUVLayerNames(InputMeshSettings->TargetUVLayer, InputMeshSettings->TargetUVLayerNamesList, TargetMesh);
 	InputMeshSettings->WatchProperty(InputMeshSettings->bHideSourceMesh, [this](bool bState) { SetSourceObjectVisible(!bState); });
diff --git a/Engine/Plugins/Experimental/MeshModelingToolsetExp/Source/MeshModelingToolsExp/Private/BakeMeshAttributeMapsToolBase.cpp b/Engine/Plugins/Experimental/MeshModelingToolsetExp/Source/MeshModelingToolsExp/Private/BakeMeshAttributeMapsToolBase.cpp
--- a/Engine/Plugins/Experimental/MeshModelingToolsetExp/Source/MeshModelingToolsExp/Private/BakeMeshAttributeMapsToolBase.cpp
+++ b/Engine/Plugins/Experimental/MeshModelingToolsetExp/Source/MeshModelingToolsExp/Private/BakeMeshAttributeMapsToolBase.cpp
@@ -54,20 +54,11 @@ void UBakeMeshAttributeMapsToolBase::Setup()
 		BentNormalPreviewMaterial = UMaterialInstanceDynamic::Create(BentNormalMaterial, GetToolManager());
 	}
 
-	// Initialize preview mesh
 	UE::ToolTarget::HideSourceObject(Targets[0]);
 
-	const FDynamicMesh3 InputMesh = UE::ToolTarget::GetDynamicMeshCopy(Targets[0], true);
-	const FTransformSRT3d BaseToWorld = UE::ToolTarget::GetLocalToWorldTransform(Targets[0]);
-	PreviewMesh = NewObject<UPreviewMesh>(this);
-	PreviewMesh->CreateInWorld(GetTargetWorld(), FTransform::Identity);
-	ToolSetupUtil::ApplyRenderingConfigurationToPreview(PreviewMesh, nullptr);
-	PreviewMesh->SetTransform(static_cast<FTransform>(BaseToWorld));
-	PreviewMesh->SetTangentsMode(EDynamicMeshComponentTangentsMode::ExternallyProvided);
-	PreviewMesh->ReplaceMesh(InputMesh);
-	PreviewMesh->SetMaterials(UE::ToolTarget::GetMaterialSet(Targets[0]).Materials);
+	// Initialize preview mesh
+	PreviewMesh = CreateBakePreviewMesh(this, Targets[0], GetTargetWorld());
 	PreviewMesh->SetOverrideRenderMaterial(PreviewMaterial);
-	PreviewMesh->SetVisible(true);
 }
 
 void UBakeMeshAttributeMapsToolBase::PostSetup()
@@ -129,11 +120,11 @@ void UBakeMeshAttributeMapsToolBase::OnTick(float DeltaTime)
 void UBakeMeshAttributeMapsToolBase::Render(IToolsContextRenderAPI* RenderAPI)
 {
 	UpdateResult();
-	
+
 	const float Brightness = VisualizationProps->Brightness;
 	const FVector BrightnessColor(Brightness, Brightness, Brightness);
 	const float AOMultiplier = VisualizationProps->AOMultiplier;
-	
+
 	PreviewMaterial->SetVectorParameterValue(TEXT("Brightness"), BrightnessColor );
 	PreviewMaterial->SetScalarParameterValue(TEXT("AOMultiplier"), AOMultiplier );
 	BentNormalPreviewMaterial->SetVectorParameterValue(TEXT("Brightness"), BrightnessColor );
diff --git a/Engine/Plugins/Experimental/MeshModelingToolsetExp/Source/MeshModelingToolsExp/Private/BakeMeshAttributeTool.cpp b/Engine/Plugins/Experimental/MeshModelingToolsetExp/Source/MeshModelingToolsExp/Private/BakeMeshAttributeTool.cpp
--- a/Engine/Plugins/Experimental/MeshModelingToolsetExp/Source/MeshModelingToolsExp/Private/BakeMeshAttributeTool.cpp
+++ b/Engine/Plugins/Experimental/MeshModelingToolsetExp/Source/MeshModelingToolsExp/Private/BakeMeshAttributeTool.cpp
@@ -342,24 +342,6 @@ void UBakeMeshAttributeTool::UpdateMultiTextureMaterialIDs(
 }
 
 
-void UBakeMeshAttributeTool::UpdateUVLayerNames(FString& UVLayer, TArray<FString>& UVLayerNamesList, const FDynamicMesh3& Mesh)
-{
-	UVLayerNamesList.Reset();
-	int32 FoundIndex = -1;
-	for (int32 k = 0; k < Mesh.Attributes()->NumUVLayers(); ++k)
-	{
-		UVLayerNamesList.Add(FString::Printf(TEXT("UV %d"), k));
-		if (UVLayer == UVLayerNamesList.Last())
-		{
-			FoundIndex = k;
-		}
-	}
-	if (FoundIndex == -1)
-	{
-		UVLayer = UVLayerNamesList[0];
-	}
-}
-
 #undef LOCTEXT_NAMESPACE
 
 
diff --git a/Engine/Plugins/Experimental/MeshModelingToolsetExp/Source/MeshModelingToolsExp/Private/BakeMeshAttributeVertexTool.cpp b/Engine/Plugins/Experimental/MeshModelingToolsetExp/Source/MeshModelingToolsExp/Private/BakeMeshAttributeVertexTool.cpp
--- a/Engine/Plugins/Experimental/MeshModelingToolsetExp/Source/MeshModelingToolsExp/Private/BakeMeshAttributeVertexTool.cpp
+++ b/Engine/Plugins/Experimental/MeshModelingToolsetExp/Source/MeshModelingToolsExp/Private/BakeMeshAttributeVertexTool.cpp
@@ -4,6 +4,7 @@
 #include "InteractiveToolManager.h"
 #include "ToolBuilderUtil.h"
 #include "ToolSetupUtil.h"
+#include "ModelingToolTargetUtil.h"
 
 #include "Materials/Material.h"
 #include "Materials/MaterialInstanceDynamic.h"
@@ -272,15 +273,8 @@ void UBakeMeshAttributeVertexTool::Setup()
 
 	// PreviewMesh stores computed result mesh. On shutdown, PreviewMesh will be
 	// used to commit the dynamic mesh to the target tool target.
-	PreviewMesh = NewObject<UPreviewMesh>(this);
-	PreviewMesh->CreateInWorld(GetTargetWorld(), FTransform::Identity);
-	ToolSetupUtil::ApplyRenderingConfigurationToPreview(PreviewMesh, nullptr);
-	PreviewMesh->SetTransform(static_cast<FTransform>(UE::ToolTarget::GetLocalToWorldTransform(Targets[0])));
-	PreviewMesh->SetTangentsMode(EDynamicMeshComponentTangentsMode::ExternallyProvided);
-	PreviewMesh->ReplaceMesh(TargetMesh);
-	PreviewMesh->SetMaterials(UE::ToolTarget::GetMaterialSet(Targets[0]).Materials);
+	PreviewMesh = CreateBakePreviewMesh(this, Targets[0], GetTargetWorld());
 	PreviewMesh->SetOverrideRenderMaterial(PreviewMaterial);
-	PreviewMesh->SetVisible(true);
 
 	UToolTarget* Target = Targets[0];
 	UToolTarget* DetailTarget = Targets[bIsBakeToSelf ? 0 : 1];
@@ -307,12 +301,12 @@ void UBakeMeshAttributeVertexTool::Setup()
 	SetToolPropertySourceEnabled(InputMeshSettings, true);
 	InputMeshSettings->bHasTargetUVLayer = false;
 	InputMeshSettings->bHasSourceNormalMap = false;
-	InputMeshSettings->TargetStaticMesh = GetStaticMeshTarget(Target);
-	InputMeshSettings->TargetSkeletalMesh = GetSkeletalMeshTarget(Target);
-	InputMeshSettings->TargetDynamicMesh = GetDynamicMeshTarget(Target);
-	InputMeshSettings->SourceStaticMesh = !bIsBakeToSelf ? GetStaticMeshTarget(DetailTarget) : nullptr;
-	InputMeshSettings->SourceSkeletalMesh = !bIsBakeToSelf ? GetSkeletalMeshTarget(DetailTarget) : nullptr;
-	InputMeshSettings->SourceDynamicMesh = !bIsBakeToSelf ? GetDynamicMeshTarget(DetailTarget) : nullptr;
+	InputMeshSettings->TargetStaticMesh = UE::ToolTarget::GetStaticMeshFromTargetIfAvailable(Target);
+	InputMeshSettings->TargetSkeletalMesh = UE::ToolTarget::GetSkeletalMeshFromTargetIfAvailable(Target);
+	InputMeshSettings->TargetDynamicMesh = GetTargetActorViaIPersistentDynamicMeshSource(Target);
+	InputMeshSettings->SourceStaticMesh = !bIsBakeToSelf ? UE::ToolTarget::GetStaticMeshFromTargetIfAvailable(DetailTarget) : nullptr;
+	InputMeshSettings->SourceSkeletalMesh = !bIsBakeToSelf ? UE::ToolTarget::GetSkeletalMeshFromTargetIfAvailable(DetailTarget) : nullptr;
+	InputMeshSettings->SourceDynamicMesh = !bIsBakeToSelf ? GetTargetActorViaIPersistentDynamicMeshSource(DetailTarget) : nullptr;
 	InputMeshSettings->SourceNormalMap = nullptr;
 	InputMeshSettings->WatchProperty(InputMeshSettings->bHideSourceMesh, [this](bool bState) { SetSourceObjectVisible(!bState); });
 	InputMeshSettings->WatchProperty(InputMeshSettings->ProjectionDistance, [this](float) { OpState |= EBakeOpState::Evaluate; });
diff --git a/Engine/Plugins/Experimental/MeshModelingToolsetExp/Source/MeshModelingToolsExp/Private/BakeMultiMeshAttributeMapsTool.cpp b/Engine/Plugins/Experimental/MeshModelingToolsetExp/Source/MeshModelingToolsExp/Private/BakeMultiMeshAttributeMapsTool.cpp
--- a/Engine/Plugins/Experimental/MeshModelingToolsetExp/Source/MeshModelingToolsExp/Private/BakeMultiMeshAttributeMapsTool.cpp
+++ b/Engine/Plugins/Experimental/MeshModelingToolsetExp/Source/MeshModelingToolsExp/Private/BakeMultiMeshAttributeMapsTool.cpp
@@ -521,9 +521,9 @@ void UBakeMultiMeshAttributeMapsTool::Setup()
 	InputMeshSettings = NewObject<UBakeMultiMeshInputToolProperties>(this);
 	InputMeshSettings->RestoreProperties(this);
 	AddToolPropertySource(InputMeshSettings);
-	InputMeshSettings->TargetStaticMesh = GetStaticMeshTarget(Target);
-	InputMeshSettings->TargetSkeletalMesh = GetSkeletalMeshTarget(Target);
-	InputMeshSettings->TargetDynamicMesh = GetDynamicMeshTarget(Target);
+	InputMeshSettings->TargetStaticMesh = UE::ToolTarget::GetStaticMeshFromTargetIfAvailable(Target);
+	InputMeshSettings->TargetSkeletalMesh = UE::ToolTarget::GetSkeletalMeshFromTargetIfAvailable(Target);
+	InputMeshSettings->TargetDynamicMesh = GetTargetActorViaIPersistentDynamicMeshSource(Target);
 	UpdateUVLayerNames(InputMeshSettings->TargetUVLayer, InputMeshSettings->TargetUVLayerNamesList, TargetMesh);
 	InputMeshSettings->WatchProperty(InputMeshSettings->TargetUVLayer, [this](FString) { OpState |= EBakeOpState::Evaluate; });
 	InputMeshSettings->WatchProperty(InputMeshSettings->ProjectionDistance, [this](float) { OpState |= EBakeOpState::Evaluate; });
@@ -560,7 +560,7 @@ void UBakeMultiMeshAttributeMapsTool::Setup()
 		});
 
 		FBakeMultiMeshDetailProperties SourceMeshProp;
-		SourceMeshProp.SourceMesh = GetStaticMeshTarget(DetailTarget);
+		SourceMeshProp.SourceMesh = UE::ToolTarget::GetStaticMeshFromTargetIfAvailable(DetailTarget);
 		SourceMeshProp.SourceTexture = DetailColorTexture;
 		InputMeshSettings->SourceMeshes.Add(SourceMeshProp);
 		InputMeshSettings->WatchProperty(InputMeshSettings->SourceMeshes[Idx-1].SourceTexture, [this](UTexture2D*) { OpState |= EBakeOpState::Evaluate; });
diff --git a/Engine/Plugins/Experimental/MeshModelingToolsetExp/Source/MeshModelingToolsExp/Private/BakeRenderCaptureTool.cpp b/Engine/Plugins/Experimental/MeshModelingToolsetExp/Source/MeshModelingToolsExp/Private/BakeRenderCaptureTool.cpp
--- a/Engine/Plugins/Experimental/MeshModelingToolsetExp/Source/MeshModelingToolsExp/Private/BakeRenderCaptureTool.cpp
+++ b/Engine/Plugins/Experimental/MeshModelingToolsetExp/Source/MeshModelingToolsExp/Private/BakeRenderCaptureTool.cpp
@@ -11,6 +11,8 @@
 
 #include "DynamicMesh/MeshTransforms.h"
 
+#include "BakeToolUtils.h"
+#include "ToolSetupUtil.h"
 #include "ModelingToolTargetUtil.h"
 
 #include "ModelingObjectsCreationAPI.h"
@@ -79,7 +81,9 @@ class FRenderCaptureMapBakerOp : public TGenericDataOperator<FMeshMapBaker>
 {
 public:
 	UE::Geometry::FDynamicMesh3* BaseMesh = nullptr;
+	UE::Geometry::FDynamicMeshAABBTree3* BaseMeshSpatial = nullptr;
 	TSharedPtr<UE::Geometry::FMeshTangentsd, ESPMode::ThreadSafe> BaseMeshTangents;
+	TSharedPtr<TArray<int32>, ESPMode::ThreadSafe> BaseMeshUVCharts;
 	int32 TargetUVLayer;
 	double ValidSampleDepthThreshold;
 	EBakeTextureResolution TextureImageSize;
@@ -97,13 +101,11 @@ public:
 // Bake textures onto the base/target mesh by projecting/sampling the set of captured photos
 void FRenderCaptureMapBakerOp::CalculateResult(FProgressCancel*)
 {
-	const FDynamicMeshAABBTree3 BaseMeshSpatial(BaseMesh);
-
 	FSceneCapturePhotoSetSampler Sampler(
 		SceneCapture,
 		ValidSampleDepthThreshold,
 		BaseMesh,
-		&BaseMeshSpatial,
+		BaseMeshSpatial,
 		BaseMeshTangents.Get());
 
 	const FImageDimensions TextureDimensions(
@@ -115,6 +117,7 @@ void FRenderCaptureMapBakerOp::CalculateResult(FProgressCancel*)
 	Result = MakeRenderCaptureBaker(
 		BaseMesh,
 		BaseMeshTangents,
+		BaseMeshUVCharts,
 		SceneCapture,
 		&Sampler,
 		PendingBake,
@@ -172,19 +175,31 @@ void UBakeRenderCaptureTool::Setup()
 
 	Super::Setup();
 
+	UE::ToolTarget::HideSourceObject(Targets[0]);
+
+	// Initialize materials and textures for background compute, error feedback and previewing the results
 	InitializePreviewMaterials();
 
-	// Initialize base mesh
-	const FTransformSRT3d BaseToWorld = UE::ToolTarget::GetLocalToWorldTransform(Targets[0]);
-	PreviewMesh->ProcessMesh([this, BaseToWorld](const FDynamicMesh3& Mesh)
+	// Initialize the PreviewMesh, which displays intermediate results
+	PreviewMesh = CreateBakePreviewMesh(this, Targets[0], GetTargetWorld());
+
+	// Initialize the datastructures used by the bake background compute/tool operator
+	PreviewMesh->ProcessMesh([this](const FDynamicMesh3& Mesh)
 	{
 		TargetMesh.Copy(Mesh);
-		TargetMeshTangents = MakeShared<FMeshTangentsd, ESPMode::ThreadSafe>(&TargetMesh);
-		TargetMeshTangents->CopyTriVertexTangents(Mesh);
-		
-		// FMeshSceneAdapter operates in world space, so ensure our mesh transformed to world.
+		const FTransformSRT3d BaseToWorld = UE::ToolTarget::GetLocalToWorldTransform(Targets[0]);
 		MeshTransforms::ApplyTransform(TargetMesh, BaseToWorld, true);
-		TargetSpatial.SetMesh(&TargetMesh, true);
+
+		// Initialize UV charts
+		TargetMeshUVCharts = MakeShared<TArray<int32>, ESPMode::ThreadSafe>();
+		FMeshMapBaker::ComputeUVCharts(TargetMesh, *TargetMeshUVCharts);
+
+		// Initialize tangents
+		TargetMeshTangents = MakeShared<FMeshTangentsd, ESPMode::ThreadSafe>(&TargetMesh);
+		TargetMeshTangents->CopyTriVertexTangents(TargetMesh);
+
+		// Initialize spatial index
+		TargetMeshSpatial.SetMesh(&TargetMesh, true);
 	});
 
 	// Initialize actors
@@ -252,7 +267,7 @@ void UBakeRenderCaptureTool::Setup()
 	InputMeshSettings = NewObject<UBakeRenderCaptureInputToolProperties>(this);
 	InputMeshSettings->RestoreProperties(this);
 	AddToolPropertySource(InputMeshSettings);
-	InputMeshSettings->TargetStaticMesh = GetStaticMeshTarget(Target);
+	InputMeshSettings->TargetStaticMesh = UE::ToolTarget::GetStaticMeshFromTargetIfAvailable(Target);
 	UpdateUVLayerNames(InputMeshSettings->TargetUVLayer, InputMeshSettings->TargetUVLayerNamesList, TargetMesh);
 	InputMeshSettings->WatchProperty(InputMeshSettings->TargetUVLayer, [this](FString) { OpState |= EBakeOpState::Evaluate; });
 	
@@ -273,6 +288,11 @@ void UBakeRenderCaptureTool::Setup()
 	AddToolPropertySource(ResultSettings);
 	SetToolPropertySourceEnabled(ResultSettings, true);
 
+	VisualizationProps = NewObject<UBakeRenderCaptureVisualizationProperties>(this);
+	VisualizationProps->RestoreProperties(this);
+	AddToolPropertySource(VisualizationProps);
+	VisualizationProps->WatchProperty(VisualizationProps->bPreviewAsMaterial, [this](bool) { UpdateVisualization(); });
+
 	TargetUVLayerToError.Reset();
 
 	// Hide the render capture meshes since this baker operates solely in world space which will occlude the preview of
@@ -290,7 +310,7 @@ void UBakeRenderCaptureTool::Setup()
 		LOCTEXT("OnStartTool", "Bake Render Capture. Select Bake Mesh (LowPoly) first, then select Detail Meshes (HiPoly) to bake. Assets will be created on Accept."),
 		EToolMessageLevel::UserNotification);
 
-	PostSetup();
+	GatherAnalytics(BakeAnalytics.MeshSettings);
 }
 
 
@@ -299,28 +319,76 @@ void UBakeRenderCaptureTool::Setup()
 
 void UBakeRenderCaptureTool::Render(IToolsContextRenderAPI* RenderAPI)
 {
-	Super::Render(RenderAPI);
-
-	const float Brightness = VisualizationProps->Brightness;
-	const FVector BrightnessColor(Brightness, Brightness, Brightness);
-	PreviewMaterialRC->SetVectorParameterValue(TEXT("Brightness"), BrightnessColor);
-	PreviewMaterialPackedRC->SetVectorParameterValue(TEXT("Brightness"), BrightnessColor);
-	PreviewMaterialRC_Subsurface->SetVectorParameterValue(TEXT("Brightness"), BrightnessColor);
-	PreviewMaterialPackedRC_Subsurface->SetVectorParameterValue(TEXT("Brightness"), BrightnessColor);
+	// This block contains things we want to update every frame to respond to sliders in the UI
+	{
+		const float Brightness = VisualizationProps->Brightness;
+		const FVector BaseColorBrightness(Brightness, Brightness, Brightness);
+
+		PreviewMaterialRC->SetVectorParameterValue(TEXT("Brightness"), BaseColorBrightness);
+		PreviewMaterialPackedRC->SetVectorParameterValue(TEXT("Brightness"), BaseColorBrightness);
+		PreviewMaterialRC_Subsurface->SetVectorParameterValue(TEXT("Brightness"), BaseColorBrightness);
+		PreviewMaterialPackedRC_Subsurface->SetVectorParameterValue(TEXT("Brightness"), BaseColorBrightness);
+
+		const float EmissiveScale = VisualizationProps->EmissiveScale;
+
+		PreviewMaterialRC->SetScalarParameterValue(TEXT("EmissiveScale"), EmissiveScale);
+		PreviewMaterialPackedRC->SetScalarParameterValue(TEXT("EmissiveScale"), EmissiveScale);
+		PreviewMaterialRC_Subsurface->SetScalarParameterValue(TEXT("EmissiveScale"), EmissiveScale);
+		PreviewMaterialPackedRC_Subsurface->SetScalarParameterValue(TEXT("EmissiveScale"), EmissiveScale);
+
+		const float SSBrightness = VisualizationProps->SSBrightness;
+		const FVector SubsurfaceColorBrightness(SSBrightness, SSBrightness, SSBrightness);
+
+		PreviewMaterialRC_Subsurface->SetVectorParameterValue(TEXT("BrightnessSubsurface"), SubsurfaceColorBrightness);
+		PreviewMaterialPackedRC_Subsurface->SetVectorParameterValue(TEXT("BrightnessSubsurface"), SubsurfaceColorBrightness);
+	}
+
+	// This is a no-op if OpState == EBakeOpState::Clean
+	UpdateResult();
 }
 
 
+void UBakeRenderCaptureTool::OnTick(float DeltaTime)
+{
+	if (Compute)
+	{
+		Compute->Tick(DeltaTime);
+
+		if (static_cast<bool>(OpState & EBakeOpState::Invalid))
+		{
+			PreviewMesh->SetOverrideRenderMaterial(ErrorPreviewMaterial);
+		}
+		else
+		{
+			const float ElapsedComputeTime = Compute->GetElapsedComputeTime();
+			if (!CanAccept() && ElapsedComputeTime > SecondsBeforeWorkingMaterial)
+			{
+				PreviewMesh->SetOverrideRenderMaterial(WorkingPreviewMaterial);
+			}
+		}
+	}
+	else if (static_cast<bool>(OpState & EBakeOpState::Invalid))
+	{
+		PreviewMesh->SetOverrideRenderMaterial(ErrorPreviewMaterial);
+	} 
+}
 
 
 void UBakeRenderCaptureTool::OnShutdown(EToolShutdownType ShutdownType)
 {
 	TRACE_CPUPROFILER_EVENT_SCOPE(UBakeRenderCaptureTool::Shutdown);
-	
-	Super::OnShutdown(ShutdownType);
-	
+
 	Settings->SaveProperties(this);
 	RenderCaptureProperties->SaveProperties(this);
 	InputMeshSettings->SaveProperties(this);
+	VisualizationProps->SaveProperties(this);
+
+	if (PreviewMesh != nullptr)
+	{
+		PreviewMesh->SetVisible(false);
+		PreviewMesh->Disconnect();
+		PreviewMesh = nullptr;
+	}
 
 	if (Compute)
 	{
@@ -336,17 +404,20 @@ void UBakeRenderCaptureTool::OnShutdown(EToolShutdownType ShutdownType)
 	
 	if (ShutdownType == EToolShutdownType::Accept)
 	{
+		// TODO Support skeletal meshes here---see BakeMeshAttributeMapsTool::OnShutdown
 		IStaticMeshBackedTarget* StaticMeshTarget = Cast<IStaticMeshBackedTarget>(Targets[0]);
 		UObject* SourceAsset = StaticMeshTarget ? StaticMeshTarget->GetStaticMesh() : nullptr;
 		const UPrimitiveComponent* SourceComponent = UE::ToolTarget::GetTargetComponent(Targets[0]);
-		CreateTextureAssetsRC(SourceComponent->GetWorld(), SourceAsset);
+		CreateTextureAssets(SourceComponent->GetWorld(), SourceAsset);
 	}
 
 	// Clear actors on shutdown so that their lifetime is not tied to the lifetime of the tool
 	Actors.Empty();
+
+	UE::ToolTarget::ShowSourceObject(Targets[0]);
 }
 
-void UBakeRenderCaptureTool::CreateTextureAssetsRC(UWorld* SourceWorld, UObject* SourceAsset)
+void UBakeRenderCaptureTool::CreateTextureAssets(UWorld* SourceWorld, UObject* SourceAsset)
 {
 	bool bCreatedAssetOK = true;
 	const FString BaseName = UE::ToolTarget::GetTargetActor(Targets[0])->GetActorNameOrLabel();
@@ -501,7 +572,9 @@ TUniquePtr<TGenericDataOperator<FMeshMapBaker>> UBakeRenderCaptureTool::MakeNewO
 {
 	TUniquePtr<FRenderCaptureMapBakerOp> Op = MakeUnique<FRenderCaptureMapBakerOp>();
 	Op->BaseMesh = &TargetMesh;
+	Op->BaseMeshSpatial = &TargetMeshSpatial;
 	Op->BaseMeshTangents = TargetMeshTangents;
+	Op->BaseMeshUVCharts = TargetMeshUVCharts;
 	Op->TargetUVLayer = InputMeshSettings->GetTargetUVLayerIndex();
 	Op->ValidSampleDepthThreshold = Settings->ValidSampleDepthThreshold;
 	Op->TextureImageSize = Settings->TextureSize;
@@ -530,7 +603,7 @@ TUniquePtr<TGenericDataOperator<FMeshMapBaker>> UBakeRenderCaptureTool::MakeNewO
 
 
 
-void UBakeRenderCaptureTool::OnMapsUpdatedRC(const TUniquePtr<FMeshMapBaker>& NewResult)
+void UBakeRenderCaptureTool::OnMapsUpdated(const TUniquePtr<FMeshMapBaker>& NewResult)
 {
 	TRACE_CPUPROFILER_EVENT_SCOPE(BakeRenderCaptureTool_Textures_BuildTextures);
 
@@ -581,10 +654,40 @@ void UBakeRenderCaptureTool::OnMapsUpdatedRC(const TUniquePtr<FMeshMapBaker>& Ne
 	UpdateVisualization();
 }
 
+bool UBakeRenderCaptureTool::ValidTargetMeshTangents()
+{
+	if (bCheckTargetMeshTangents)
+	{
+		bValidTargetMeshTangents = TargetMeshTangents ? FDynamicMeshTangents(&TargetMesh).HasValidTangents(true) : false;
+		bCheckTargetMeshTangents = false;
+	}
+	return bValidTargetMeshTangents;
+}
 
 void UBakeRenderCaptureTool::InitializePreviewMaterials()
 {
-	// EmptyColorMapWhite, EmptyColorMapBlack and EmptyNormalMap are defined in the base tool
+	{
+		FTexture2DBuilder Builder;
+		Builder.Initialize(FTexture2DBuilder::ETextureType::NormalMap, FImageDimensions(16, 16));
+		Builder.Commit(false);
+		EmptyNormalMap = Builder.GetTexture2D();
+	}
+
+	{
+		FTexture2DBuilder Builder;
+		Builder.Initialize(FTexture2DBuilder::ETextureType::Color, FImageDimensions(16, 16));
+		Builder.Clear(FColor(0,0,0));
+		Builder.Commit(false);
+		EmptyColorMapBlack = Builder.GetTexture2D();
+	}
+
+	{
+		FTexture2DBuilder Builder;
+		Builder.Initialize(FTexture2DBuilder::ETextureType::Color, FImageDimensions(16, 16));
+		Builder.Clear(FColor::White);
+		Builder.Commit(false);
+		EmptyColorMapWhite = Builder.GetTexture2D();
+	}
 
 	{
 		FTexture2DBuilder Builder;
@@ -639,7 +742,11 @@ void UBakeRenderCaptureTool::InitializePreviewMaterials()
 		Builder.Commit(false);
 		EmptySpecularMap = Builder.GetTexture2D();
 	}
-	
+
+	WorkingPreviewMaterial = ToolSetupUtil::GetDefaultWorkingMaterialInstance(GetToolManager());
+
+	ErrorPreviewMaterial = ToolSetupUtil::GetDefaultErrorMaterial(GetToolManager());
+
 	{
 		UMaterial* Material = LoadObject<UMaterial>(nullptr,
 			TEXT("/MeshModelingToolsetExp/Materials/BakeRenderCapturePreviewMaterial"));
@@ -708,18 +815,17 @@ void UBakeRenderCaptureTool::InitializePreviewMaterials()
 }
 
 
-void UBakeRenderCaptureTool::InvalidateComputeRC()
+void UBakeRenderCaptureTool::InvalidateCompute()
 {
-	// Note: This implementation is identical to UBakeMeshAttributeMapsToolBase::InvalidateCompute but calls
-	// OnMapsUpdatedRC rather than OnMapsUpdated
 	if (!Compute)
 	{
 		// Initialize background compute
 		Compute = MakeUnique<TGenericDataBackgroundCompute<FMeshMapBaker>>();
 		Compute->Setup(this);
-		Compute->OnResultUpdated.AddLambda([this](const TUniquePtr<FMeshMapBaker>& NewResult) { OnMapsUpdatedRC(NewResult); });
+		Compute->OnResultUpdated.AddLambda([this](const TUniquePtr<FMeshMapBaker>& NewResult) { OnMapsUpdated(NewResult); });
 	}
 	Compute->InvalidateResult();
+	OpState = EBakeOpState::Clean;
 }
 
 FRenderCaptureUpdate UBakeRenderCaptureTool::UpdateSceneCapture()
@@ -788,7 +894,7 @@ FRenderCaptureUpdate UBakeRenderCaptureTool::UpdateSceneCapture()
 	return Update;
 }
 
-// Process dirty props and update background compute. Called by UBakeMeshAttributeMapsToolBase::Render
+// Process dirty props and update background compute. Called by Render
 void UBakeRenderCaptureTool::UpdateResult()
 {
 	// Return if the bake is already launched/complete.
@@ -885,15 +991,12 @@ void UBakeRenderCaptureTool::UpdateResult()
 		UpdateVisualization();
 
 		// Start another bake operation
-		InvalidateComputeRC();
+		InvalidateCompute();
 
 		// Cache computed parameters which are used to determine if results need re-baking
 		ComputedTextureSize = Settings->TextureSize;
 		ComputedSamplesPerPixel = Settings->SamplesPerPixel;
 		ComputedValidDepthThreshold = Settings->ValidSampleDepthThreshold;
-
-		// Set a clean op state so we don't render this function until its completed
-		OpState = EBakeOpState::Clean;
 	}
 }
 
diff --git a/Engine/Plugins/Experimental/MeshModelingToolsetExp/Source/MeshModelingToolsExp/Private/BakeToolUtils.cpp b/Engine/Plugins/Experimental/MeshModelingToolsetExp/Source/MeshModelingToolsExp/Private/BakeToolUtils.cpp
new file mode 100644
--- /dev/null
+++ b/Engine/Plugins/Experimental/MeshModelingToolsetExp/Source/MeshModelingToolsExp/Private/BakeToolUtils.cpp
@@ -0,0 +1,76 @@
+// Copyright Epic Games, Inc. All Rights Reserved.
+
+#include "BakeToolUtils.h"
+
+#include "PreviewMesh.h"
+#include "ToolSetupUtil.h"
+#include "ModelingToolTargetUtil.h"
+#include "DynamicMesh/DynamicMesh3.h"
+#include "Components/DynamicMeshComponent.h"
+#include "TargetInterfaces/DynamicMeshSource.h"
+#include "ToolTargets/ToolTarget.h"
+
+namespace UE
+{
+namespace Geometry
+{
+
+void UpdateUVLayerNames(FString& UVLayer, TArray<FString>& UVLayerNamesList, const FDynamicMesh3& Mesh)
+{
+	UVLayerNamesList.Reset();
+	int32 FoundIndex = -1;
+	for (int32 k = 0; k < Mesh.Attributes()->NumUVLayers(); ++k)
+	{
+		UVLayerNamesList.Add(FString::Printf(TEXT("UV %d"), k));
+		if (UVLayer == UVLayerNamesList.Last())
+		{
+			FoundIndex = k;
+		}
+	}
+	if (FoundIndex == -1)
+	{
+		UVLayer = UVLayerNamesList[0];
+	}
+}
+
+// TODO Some variation of this function may be able to be reused in other tools
+UPreviewMesh* CreateBakePreviewMesh(
+	UObject* Tool,
+	UToolTarget* ToolTarget,
+	UWorld* World)
+{
+	const FDynamicMesh3 InputMesh = UE::ToolTarget::GetDynamicMeshCopy(ToolTarget, true);
+	const FTransformSRT3d BaseToWorld = UE::ToolTarget::GetLocalToWorldTransform(ToolTarget);
+	const FComponentMaterialSet MaterialSet = UE::ToolTarget::GetMaterialSet(ToolTarget);
+
+	UPreviewMesh* PreviewMesh = NewObject<UPreviewMesh>(Tool);
+
+	PreviewMesh->CreateInWorld(World, FTransform::Identity);
+	ToolSetupUtil::ApplyRenderingConfigurationToPreview(PreviewMesh, nullptr);
+	PreviewMesh->SetTransform(static_cast<FTransform>(BaseToWorld));
+	PreviewMesh->SetTangentsMode(EDynamicMeshComponentTangentsMode::ExternallyProvided);
+	PreviewMesh->ReplaceMesh(InputMesh);
+	PreviewMesh->SetMaterials(MaterialSet.Materials);
+	PreviewMesh->SetVisible(true);
+
+	return PreviewMesh;
+}
+
+
+// TODO We could probably use UE::ToolTarget::GetTargetActor to implement/replace this function
+AActor* GetTargetActorViaIPersistentDynamicMeshSource(UToolTarget* Target)
+{
+	IPersistentDynamicMeshSource* TargetDynamicMeshTarget = Cast<IPersistentDynamicMeshSource>(Target);
+	if (TargetDynamicMeshTarget)
+	{
+		UDynamicMeshComponent* TargetDynamicMeshComponent = TargetDynamicMeshTarget->GetDynamicMeshComponent();
+		if (TargetDynamicMeshComponent)
+		{
+			return TargetDynamicMeshComponent->GetOwner();
+		}
+	}
+	return nullptr;
+}
+
+} // namespace Geometry
+} // namespace UE
diff --git a/Engine/Plugins/Experimental/MeshModelingToolsetExp/Source/MeshModelingToolsExp/Public/BakeMeshAttributeTool.h b/Engine/Plugins/Experimental/MeshModelingToolsetExp/Source/MeshModelingToolsExp/Public/BakeMeshAttributeTool.h
--- a/Engine/Plugins/Experimental/MeshModelingToolsetExp/Source/MeshModelingToolsExp/Public/BakeMeshAttributeTool.h
+++ b/Engine/Plugins/Experimental/MeshModelingToolsetExp/Source/MeshModelingToolsExp/Public/BakeMeshAttributeTool.h
@@ -10,6 +10,7 @@
 #include "InteractiveToolQueryInterfaces.h" // for UInteractiveToolExclusiveToolAPI
 #include "Materials/MaterialInstanceDynamic.h"
 #include "PreviewMesh.h"
+#include "BakeToolUtils.h"
 #include "TargetInterfaces/DynamicMeshSource.h"
 #include "TargetInterfaces/SkeletalMeshBackedTarget.h"
 #include "TargetInterfaces/StaticMeshBackedTarget.h"
@@ -171,32 +172,7 @@ protected:
 	//
 	// Utilities
 	//
-	
-	/** @return StaticMesh from a tool target */ 
-	static UStaticMesh* GetStaticMeshTarget(UToolTarget* Target)
-	{
-		IStaticMeshBackedTarget* TargetStaticMeshTarget = Cast<IStaticMeshBackedTarget>(Target);
-		UStaticMesh* TargetStaticMesh = TargetStaticMeshTarget ? TargetStaticMeshTarget->GetStaticMesh() : nullptr;
-		return TargetStaticMesh;
-	}
-
-	/** @return SkeletalMesh from a tool target */
-	static USkeletalMesh* GetSkeletalMeshTarget(UToolTarget* Target)
-	{
-		ISkeletalMeshBackedTarget* TargetSkeletalMeshTarget = Cast<ISkeletalMeshBackedTarget>(Target);
-		USkeletalMesh* TargetSkeletalMesh = TargetSkeletalMeshTarget ? TargetSkeletalMeshTarget->GetSkeletalMesh() : nullptr;
-		return TargetSkeletalMesh;
-	}
 
-	/** @return AActor that owns a DynamicMeshComponent from a tool target */
-	static AActor* GetDynamicMeshTarget(UToolTarget* Target)
-	{
-		IPersistentDynamicMeshSource* TargetDynamicMeshTarget = Cast<IPersistentDynamicMeshSource>(Target);
-		UDynamicMeshComponent* TargetDynamicMeshComponent = TargetDynamicMeshTarget ? TargetDynamicMeshTarget->GetDynamicMeshComponent() : nullptr;
-		AActor* TargetDynamicMesh = TargetDynamicMeshComponent ? TargetDynamicMeshComponent->GetOwner() : nullptr;
-		return TargetDynamicMesh;
-	}
-		
 	/**
 	 * Given an array of textures associated with a material,
 	 * use heuristics to identify the color/albedo texture.
@@ -223,17 +199,6 @@ protected:
 		UToolTarget* Target,
 		TArray<TObjectPtr<UTexture2D>>& AllSourceTextures,
 		TArray<TObjectPtr<UTexture2D>>& MaterialIDTextures);
-
-	/**
-	 * Updates a tool property set's UVLayerNamesList from the list of UV layers
-	 * on a given mesh. Also updates the UVLayer property if the current UV layer
-	 * is no longer available.
-	 *
-	 * @param UVLayer Selected UV Layer.
-	 * @param UVLayerNamesList List of available UV layers.
-	 * @param Mesh the mesh to query
-	 */
-	static void UpdateUVLayerNames(FString& UVLayer, TArray<FString>& UVLayerNamesList, const FDynamicMesh3& Mesh);
 };
 
 
diff --git a/Engine/Plugins/Experimental/MeshModelingToolsetExp/Source/MeshModelingToolsExp/Public/BakeRenderCaptureTool.h b/Engine/Plugins/Experimental/MeshModelingToolsetExp/Source/MeshModelingToolsExp/Public/BakeRenderCaptureTool.h
--- a/Engine/Plugins/Experimental/MeshModelingToolsetExp/Source/MeshModelingToolsExp/Public/BakeRenderCaptureTool.h
+++ b/Engine/Plugins/Experimental/MeshModelingToolsetExp/Source/MeshModelingToolsExp/Public/BakeRenderCaptureTool.h
@@ -253,6 +253,31 @@ public:
 	TArray<FString> TargetUVLayerNamesList;
 };
 
+UCLASS()
+class MESHMODELINGTOOLSEXP_API UBakeRenderCaptureVisualizationProperties : public UInteractiveToolPropertySet
+{
+	GENERATED_BODY()
+public:
+	/**
+	 * If true  preview results by connecting them to corresponding material inputs
+	 * If false connect the selected preview output as the Base Color input and use empty maps for other material inputs
+	 */
+	UPROPERTY(EditAnywhere, Category = Preview)
+	bool bPreviewAsMaterial = false;
+
+	/** Adjust the brightness of the Base Color material input; does not affect results stored in textures */
+	UPROPERTY(EditAnywhere, Category = Preview, meta = (UIMin = "0.0", UIMax = "1.0"))
+	float Brightness = 1.0f;
+
+	/** Adjust the brightness of the Subsurface Color material input; does not affect results stored in textures */
+	UPROPERTY(EditAnywhere, Category = Preview, meta = (UIMin = "0.0", UIMax = "1.0", DisplayName = "SS Brightness"))
+	float SSBrightness = 1.0f;
+
+	/** Adjust the brightness of the Emissive Color material input; does not affect results stored in textures */
+	UPROPERTY(EditAnywhere, Category = Preview, meta = (UIMin = "0.0", UIMax = "1.0"))
+	float EmissiveScale = 1.0f;
+};
+
 
 //
 // Tool
@@ -261,21 +286,25 @@ public:
 
 
 UCLASS()
-class MESHMODELINGTOOLSEXP_API UBakeRenderCaptureTool : public UBakeMeshAttributeMapsToolBase
+class MESHMODELINGTOOLSEXP_API UBakeRenderCaptureTool :
+	public UMultiSelectionMeshEditingTool,
+	public UE::Geometry::IGenericDataOperatorFactory<UE::Geometry::FMeshMapBaker>,
+	public IInteractiveToolExclusiveToolAPI
 {
 	GENERATED_BODY()
 
 public:
 	UBakeRenderCaptureTool() = default;
 
-	// Begin UInteractiveTool interface
+	// Begin UMultiSelectionMeshEditing < UMultiSelectionTool < UInteractiveTool interface
 	virtual void Setup() override;
 	virtual void Render(IToolsContextRenderAPI* RenderAPI) override;
+	virtual void OnTick(float DeltaTime) override;
 	virtual void OnShutdown(EToolShutdownType ShutdownType) override;
 	virtual bool HasCancel() const override { return true; }
 	virtual bool HasAccept() const override { return true; }
 	virtual bool CanAccept() const override;
-	// End UInteractiveTool interface
+	// End UMultiSelectionMeshEditing < UMultiSelectionTool < UInteractiveTool interface
 
 	// Begin IGenericDataOperatorFactory interface
 	virtual TUniquePtr<UE::Geometry::TGenericDataOperator<UE::Geometry::FMeshMapBaker>> MakeNewOperator() override;
@@ -295,6 +324,9 @@ protected:
 	UPROPERTY()
 	TObjectPtr<UBakeRenderCaptureInputToolProperties> InputMeshSettings;
 
+	UPROPERTY()
+	TObjectPtr<UBakeRenderCaptureVisualizationProperties> VisualizationProps;
+
 	// The computed textures are displayed in the details panel and used in the preview material, they are written
 	// out to assest on shutdown.
 	UPROPERTY()
@@ -302,29 +334,45 @@ protected:
 	
 protected:
 
-	// Begin UBakeMeshAttributeMapsToolBase interface
-	virtual void UpdateResult() override;
-	virtual void UpdateVisualization() override;
-	virtual void GatherAnalytics(FBakeAnalytics::FMeshSettings& Data) override;
-	virtual FString GetAnalyticsEventName() const override { return TEXT("BakeRC"); }
-	// End UBakeMeshAttributeMapsToolBase interface
-
+	void UpdateResult();
+	void UpdateVisualization();
 	void InvalidateResults(UE::Geometry::FRenderCaptureUpdate Update);
+	void InvalidateCompute();
+	void OnMapsUpdated(const TUniquePtr<UE::Geometry::FMeshMapBaker>& NewResult);
+
+	/**
+	 * Compute validity of the Target Mesh tangents. Only checks validity
+	 * once and then caches the result for successive calls.
+	 * 
+	 * @return true if the TargetMesh tangents are valid
+	 */
+	bool ValidTargetMeshTangents();
+	bool bCheckTargetMeshTangents = true;
+	bool bValidTargetMeshTangents = false;
 
-	// In this tool we don't call UBakeMeshAttributeMapsToolBase::OnMapsUpdated because it would require e.g, adding
-	// the render capture channels to EBakeMapType.  The implementation is simpler and leads to less coupling if we
-	// just implement custom versions of the following functions.
-	void InvalidateComputeRC();
-	void OnMapsUpdatedRC(const TUniquePtr<UE::Geometry::FMeshMapBaker>& NewResult);
-	
 	/**
 	 * Create texture assets from our result map of Texture2D
 	 * @param SourceWorld the source world to define where the texture assets will be stored.
 	 * @param SourceAsset if not null, result textures will be stored adjacent to this asset.
 	 */
-	void CreateTextureAssetsRC(UWorld* SourceWorld, UObject* SourceAsset);
+	void CreateTextureAssets(UWorld* SourceWorld, UObject* SourceAsset);
+
+	// The baking background compute operation
+	TUniquePtr<TGenericDataBackgroundCompute<UE::Geometry::FMeshMapBaker>> Compute = nullptr;
+	EBakeOpState OpState = EBakeOpState::Evaluate;
+
+	UE::Geometry::FDynamicMesh3 TargetMesh;
+	UE::Geometry::FDynamicMeshAABBTree3 TargetMeshSpatial;
+	TSharedPtr<TArray<int32>, ESPMode::ThreadSafe> TargetMeshUVCharts;
+	TSharedPtr<UE::Geometry::FMeshTangentsd, ESPMode::ThreadSafe> TargetMeshTangents;
 
-	// empty maps are shown when nothing is computed
+	// Empty maps are shown when nothing is computed
+	UPROPERTY()
+	TObjectPtr<UTexture2D> EmptyNormalMap;
+	UPROPERTY()
+	TObjectPtr<UTexture2D> EmptyColorMapBlack;
+	UPROPERTY()
+	TObjectPtr<UTexture2D> EmptyColorMapWhite;
 	UPROPERTY()
 	TObjectPtr<UTexture2D> EmptyEmissiveMap;
 	UPROPERTY()
@@ -339,6 +387,13 @@ protected:
 	TObjectPtr<UTexture2D> EmptyMetallicMap;
 	UPROPERTY()
 	TObjectPtr<UTexture2D> EmptySpecularMap;
+
+	float SecondsBeforeWorkingMaterial = 0.75;
+
+	UPROPERTY()
+	TObjectPtr<UMaterialInstanceDynamic> WorkingPreviewMaterial;
+	UPROPERTY()
+	TObjectPtr<UMaterialInstanceDynamic> ErrorPreviewMaterial;
 	UPROPERTY()
 	TObjectPtr<UMaterialInstanceDynamic> PreviewMaterialRC;
 	UPROPERTY()
@@ -347,8 +402,12 @@ protected:
 	TObjectPtr<UMaterialInstanceDynamic> PreviewMaterialRC_Subsurface;
 	UPROPERTY()
 	TObjectPtr<UMaterialInstanceDynamic> PreviewMaterialPackedRC_Subsurface;
+
 	void InitializePreviewMaterials();
 
+	UPROPERTY()
+	TObjectPtr<UPreviewMesh> PreviewMesh;
+
 	// TODO We currently need to compute this on the game thread because the implementation has checks for this
 	TUniquePtr<UE::Geometry::FSceneCapturePhotoSet> SceneCapture;
 
@@ -362,9 +421,30 @@ protected:
 	//
 	// Analytics
 	//
-
-	void RecordAnalytics() const;
+	
+	struct FBakeAnalytics
+	{
+		double TotalBakeDuration = 0.0;
+		double WriteToImageDuration = 0.0;
+		double WriteToGutterDuration = 0.0;
+		int64 NumSamplePixels = 0;
+		int64 NumGutterPixels = 0;
+
+		struct FMeshSettings
+		{
+			int32 NumTargetMeshTris = 0;
+			int32 NumDetailMesh = 0;
+			int64 NumDetailMeshTris = 0;
+		};
+		FMeshSettings MeshSettings;
+	};
+
+	FBakeAnalytics BakeAnalytics;
+
+	FString GetAnalyticsEventName() const { return TEXT("BakeRC"); }
+	void GatherAnalytics(FBakeAnalytics::FMeshSettings& Data);
 	void GatherAnalytics(const UE::Geometry::FMeshMapBaker& Result);
+	void RecordAnalytics() const;
 
 private:
 	UE::Geometry::FRenderCaptureUpdate UpdateSceneCapture();
diff --git a/Engine/Plugins/Experimental/MeshModelingToolsetExp/Source/MeshModelingToolsExp/Public/BakeToolUtils.h b/Engine/Plugins/Experimental/MeshModelingToolsetExp/Source/MeshModelingToolsExp/Public/BakeToolUtils.h
new file mode 100644
--- /dev/null
+++ b/Engine/Plugins/Experimental/MeshModelingToolsetExp/Source/MeshModelingToolsExp/Public/BakeToolUtils.h
@@ -0,0 +1,36 @@
+// Copyright Epic Games, Inc. All Rights Reserved.
+
+#pragma once
+
+#include "CoreMinimal.h"
+
+class UPreviewMesh;
+class UToolTarget;
+class AActor;
+
+namespace UE
+{
+namespace Geometry
+{
+
+class FDynamicMesh3;
+
+/**
+ * Updates a tool property set's UVLayerNamesList from the list of UV layers
+ * on a given mesh. Also updates the UVLayer property if the current UV layer
+ * is no longer available.
+ *
+ * @param UVLayer Selected UV Layer.
+ * @param UVLayerNamesList List of available UV layers.
+ * @param Mesh the mesh to query
+ */
+MESHMODELINGTOOLSEXP_API void UpdateUVLayerNames(FString& UVLayer, TArray<FString>& UVLayerNamesList, const FDynamicMesh3& Mesh);
+
+MESHMODELINGTOOLSEXP_API UPreviewMesh* CreateBakePreviewMesh(UObject* Tool, UToolTarget* ToolTarget, UWorld* World);
+
+/** @return AActor that owns a DynamicMeshComponent from a tool target, or nullptr if there is no such Actor */
+MESHMODELINGTOOLSEXP_API AActor* GetTargetActorViaIPersistentDynamicMeshSource(UToolTarget* Target);
+
+
+} // namespace Geometry
+} // namespace UE
diff --git a/Engine/Plugins/Runtime/MeshModelingToolset/Source/ModelingComponents/Private/Baking/RenderCaptureFunctions.cpp b/Engine/Plugins/Runtime/MeshModelingToolset/Source/ModelingComponents/Private/Baking/RenderCaptureFunctions.cpp
--- a/Engine/Plugins/Runtime/MeshModelingToolset/Source/ModelingComponents/Private/Baking/RenderCaptureFunctions.cpp
+++ b/Engine/Plugins/Runtime/MeshModelingToolset/Source/ModelingComponents/Private/Baking/RenderCaptureFunctions.cpp
@@ -407,6 +407,7 @@ MakeColorEvaluator(
 TUniquePtr<FMeshMapBaker> UE::Geometry::MakeRenderCaptureBaker(
 	FDynamicMesh3* BaseMesh,
 	TSharedPtr<UE::Geometry::FMeshTangentsd, ESPMode::ThreadSafe> BaseMeshTangents,
+	TSharedPtr<TArray<int32>, ESPMode::ThreadSafe> BaseMeshUVCharts,
 	FSceneCapturePhotoSet* SceneCapture,
 	FSceneCapturePhotoSetSampler* Sampler,
 	FRenderCaptureOptions PendingBake,
@@ -419,6 +420,7 @@ TUniquePtr<FMeshMapBaker> UE::Geometry::MakeRenderCaptureBaker(
 
 	check(BaseMesh != nullptr);
 	check(BaseMeshTangents.IsValid());
+	check(BaseMeshUVCharts.IsValid());
 	check(SceneCapture != nullptr);
 	check(Sampler != nullptr);
 	check(OcclusionHandler != nullptr);
@@ -436,6 +438,7 @@ TUniquePtr<FMeshMapBaker> UE::Geometry::MakeRenderCaptureBaker(
 	TUniquePtr<FMeshMapBaker> Result = MakeUnique<FMeshMapBaker>();
 	Result->SetTargetMesh(BaseMesh);
 	Result->SetTargetMeshTangents(BaseMeshTangents);
+	Result->SetTargetMeshUVCharts(BaseMeshUVCharts.Get());
 	Result->SetDimensions(FImageDimensions(static_cast<int32>(TextureImageSize), static_cast<int32>(TextureImageSize)));
 	Result->SetSamplesPerPixel(static_cast<int32>(SamplesPerPixel));
 	Result->SetFilter(FMeshMapBaker::EBakeFilterType::BSpline);
diff --git a/Engine/Plugins/Runtime/MeshModelingToolset/Source/ModelingComponents/Private/ModelingToolTargetUtil.cpp b/Engine/Plugins/Runtime/MeshModelingToolset/Source/ModelingComponents/Private/ModelingToolTargetUtil.cpp
--- a/Engine/Plugins/Runtime/MeshModelingToolset/Source/ModelingComponents/Private/ModelingToolTargetUtil.cpp
+++ b/Engine/Plugins/Runtime/MeshModelingToolset/Source/ModelingComponents/Private/ModelingToolTargetUtil.cpp
@@ -10,6 +10,8 @@
 #include "TargetInterfaces/MeshDescriptionProvider.h"
 #include "TargetInterfaces/PrimitiveComponentBackedTarget.h"
 #include "TargetInterfaces/AssetBackedTarget.h"
+#include "TargetInterfaces/SkeletalMeshBackedTarget.h"
+#include "TargetInterfaces/StaticMeshBackedTarget.h"
 #include "TargetInterfaces/DynamicMeshSource.h"
 #include "TargetInterfaces/PhysicsDataSource.h"
 
@@ -563,7 +565,19 @@ IInterface_CollisionDataProvider* UE::ToolTarget::GetPhysicsCollisionDataProvide
 	return nullptr;
 }
 
+UStaticMesh* UE::ToolTarget::GetStaticMeshFromTargetIfAvailable(UToolTarget* Target)
+{
+	IStaticMeshBackedTarget* TargetStaticMeshTarget = Cast<IStaticMeshBackedTarget>(Target);
+	UStaticMesh* TargetStaticMesh = TargetStaticMeshTarget ? TargetStaticMeshTarget->GetStaticMesh() : nullptr;
+	return TargetStaticMesh;
+}
 
+USkeletalMesh* UE::ToolTarget::GetSkeletalMeshFromTargetIfAvailable(UToolTarget* Target)
+{
+	ISkeletalMeshBackedTarget* TargetSkeletalMeshTarget = Cast<ISkeletalMeshBackedTarget>(Target);
+	USkeletalMesh* TargetSkeletalMesh = TargetSkeletalMeshTarget ? TargetSkeletalMeshTarget->GetSkeletalMesh() : nullptr;
+	return TargetSkeletalMesh;
+}
 
 
 
diff --git a/Engine/Plugins/Runtime/MeshModelingToolset/Source/ModelingComponents/Public/Baking/RenderCaptureFunctions.h b/Engine/Plugins/Runtime/MeshModelingToolset/Source/ModelingComponents/Public/Baking/RenderCaptureFunctions.h
--- a/Engine/Plugins/Runtime/MeshModelingToolset/Source/ModelingComponents/Public/Baking/RenderCaptureFunctions.h
+++ b/Engine/Plugins/Runtime/MeshModelingToolset/Source/ModelingComponents/Public/Baking/RenderCaptureFunctions.h
@@ -171,6 +171,7 @@ MODELINGCOMPONENTS_API
 TUniquePtr<FMeshMapBaker> MakeRenderCaptureBaker(
 	FDynamicMesh3* BaseMesh,
 	TSharedPtr<FMeshTangentsd, ESPMode::ThreadSafe> BaseMeshTangents,
+	TSharedPtr<TArray<int32>, ESPMode::ThreadSafe> BaseMeshUVCharts,
 	FSceneCapturePhotoSet* SceneCapture,
 	FSceneCapturePhotoSetSampler* Sampler,
 	FRenderCaptureOptions PendingBake,
diff --git a/Engine/Plugins/Runtime/MeshModelingToolset/Source/ModelingComponents/Public/ModelingToolTargetUtil.h b/Engine/Plugins/Runtime/MeshModelingToolset/Source/ModelingComponents/Public/ModelingToolTargetUtil.h
--- a/Engine/Plugins/Runtime/MeshModelingToolset/Source/ModelingComponents/Public/ModelingToolTargetUtil.h
+++ b/Engine/Plugins/Runtime/MeshModelingToolset/Source/ModelingComponents/Public/ModelingToolTargetUtil.h
@@ -8,6 +8,8 @@
 #include "TargetInterfaces/MaterialProvider.h" // FComponentMaterialSet
 #include "MeshConversionOptions.h"
 
+class UStaticMesh;
+class USkeletalMesh;
 class UToolTarget;
 class UPrimitiveComponent;
 class AActor;
@@ -35,6 +37,7 @@ namespace ToolTarget
  * @return the AActor backing a ToolTarget, or nullptr if there is no such Actor
  */
 MODELINGCOMPONENTS_API AActor* GetTargetActor(UToolTarget* Target);
+	
 
 /**
  * @return the UPrimitiveComponent backing a ToolTarget, or nullptr if there is no such Component
@@ -236,6 +239,11 @@ MODELINGCOMPONENTS_API UBodySetup* GetPhysicsBodySetup(UToolTarget* Target);
 */
 MODELINGCOMPONENTS_API IInterface_CollisionDataProvider* GetPhysicsCollisionDataProvider(UToolTarget* Target);
 
+/** @return StaticMesh from a tool target */ 
+MODELINGCOMPONENTS_API UStaticMesh* GetStaticMeshFromTargetIfAvailable(UToolTarget* Target);
+
+/** @return SkeletalMesh from a tool target */
+MODELINGCOMPONENTS_API USkeletalMesh* GetSkeletalMeshFromTargetIfAvailable(UToolTarget* Target);
 
 
 }  // end namespace ToolTarget
