MeshModleingTools - BakeRS tool.  Update so the optional Recenter Pivot works correctly.  In the past this option would move the pivot to the center of the original mesh, not the scaled and rotated mesh.

diff --git a/Engine/Plugins/Experimental/MeshModelingToolsetExp/Source/MeshModelingToolsExp/Private/BakeTransformTool.cpp b/Engine/Plugins/Experimental/MeshModelingToolsetExp/Source/MeshModelingToolsExp/Private/BakeTransformTool.cpp
--- a/Engine/Plugins/Experimental/MeshModelingToolsetExp/Source/MeshModelingToolsExp/Private/BakeTransformTool.cpp
+++ b/Engine/Plugins/Experimental/MeshModelingToolsetExp/Source/MeshModelingToolsExp/Private/BakeTransformTool.cpp
@@ -122,17 +122,29 @@ void UBakeTransformTool::UpdateAssets()
 	// mesh descriptions inside a transaction.
 	// TODO: this may not be necessary anymore. Also may not be the most efficient
 	// Note: for the crash workaround below, this also now pre-computes the source mesh bounds
-	TArray<FBox> SourceMeshBounds;
+	TArray<FBox> BoundsOfScaledRotatedMesh;
 	for (int32 ComponentIdx = 0; ComponentIdx < Targets.Num(); ComponentIdx++)
 	{
 		const FMeshDescription* MeshDescription = UE::ToolTarget::GetMeshDescription(Targets[ComponentIdx]);
 		if (MapToFirstOccurrences[ComponentIdx] < ComponentIdx)
 		{
-			SourceMeshBounds.Add(SourceMeshBounds[MapToFirstOccurrences[ComponentIdx]]);
+			BoundsOfScaledRotatedMesh.Add(BoundsOfScaledRotatedMesh[MapToFirstOccurrences[ComponentIdx]]);
 		}
 		else
 		{
-			SourceMeshBounds.Add(MeshDescription->ComputeBoundingBox());
+			// Apply the Scale and Rotation for this mesh and compute bounds.
+			FTransformSRT3d ComponentToWorld = UE::ToolTarget::GetLocalToWorldTransform( Targets[ComponentIdx] );
+			ComponentToWorld.SetTranslation(FVector::Zero());
+			
+			FBox BoundingBox(ForceInit);
+			for (const FVertexID VertexID : MeshDescription->Vertices().GetElementIDs())
+			{
+				FVector3f Pos = MeshDescription->GetVertexPosition(VertexID);
+				FVector3f NewPos = ComponentToWorld.TransformPosition(Pos);
+				BoundingBox += FVector(NewPos);
+			}
+
+			BoundsOfScaledRotatedMesh.Add(BoundingBox);
 		}
 	}
 
@@ -234,7 +246,7 @@ void UBakeTransformTool::UpdateAssets()
 			// do this part within the commit because we have the MeshDescription already computed
 			if (BasicProperties->bRecenterPivot)
 			{
-				FBox BBox = SourceMeshBounds[ComponentIdx];
+				FBox BBox = BoundsOfScaledRotatedMesh[ComponentIdx];
 				FVector3d Center(BBox.GetCenter());
 				FFrame3d LocalFrame(Center);
 				ToBakePart.SetTranslation(ToBakePart.GetTranslation() - Center);
