MeshModeling - MapPnt tool.  Make the flood fill option more robust by allowing it to work based on the triangle hit in the case that no vertices are within the brush.

diff --git a/Engine/Plugins/Experimental/MeshModelingToolsetExp/Source/MeshModelingToolsExp/Private/MeshAttributePaintTool.cpp b/Engine/Plugins/Experimental/MeshModelingToolsetExp/Source/MeshModelingToolsExp/Private/MeshAttributePaintTool.cpp
--- a/Engine/Plugins/Experimental/MeshModelingToolsetExp/Source/MeshModelingToolsExp/Private/MeshAttributePaintTool.cpp
+++ b/Engine/Plugins/Experimental/MeshModelingToolsetExp/Source/MeshModelingToolsExp/Private/MeshAttributePaintTool.cpp
@@ -509,6 +509,17 @@ void UMeshAttributePaintTool::ApplyStamp(const FBrushStampData& Stamp)
 
 	if (BrushActionProps->BrushAction == EBrushActionMode::FloodFill)
 	{
+		const bool bEmptyROIVertices = (ActionData.ROIVertices.Num() == 0);
+		if (bEmptyROIVertices) // append a vertex from the hit triangle to start the flood fill.
+		{
+			const int32 TID = Stamp.HitResult.FaceIndex;
+			const FDynamicMesh3* Mesh = PreviewMesh->GetPreviewDynamicMesh();
+			if (Mesh->IsTriangle(TID))
+			{
+				const FIndex3i VIDs = Mesh->GetTriangle(TID);
+				ActionData.ROIVertices.Add(VIDs[0]);
+			}
+		}
 		ApplyStamp_FloodFill(Stamp, ActionData);
 	}
 	else
