SubdividePolyTool: don't ensure if the subdivision level is zero, since it might have been restored to zero from the previous tool use

diff --git a/Engine/Plugins/Experimental/MeshModelingToolsetExp/Source/MeshModelingToolsEditorOnlyExp/Private/SubdividePolyTool.cpp b/Engine/Plugins/Experimental/MeshModelingToolsetExp/Source/MeshModelingToolsEditorOnlyExp/Private/SubdividePolyTool.cpp
--- a/Engine/Plugins/Experimental/MeshModelingToolsetExp/Source/MeshModelingToolsEditorOnlyExp/Private/SubdividePolyTool.cpp
+++ b/Engine/Plugins/Experimental/MeshModelingToolsetExp/Source/MeshModelingToolsEditorOnlyExp/Private/SubdividePolyTool.cpp
@@ -220,7 +220,7 @@ void USubdividePolyTool::Setup()
 		return;
 	}
 
-	if (!ensure(Properties->SubdivisionLevel >= 1)) // Should be enforced by UPROPERTY meta tags
+	if (Properties->SubdivisionLevel < 1)
 	{
 		Properties->SubdivisionLevel = 1;
 	}
