Render Capture Baking: Fix tooltip text for the BakeRC tool Opacity channel

diff --git a/Engine/Plugins/Experimental/MeshModelingToolsetExp/Source/MeshModelingToolsExp/Public/BakeRenderCaptureTool.h b/Engine/Plugins/Experimental/MeshModelingToolsetExp/Source/MeshModelingToolsExp/Public/BakeRenderCaptureTool.h
--- a/Engine/Plugins/Experimental/MeshModelingToolsetExp/Source/MeshModelingToolsExp/Public/BakeRenderCaptureTool.h
+++ b/Engine/Plugins/Experimental/MeshModelingToolsetExp/Source/MeshModelingToolsExp/Public/BakeRenderCaptureTool.h
@@ -131,7 +131,7 @@ public:
 	UPROPERTY(Category = RenderCaptureOptions, EditAnywhere, meta = (NoResetToDefault))
 	bool bEmissiveMap = false;
 
-	// Whether to generate a texture for the Emissive property
+	// Whether to generate a texture for the Opacity property
 	UPROPERTY(Category = RenderCaptureOptions, EditAnywhere, meta = (NoResetToDefault))
 	bool bOpacityMap = false;
 
