Fix MshToVol tool not actually simplifying models with a "large" number of faces

This fixes the editor hanging in the inefficient BSP processing code

diff --git a/Engine/Plugins/Experimental/MeshModelingToolsetExp/Source/MeshModelingToolsEditorOnlyExp/Private/MeshToVolumeTool.cpp b/Engine/Plugins/Experimental/MeshModelingToolsetExp/Source/MeshModelingToolsEditorOnlyExp/Private/MeshToVolumeTool.cpp
--- a/Engine/Plugins/Experimental/MeshModelingToolsetExp/Source/MeshModelingToolsEditorOnlyExp/Private/MeshToVolumeTool.cpp
+++ b/Engine/Plugins/Experimental/MeshModelingToolsetExp/Source/MeshModelingToolsEditorOnlyExp/Private/MeshToVolumeTool.cpp
@@ -39,6 +39,93 @@ USingleSelectionMeshEditingTool* UMeshToVolumeToolBuilder::CreateNewTool(const F
 	return NewObject<UMeshToVolumeTool>(SceneState.ToolManager);
 }
 
+
+class FCalculateVolumeOp : public TGenericDataOperator<FDynamicMeshFaceArray>
+{
+public:
+	virtual ~FCalculateVolumeOp() {}
+
+	// inputs
+	TSharedPtr<FDynamicMesh3, ESPMode::ThreadSafe> SourceMesh;
+
+	// parameters
+	EMeshToVolumeMode ConversionMode;
+	UE::Conversion::FMeshToVolumeOptions MeshToVolumeOptions;
+
+	// error flags
+	bool bTooManyTriangles = false;
+
+	//
+	// TGenericDataOperator implementation
+	// 
+
+	virtual void CalculateResult(FProgressCancel* Progress) override
+	{
+		if (Progress && Progress->Cancelled())
+		{
+			return;
+		}
+
+		bool bShowTooLargeWarning = false;
+
+		auto SimplifyToMax = [this](FDynamicMesh3& Mesh, int32 TriangleCount)
+		{
+			FVolPresMeshSimplification Simplifier(&Mesh);
+			Simplifier.SimplifyToTriangleCount(TriangleCount);
+		};
+
+		if (ConversionMode == EMeshToVolumeMode::MinimalPolygons)
+		{
+			// Since this tool is likely to be a sink, there isn't much reason to keep
+			// the group differentiations if they are coplanar.
+			bool bRespectGroupBoundaries = false;
+
+			// Apply minimal-planar simplification to remove extra vertices along straight edges
+			FDynamicMesh3 LocalMesh = *SourceMesh;
+			LocalMesh.DiscardAttributes();
+			FQEMSimplification PlanarSimplifier(&LocalMesh);
+			PlanarSimplifier.SimplifyToMinimalPlanar(0.1);		// angle tolerance in degrees
+
+			UE::Conversion::GetPolygonFaces(LocalMesh, *Result, bRespectGroupBoundaries);
+
+			bTooManyTriangles = (LocalMesh.TriangleCount() > MeshToVolumeOptions.MaxTriangles);
+
+			if (Progress && Progress->Cancelled())
+			{
+				return;
+			}
+
+			if (bTooManyTriangles)
+			{
+				SimplifyToMax(LocalMesh, MeshToVolumeOptions.MaxTriangles);
+				UE::Conversion::GetPolygonFaces(LocalMesh, *Result, bRespectGroupBoundaries);
+			}
+		}
+		else
+		{
+			bTooManyTriangles = (SourceMesh->TriangleCount() > MeshToVolumeOptions.MaxTriangles);
+
+			if (Progress && Progress->Cancelled())
+			{
+				return;
+			}
+
+			if (bTooManyTriangles)
+			{
+				FDynamicMesh3 LocalMesh = *SourceMesh;
+				LocalMesh.DiscardAttributes();
+				SimplifyToMax(LocalMesh, MeshToVolumeOptions.MaxTriangles);
+				UE::Conversion::GetTriangleFaces(LocalMesh, *Result);
+			}
+			else
+			{
+				UE::Conversion::GetTriangleFaces(*SourceMesh, *Result);
+			}
+		}
+	}
+};
+
+
 /*
  * Tool
  */
@@ -64,7 +151,8 @@ void UMeshToVolumeTool::Setup()
 	FComponentMaterialSet MaterialSet = UE::ToolTarget::GetMaterialSet(Target);
 	PreviewMesh->SetMaterials(MaterialSet.Materials);
 
-	InputMesh.Copy(*PreviewMesh->GetMesh());
+	InputMesh = MakeShared<FDynamicMesh3, ESPMode::ThreadSafe>();
+	InputMesh->Copy(*PreviewMesh->GetMesh());
 
 	VolumeEdgesSet = NewObject<ULineSetComponent>(PreviewMesh->GetRootComponent());
 	VolumeEdgesSet->SetupAttachment(PreviewMesh->GetRootComponent());
@@ -79,14 +167,34 @@ void UMeshToVolumeTool::Setup()
 
 	Settings->WatchProperty(Settings->ConversionMode,
 							[this](EMeshToVolumeMode NewMode)
-							{ bVolumeValid = false; });
+							{ Compute->InvalidateResult(); });
 
 
 	HandleSourcesProperties = NewObject<UOnAcceptHandleSourcesPropertiesSingle>(this);
 	HandleSourcesProperties->RestoreProperties(this);
 	AddToolPropertySource(HandleSourcesProperties);
 
-	bVolumeValid = false;
+	Compute = MakeUnique<TGenericDataBackgroundCompute<FDynamicMeshFaceArray>>();
+	Compute->Setup(this);
+	Compute->OnResultUpdated.AddLambda([this](const TUniquePtr<FDynamicMeshFaceArray>& NewResult)
+	{
+		UpdateLineSet(*NewResult);
+	});
+	Compute->OnOpCompleted.AddLambda([this](const TGenericDataOperator<FDynamicMeshFaceArray>* UncastOp)
+	{
+		const FCalculateVolumeOp* Op = static_cast<const FCalculateVolumeOp*>(UncastOp);
+		if (Op->bTooManyTriangles)
+		{
+			GetToolManager()->DisplayMessage(
+				LOCTEXT("LargeFaceCount", "Mesh has large face count; output Volume representation has been automatically simplified"),
+				EToolMessageLevel::UserWarning);
+		}
+		else
+		{
+			GetToolManager()->DisplayMessage({}, EToolMessageLevel::UserWarning);
+		}
+	});
+	Compute->InvalidateResult();
 	
 
 	GetToolManager()->DisplayMessage(
@@ -95,7 +203,7 @@ void UMeshToVolumeTool::Setup()
 
 	// check for errors in input mesh
 	bool bFoundBoundaryEdges = false;
-	for (int32 BoundaryEdgeID : InputMesh.BoundaryEdgeIndicesItr())
+	for (int32 BoundaryEdgeID : InputMesh->BoundaryEdgeIndicesItr())
 	{
 		bFoundBoundaryEdges = true;
 		break;
@@ -156,7 +264,9 @@ void UMeshToVolumeTool::OnShutdown(EToolShutdownType ShutdownType)
 			TargetVolume->GetBrushComponent()->Modify();
 		}
 
-		UE::Conversion::DynamicMeshToVolume(InputMesh, Faces, TargetVolume);
+		// Note: InputMesh parameter not actually used by DynamicMeshToVolume
+		TUniquePtr<FDynamicMeshFaceArray> Faces = Compute->Shutdown();
+		UE::Conversion::DynamicMeshToVolume(*InputMesh, *Faces, TargetVolume);
 		TargetVolume->SetActorTransform(SetTransform);
 		TargetVolume->PostEditChange();
 
@@ -172,25 +282,27 @@ void UMeshToVolumeTool::OnShutdown(EToolShutdownType ShutdownType)
 
 void UMeshToVolumeTool::OnTick(float DeltaTime)
 {
-	if (bVolumeValid == false)
-	{
-		RecalculateVolume();
-	}
+	Compute->Tick(DeltaTime);
 }
 
 void UMeshToVolumeTool::Render(IToolsContextRenderAPI* RenderAPI)
 {
 }
 
+bool UMeshToVolumeTool::CanAccept() const
+{
+	return Super::CanAccept() && Compute->HaveValidResult();
+}
 
-void UMeshToVolumeTool::UpdateLineSet()
+
+void UMeshToVolumeTool::UpdateLineSet(FDynamicMeshFaceArray& FaceArr)
 {
 	FColor BoundaryEdgeColor(240, 15, 15);
 	float BoundaryEdgeThickness = 0.5;
 	float BoundaryEdgeDepthBias = 2.0f;
 
 	VolumeEdgesSet->Clear();
-	for (const UE::Conversion::FDynamicMeshFace& Face : Faces)
+	for (const UE::Conversion::FDynamicMeshFace& Face : FaceArr)
 	{
 		int32 NumV = Face.BoundaryLoop.Num();
 		for (int32 k = 0; k < NumV; ++k)
@@ -203,47 +315,14 @@ void UMeshToVolumeTool::UpdateLineSet()
 
 }
 
-void UMeshToVolumeTool::RecalculateVolume()
+TUniquePtr<UE::Geometry::TGenericDataOperator<FDynamicMeshFaceArray>> UMeshToVolumeTool::MakeNewOperator()
 {
-	UE::Conversion::FMeshToVolumeOptions DefaultOptions;
-	bool bShowTooLargeWarning = false;
-
-	if (Settings->ConversionMode == EMeshToVolumeMode::MinimalPolygons)
-	{
-		// Since this tool is likely to be a sink, there isn't much reason to keep
-		// the group differentiations if they are coplanar.
-		bool bRespectGroupBoundaries = false;
-
-		// Apply minimal-planar simplification to remove extra vertices along straight edges
-		FDynamicMesh3 LocalMesh = InputMesh;
-		LocalMesh.DiscardAttributes();
-		FQEMSimplification PlanarSimplifier(&LocalMesh);
-		PlanarSimplifier.SimplifyToMinimalPlanar(0.1);		// angle tolerance in degrees
-
-		UE::Conversion::GetPolygonFaces(LocalMesh, Faces, bRespectGroupBoundaries);
-
-		bShowTooLargeWarning = (LocalMesh.TriangleCount() > DefaultOptions.MaxTriangles);
-	}
-	else
-	{
-		UE::Conversion::GetTriangleFaces(InputMesh, Faces);
+	TUniquePtr<FCalculateVolumeOp> VolumeOp = MakeUnique<FCalculateVolumeOp>();
 
-		bShowTooLargeWarning = (InputMesh.TriangleCount() > DefaultOptions.MaxTriangles);
-	}
-
-	if (bShowTooLargeWarning)
-	{
-		GetToolManager()->DisplayMessage(
-			LOCTEXT("LargeFaceCount", "Mesh has large face count, output Volume representation may be automatically simplified"),
-			EToolMessageLevel::UserWarning);
-	}
-	else
-	{
-		GetToolManager()->DisplayMessage({}, EToolMessageLevel::UserWarning);
-	}
+	VolumeOp->SourceMesh = InputMesh;
+	VolumeOp->ConversionMode = Settings->ConversionMode;
 
-	UpdateLineSet();
-	bVolumeValid = true;
+	return VolumeOp;
 }
 
 
diff --git a/Engine/Plugins/Experimental/MeshModelingToolsetExp/Source/MeshModelingToolsEditorOnlyExp/Public/MeshToVolumeTool.h b/Engine/Plugins/Experimental/MeshModelingToolsetExp/Source/MeshModelingToolsEditorOnlyExp/Public/MeshToVolumeTool.h
--- a/Engine/Plugins/Experimental/MeshModelingToolsetExp/Source/MeshModelingToolsEditorOnlyExp/Public/MeshToVolumeTool.h
+++ b/Engine/Plugins/Experimental/MeshModelingToolsetExp/Source/MeshModelingToolsEditorOnlyExp/Public/MeshToVolumeTool.h
@@ -14,6 +14,8 @@
 #include "PreviewMesh.h"
 #include "PropertySets/OnAcceptProperties.h"
 #include "SingleSelectionTool.h"
+#include "ModelingOperators.h"
+#include "MeshOpPreviewHelpers.h"
 #include "MeshToVolumeTool.generated.h"
 
 /**
@@ -62,6 +64,8 @@ public:
 };
 
 
+using FDynamicMeshFaceArray = TArray<UE::Conversion::FDynamicMeshFace>;
+
 /**
  * Converts a mesh to a volume.
  *
@@ -69,7 +73,7 @@ public:
  * tool can be moved out of the editor-only section and put with VolumeToMeshTool.
  */
 UCLASS()
-class MESHMODELINGTOOLSEDITORONLYEXP_API UMeshToVolumeTool : public USingleSelectionMeshEditingTool
+class MESHMODELINGTOOLSEDITORONLYEXP_API UMeshToVolumeTool : public USingleSelectionMeshEditingTool, public UE::Geometry::IGenericDataOperatorFactory<FDynamicMeshFaceArray>
 {
 	GENERATED_BODY()
 
@@ -84,6 +88,10 @@ public:
 
 	virtual bool HasCancel() const override { return true; }
 	virtual bool HasAccept() const override { return true; }
+	virtual bool CanAccept() const override;
+
+	// IGenericDataOperatorFactory API
+	virtual TUniquePtr<UE::Geometry::TGenericDataOperator<FDynamicMeshFaceArray>> MakeNewOperator() override;
 
 protected:
 	UPROPERTY()
@@ -99,12 +107,10 @@ protected:
 	TObjectPtr<ULineSetComponent> VolumeEdgesSet;
 
 protected:
-	UE::Geometry::FDynamicMesh3 InputMesh;
-
-	void RecalculateVolume();
-	void UpdateLineSet();
+	TSharedPtr<UE::Geometry::FDynamicMesh3, ESPMode::ThreadSafe> InputMesh;
+	
+	TUniquePtr<TGenericDataBackgroundCompute<FDynamicMeshFaceArray>> Compute = nullptr;
 
-	bool bVolumeValid = false;
-	TArray<UE::Conversion::FDynamicMeshFace> Faces;
+	void UpdateLineSet(FDynamicMeshFaceArray& Faces);
 
 };
