MeshModelingTools: Fix gizmo numerical UI being incorrectly placed in the presence of full-monitor scaling (such as commonly used on 4k displays), and make the grab location more accurately tracked.

diff --git a/Engine/Plugins/Experimental/MeshModelingToolsetExp/Source/ModelingEditorUI/Private/ModelingWidgets/SDraggableBox.cpp b/Engine/Plugins/Experimental/MeshModelingToolsetExp/Source/ModelingEditorUI/Private/ModelingWidgets/SDraggableBox.cpp
--- a/Engine/Plugins/Experimental/MeshModelingToolsetExp/Source/ModelingEditorUI/Private/ModelingWidgets/SDraggableBox.cpp
+++ b/Engine/Plugins/Experimental/MeshModelingToolsetExp/Source/ModelingEditorUI/Private/ModelingWidgets/SDraggableBox.cpp
@@ -27,13 +27,11 @@ void SDraggableBoxOverlay::Construct(const FArguments& InArgs)
 		});
 		OnDragCompleteFunction = SDraggableBox::FOnDragComplete::CreateLambda([this](const FVector2D& ScreenSpacePosition)
 		{
-			FVector2D TopLeftPosition = ContainingBox->GetTickSpaceGeometry().AbsoluteToLocal(ScreenSpacePosition);
+			FVector2D BottomLeftScreenSpace(ScreenSpacePosition.X, ScreenSpacePosition.Y + DraggableBox->GetTickSpaceGeometry().GetAbsoluteSize().Y);
+			FVector2D BottomLeftLocal = ContainingBox->GetTickSpaceGeometry().AbsoluteToLocal(BottomLeftScreenSpace);
 
-			DraggableBoxPaddingHorizontal = TopLeftPosition.X;
-
-			// When aligning to bottom, we need the distance between the draggable box bottom and our bottom.
-			// I.e. OurHeight - TopLeftY - BoxHeight
-			DraggableBoxPaddingVertical = ContainingBox->GetTickSpaceGeometry().GetAbsoluteSize().Y - TopLeftPosition.Y - DraggableBox->GetTickSpaceGeometry().GetAbsoluteSize().Y;
+			DraggableBoxPaddingHorizontal = BottomLeftLocal.X;
+			DraggableBoxPaddingVertical = ContainingBox->GetTickSpaceGeometry().GetLocalSize().Y - BottomLeftLocal.Y;
 		});
 	}
 	else
@@ -102,14 +100,15 @@ void SDraggableBox::Construct(const FArguments& InArgs)
 
 FReply SDraggableBox::OnMouseButtonDown(const FGeometry& MyGeometry, const FPointerEvent& MouseEvent)
 {
+	// Need to remember where within the box we grabbed. We do this here instead of OnDragDetected because 
+	// our mouse can potentially travel some distance before OnDragDetected fires.
+	ScreenSpaceOffsetOfGrab = MouseEvent.GetScreenSpacePosition() - MyGeometry.GetAbsolutePosition();
+
 	return FReply::Handled().DetectDrag(SharedThis(this), EKeys::LeftMouseButton);
 };
 
 FReply SDraggableBox::OnDragDetected(const FGeometry& MyGeometry, const FPointerEvent& MouseEvent)
 {
-	// Need to remember where within the box we grabbed
-	const FVector2D ScreenSpaceOffsetOfGrab = MouseEvent.GetScreenSpacePosition() - MyGeometry.GetAbsolutePosition();
-
 	// The drag/drop operation moves a transparent copy of our contents, so we should hide ourselves to make it
 	// look like the contents are moving.
 	SetVisibility(EVisibility::Collapsed);
diff --git a/Engine/Plugins/Experimental/MeshModelingToolsetExp/Source/ModelingEditorUI/Public/ModelingWidgets/SDraggableBox.h b/Engine/Plugins/Experimental/MeshModelingToolsetExp/Source/ModelingEditorUI/Public/ModelingWidgets/SDraggableBox.h
--- a/Engine/Plugins/Experimental/MeshModelingToolsetExp/Source/ModelingEditorUI/Public/ModelingWidgets/SDraggableBox.h
+++ b/Engine/Plugins/Experimental/MeshModelingToolsetExp/Source/ModelingEditorUI/Public/ModelingWidgets/SDraggableBox.h
@@ -72,6 +72,10 @@ public:
 protected:
 	TSharedPtr<SWidget> InnerWidget;
 	FOnDragComplete OnDragComplete;
+
+	// Remembers the point in the box that we grabbed (on click, rather than when the
+	// drag was confirmed)
+	FVector2D ScreenSpaceOffsetOfGrab;
 };
 
 
@@ -97,7 +101,7 @@ public:
 	 * @param InUIToBeDragged	  The UI being dragged
 	 * @param InDecoratorOffset   Where within the UI we grabbed, so we're not dragging by the upper left of the UI.
 	 * @param OwnerAreaSize       Size of the DockArea at the time when we start dragging.
-	 * @param OnDragComplete      Delegate to call when dropped.
+	 * @param OnDragComplete      Delegate to call when dropped. Gets passed in the screen space location of the top left corner.
 	 *
 	 * @return a new FDockingDragOperation
 	 */
@@ -116,7 +120,10 @@ protected:
 	 */
 	TSharedPtr<class SWidget> UIBeingDragged;
 
+	// Screen space offset of grab location from the top left corner of the rectangle being dragged.
 	FVector2D DecoratorOffsetFromCursor;
+
+	// Size of rectangle being dragged
 	FVector2D LastContentSize;
 
 	SDraggableBox::FOnDragComplete OnDragComplete;
