Add warning to MatEd tool when number of Materials is not large enough to contain the largest Material ID
On tool accept, clamp Material IDs to prevent the tool from creating meshes with invalid Material IDs in such cases

diff --git a/Engine/Plugins/Experimental/MeshModelingToolsetExp/Source/MeshModelingToolsEditorOnlyExp/Private/EditMeshMaterialsTool.cpp b/Engine/Plugins/Experimental/MeshModelingToolsetExp/Source/MeshModelingToolsEditorOnlyExp/Private/EditMeshMaterialsTool.cpp
--- a/Engine/Plugins/Experimental/MeshModelingToolsetExp/Source/MeshModelingToolsEditorOnlyExp/Private/EditMeshMaterialsTool.cpp
+++ b/Engine/Plugins/Experimental/MeshModelingToolsetExp/Source/MeshModelingToolsEditorOnlyExp/Private/EditMeshMaterialsTool.cpp
@@ -96,7 +96,7 @@ void UEditMeshMaterialsTool::Setup()
 	CurrentMaterials = MaterialProps->Materials;
 	InitialMaterialKey = GetMaterialKey();
 
-	MaterialProps->WatchProperty<FMaterialSetKey>(
+	MaterialSetWatchIndex = MaterialProps->WatchProperty<FMaterialSetKey>(
 		[this](){ return GetMaterialKey(); },
 		[this](FMaterialSetKey NewKey) { OnMaterialSetChanged(); });
 
@@ -108,9 +108,60 @@ void UEditMeshMaterialsTool::Setup()
 			EToolMessageLevel::UserWarning);
 	}
 
+	// Double check that materials are valid after mesh changes
+	PreviewMesh->GetOnMeshChanged().AddLambda([this]()
+	{
+		UpdateMaterialSetErrors();
+	});
+
 }
 
 
+int32 UEditMeshMaterialsTool::FindMaxActiveMaterialID() const
+{
+	int32 MaxActiveMaterialID = -1;
+	PreviewMesh->ProcessMesh([&MaxActiveMaterialID](const FDynamicMesh3& Mesh)
+	{
+		if (!Mesh.HasAttributes())
+		{
+			return;
+		}
+		const FDynamicMeshMaterialAttribute* MaterialAttrib = Mesh.Attributes()->GetMaterialID();
+		for (int32 TID : Mesh.TriangleIndicesItr())
+		{
+			MaxActiveMaterialID = FMath::Max(MaxActiveMaterialID, MaterialAttrib->GetValue(TID));
+		}
+	});
+	return MaxActiveMaterialID;
+}
+
+
+bool UEditMeshMaterialsTool::FixInvalidMaterialIDs()
+{
+	int32 NumMaterials = MaterialProps->Materials.Num();
+	if (!ensure(NumMaterials > 0)) // We must have more than 0 materials or there are no valid material IDs
+	{
+		return false;
+	}
+	bool bHasChanged = false;
+	PreviewMesh->EditMesh([NumMaterials, &bHasChanged](FDynamicMesh3& Mesh)
+	{
+		if (!Mesh.HasAttributes())
+		{
+			return;
+		}
+		FDynamicMeshMaterialAttribute* MaterialAttrib = Mesh.Attributes()->GetMaterialID();
+		for (int32 TID : Mesh.TriangleIndicesItr())
+		{
+			if (MaterialAttrib->GetValue(TID) >= NumMaterials)
+			{
+				MaterialAttrib->SetValue(TID, NumMaterials - 1);
+				bHasChanged = true;
+			}
+		}
+	});
+	return bHasChanged;
+}
 
 UMeshSelectionToolActionPropertySet* UEditMeshMaterialsTool::CreateEditActions()
 {
@@ -290,6 +341,8 @@ void UEditMeshMaterialsTool::AssignMaterialToSelectedTriangles()
 	bFullMeshInvalidationPending = true;
 	OnExternalSelectionChange();
 	bHaveModifiedMesh = true;
+
+	UpdateMaterialSetErrors();
 }
 
 
@@ -311,17 +364,35 @@ void UEditMeshMaterialsTool::OnMaterialSetChanged()
 
 	bHaveModifiedMaterials = true;
 
-	if (MaterialProps->Materials.Num() == 0)
+	UpdateMaterialSetErrors();
+}
+
+
+void UEditMeshMaterialsTool::UpdateMaterialSetErrors()
+{
+	int32 NumMaterials = MaterialProps->Materials.Num();
+	if (NumMaterials == 0)
 	{
 		GetToolManager()->DisplayMessage(LOCTEXT("NoMaterialsMessage", "Material Set must contain at least one Material"), EToolMessageLevel::UserWarning);
 		bShowingMaterialSetError = true;
+		bShowingNotEnoughMaterialsError = false;
 	}
-	else if (bShowingMaterialSetError)
+	else
 	{
-		GetToolManager()->DisplayMessage({}, EToolMessageLevel::UserWarning);
-		bShowingMaterialSetError = false;
+		int32 MaxActiveMaterialID = FindMaxActiveMaterialID();
+		if (MaxActiveMaterialID >= NumMaterials)
+		{
+			GetToolManager()->DisplayMessage(FText::Format(LOCTEXT("NotEnoughMaterialsMessage", "Material Set only has {0} {0}|plural(one=material,other=materials), but mesh expects at least {1}. Will remap invalid material IDs on 'Accept'"), NumMaterials, MaxActiveMaterialID + 1), EToolMessageLevel::UserWarning);
+			bShowingNotEnoughMaterialsError = true;
+			bShowingMaterialSetError = false;
+		}
+		else if (bShowingNotEnoughMaterialsError || bShowingMaterialSetError)
+		{
+			GetToolManager()->DisplayMessage({}, EToolMessageLevel::UserWarning);
+			bShowingNotEnoughMaterialsError = false;
+			bShowingMaterialSetError = false;
+		}
 	}
-
 }
 
 
@@ -334,6 +405,12 @@ void UEditMeshMaterialsTool::ExternalUpdateMaterialSet(const TArray<UMaterialInt
 	SetToolPropertySourceEnabled(MaterialProps, true);
 	PreviewMesh->SetMaterials(MaterialProps->Materials);
 	CurrentMaterials = MaterialProps->Materials;
+
+	UpdateMaterialSetErrors();
+	if (ensure(MaterialSetWatchIndex > -1))
+	{
+		MaterialProps->SilentUpdateWatcherAtIndex(MaterialSetWatchIndex);
+	}
 }
 
 
@@ -358,6 +435,12 @@ void UEditMeshMaterialsTool::ApplyShutdownAction(EToolShutdownType ShutdownType)
 			UE::ToolTarget::CommitMaterialSetUpdate(Target, NewMaterialSet, true);
 		}
 
+		if (bShowingNotEnoughMaterialsError)
+		{
+			bool bFixedMaterialIDs = FixInvalidMaterialIDs();
+			bHaveModifiedMesh = bHaveModifiedMesh || bFixedMaterialIDs;
+		}
+
 		if (bHaveModifiedMesh)
 		{
 			UE::ToolTarget::CommitDynamicMeshUpdate(Target, *PreviewMesh->GetMesh(), true);
diff --git a/Engine/Plugins/Experimental/MeshModelingToolsetExp/Source/MeshModelingToolsEditorOnlyExp/Public/EditMeshMaterialsTool.h b/Engine/Plugins/Experimental/MeshModelingToolsetExp/Source/MeshModelingToolsEditorOnlyExp/Public/EditMeshMaterialsTool.h
--- a/Engine/Plugins/Experimental/MeshModelingToolsetExp/Source/MeshModelingToolsEditorOnlyExp/Public/EditMeshMaterialsTool.h
+++ b/Engine/Plugins/Experimental/MeshModelingToolsetExp/Source/MeshModelingToolsEditorOnlyExp/Public/EditMeshMaterialsTool.h
@@ -121,11 +121,25 @@ protected:
 	FMaterialSetKey InitialMaterialKey;
 	bool bHaveModifiedMaterials = false;
 	bool bShowingMaterialSetError = false;
+	bool bShowingNotEnoughMaterialsError = false;
 
 	virtual void ApplyShutdownAction(EToolShutdownType ShutdownType) override;
 
 	void ExternalUpdateMaterialSet(const TArray<UMaterialInterface*>& NewMaterialSet);
 	friend class FEditMeshMaterials_MaterialSetChange;
+
+private:
+
+	// @return the max material ID used in the PreviewMesh
+	int32 FindMaxActiveMaterialID() const;
+
+	// Clamp material IDs to be less than the number of materials
+	bool FixInvalidMaterialIDs();
+
+	// Update error messages regarding whether we have enough materials on the current mesh
+	void UpdateMaterialSetErrors();
+
+	int32 MaterialSetWatchIndex = -1;
 };
 
 
