Improve the documentation of options in the SubDiv tool

diff --git a/Engine/Plugins/Experimental/MeshModelingToolsetExp/Source/MeshModelingToolsEditorOnlyExp/Public/SubdividePolyTool.h b/Engine/Plugins/Experimental/MeshModelingToolsetExp/Source/MeshModelingToolsEditorOnlyExp/Public/SubdividePolyTool.h
--- a/Engine/Plugins/Experimental/MeshModelingToolsetExp/Source/MeshModelingToolsEditorOnlyExp/Public/SubdividePolyTool.h
+++ b/Engine/Plugins/Experimental/MeshModelingToolsetExp/Source/MeshModelingToolsEditorOnlyExp/Public/SubdividePolyTool.h
@@ -38,6 +38,7 @@ class MESHMODELINGTOOLSEDITORONLYEXP_API USubdividePolyToolProperties : public U
 
 public:
 
+	/** Number of iterations/levels of subdivision to perform */
 	UPROPERTY(EditAnywhere, Category=Settings, meta = (UIMin = "1", ClampMin = "1", Delta = 1, LinearDeltaSensitivity = 50))
 	int SubdivisionLevel = 3;
 
@@ -51,16 +52,18 @@ public:
 	UPROPERTY(EditAnywhere, Category=Settings)
 	ESubdivisionOutputNormals NormalComputationMethod = ESubdivisionOutputNormals::Generated;
 
-	UPROPERTY(EditAnywhere, Category=Settings)
+	UPROPERTY(EditAnywhere, Category=Settings, meta = (DisplayName = "UV Computation Method"))
 	ESubdivisionOutputUVs UVComputationMethod = ESubdivisionOutputUVs::Interpolated;
 
 	/** Assign a new PolyGroup ID to each newly created face */
 	UPROPERTY(EditAnywhere, Category = Settings, meta = (DisplayName = "New PolyGroups"))
 	bool bNewPolyGroups = false;
 
+	/** Display each PolyGroup with an auto-generated color */
 	UPROPERTY(EditAnywhere, Category = Settings)
 	bool bRenderGroups = true;
 
+	/** Display the mesh corresponding to Subdivision Level 0 */
 	UPROPERTY(EditAnywhere, Category = Settings)
 	bool bRenderCage = true;
 
diff --git a/Engine/Plugins/Runtime/MeshModelingToolset/Source/ModelingComponentsEditorOnly/Public/Operations/SubdividePoly.h b/Engine/Plugins/Runtime/MeshModelingToolset/Source/ModelingComponentsEditorOnly/Public/Operations/SubdividePoly.h
--- a/Engine/Plugins/Runtime/MeshModelingToolset/Source/ModelingComponentsEditorOnly/Public/Operations/SubdividePoly.h
+++ b/Engine/Plugins/Runtime/MeshModelingToolset/Source/ModelingComponentsEditorOnly/Public/Operations/SubdividePoly.h
@@ -17,7 +17,7 @@ UENUM()
 enum class ESubdivisionScheme : uint8
 {
 	Bilinear,
-	CatmullClark,
+	CatmullClark UMETA(DisplayName = "Catmull-Clark"),
 	Loop
 };
 
