ModelingMode: Improve Project and Remesh tool's TargetEdgeLength tooltip, explaining why it's not starting at the default value hardcoded into the property.

diff --git a/Engine/Plugins/Experimental/MeshModelingToolsetExp/Source/MeshModelingToolsExp/Public/RemeshMeshTool.h b/Engine/Plugins/Experimental/MeshModelingToolsetExp/Source/MeshModelingToolsExp/Public/RemeshMeshTool.h
--- a/Engine/Plugins/Experimental/MeshModelingToolsetExp/Source/MeshModelingToolsExp/Public/RemeshMeshTool.h
+++ b/Engine/Plugins/Experimental/MeshModelingToolsetExp/Source/MeshModelingToolsExp/Public/RemeshMeshTool.h
@@ -80,8 +80,8 @@ public:
 	UPROPERTY(EditAnywhere, Category = Remeshing, AdvancedDisplay)
 	bool bUseTargetEdgeLength;
 
-	/** Target edge length */
-	UPROPERTY(EditAnywhere, Category = Remeshing, AdvancedDisplay, meta = (NoSpinbox = "true", EditCondition = "bUseTargetEdgeLength == true"))
+	/** Remesh to have edges approximately this length. An attempt at a reasonable value is computed automatically for this field based on the selected target mesh. */
+	UPROPERTY(EditAnywhere, Category = Remeshing, AdvancedDisplay, meta = (NoSpinbox = "true", EditCondition = "bUseTargetEdgeLength == true", NoResetToDefault))
 	float TargetEdgeLength;
 
 	/** Enable projection back to input mesh */
