ModelingTools: Adjust the Project tool's settings to properly disable the smooth filling options when smooth filling is disabled.

diff --git a/Engine/Plugins/Experimental/MeshModelingToolsetExp/Source/MeshModelingToolsExp/Public/ProjectToTargetTool.h b/Engine/Plugins/Experimental/MeshModelingToolsetExp/Source/MeshModelingToolsExp/Public/ProjectToTargetTool.h
--- a/Engine/Plugins/Experimental/MeshModelingToolsetExp/Source/MeshModelingToolsExp/Public/ProjectToTargetTool.h
+++ b/Engine/Plugins/Experimental/MeshModelingToolsetExp/Source/MeshModelingToolsExp/Public/ProjectToTargetTool.h
@@ -48,10 +48,10 @@ public:
 	UPROPERTY(EditAnywhere, Category = NormalFlow, meta = (EditCondition = "RemeshType == ERemeshType::NormalFlow"))
 	bool bSmoothInFillAreas = true;
 
-	UPROPERTY(EditAnywhere, Category = NormalFlow, meta = (EditCondition = "RemeshType == ERemeshType::NormalFlow", UIMin = "0", UIMax = "1.0", ClampMin = "0", ClampMax = "10.0"))
+	UPROPERTY(EditAnywhere, Category = NormalFlow, meta = (EditCondition = "RemeshType == ERemeshType::NormalFlow && bSmoothInFillAreas", UIMin = "0", UIMax = "1.0", ClampMin = "0", ClampMax = "10.0"))
 	float FillAreaDistanceMultiplier = 0.25f;
 
-	UPROPERTY(EditAnywhere, Category = NormalFlow, meta = (EditCondition = "RemeshType == ERemeshType::NormalFlow", UIMin = "0", UIMax = "1.0", ClampMin = "0", ClampMax = "10.0"))
+	UPROPERTY(EditAnywhere, Category = NormalFlow, meta = (EditCondition = "RemeshType == ERemeshType::NormalFlow && bSmoothInFillAreas", UIMin = "0", UIMax = "1.0", ClampMin = "0", ClampMax = "10.0"))
 	float FillAreaSmoothMultiplier = 0.25f;
 
 
