Expose simplification along cut/mirror plane as an option in the mirror and plane cut tools

diff --git a/Engine/Plugins/Experimental/MeshModelingToolsetExp/Source/MeshModelingToolsExp/Private/MirrorTool.cpp b/Engine/Plugins/Experimental/MeshModelingToolsetExp/Source/MeshModelingToolsExp/Private/MirrorTool.cpp
--- a/Engine/Plugins/Experimental/MeshModelingToolsetExp/Source/MeshModelingToolsExp/Private/MirrorTool.cpp
+++ b/Engine/Plugins/Experimental/MeshModelingToolsetExp/Source/MeshModelingToolsExp/Private/MirrorTool.cpp
@@ -59,6 +59,7 @@ TUniquePtr<FDynamicMeshOperator> UMirrorOperatorFactory::MakeNewOperator()
 	MirrorOp->OriginalMesh = MirrorTool->MeshesToMirror[ComponentIndex]->GetMesh();
 	MirrorOp->bAppendToOriginal = MirrorTool->Settings->OperationMode == EMirrorOperationMode::MirrorAndAppend;
 	MirrorOp->bCropFirst = MirrorTool->Settings->bCropAlongMirrorPlaneFirst;
+	MirrorOp->bSimplifyAlongNewEdges = MirrorTool->Settings->bSimplifyAlongCrop;
 	MirrorOp->bWeldAlongPlane = MirrorTool->Settings->bWeldVerticesOnMirrorPlane;
 	MirrorOp->PlaneTolerance = MirrorTool->Settings->PlaneTolerance;
 	MirrorOp->bAllowBowtieVertexCreation = MirrorTool->Settings->bAllowBowtieVertexCreation;
diff --git a/Engine/Plugins/Experimental/MeshModelingToolsetExp/Source/MeshModelingToolsExp/Private/PlaneCutTool.cpp b/Engine/Plugins/Experimental/MeshModelingToolsetExp/Source/MeshModelingToolsExp/Private/PlaneCutTool.cpp
--- a/Engine/Plugins/Experimental/MeshModelingToolsetExp/Source/MeshModelingToolsExp/Private/PlaneCutTool.cpp
+++ b/Engine/Plugins/Experimental/MeshModelingToolsetExp/Source/MeshModelingToolsExp/Private/PlaneCutTool.cpp
@@ -272,6 +272,7 @@ TUniquePtr<FDynamicMeshOperator> UPlaneCutOperatorFactory::MakeNewOperator()
 	TUniquePtr<FPlaneCutOp> CutOp = MakeUnique<FPlaneCutOp>();
 	CutOp->bFillCutHole = CutTool->BasicProperties->bFillCutHole;
 	CutOp->bFillSpans = CutTool->BasicProperties->bFillSpans;
+	CutOp->bSimplifyAlongNewEdges = CutTool->BasicProperties->bSimplifyAlongCut;
 
 	FTransform LocalToWorld = (FTransform) UE::ToolTarget::GetLocalToWorldTransform(CutTool->Targets[ComponentIndex]);
 	CutOp->SetTransform(LocalToWorld);
diff --git a/Engine/Plugins/Experimental/MeshModelingToolsetExp/Source/MeshModelingToolsExp/Public/MirrorTool.h b/Engine/Plugins/Experimental/MeshModelingToolsetExp/Source/MeshModelingToolsExp/Public/MirrorTool.h
--- a/Engine/Plugins/Experimental/MeshModelingToolsetExp/Source/MeshModelingToolsExp/Public/MirrorTool.h
+++ b/Engine/Plugins/Experimental/MeshModelingToolsetExp/Source/MeshModelingToolsExp/Public/MirrorTool.h
@@ -71,6 +71,10 @@ public:
 	UPROPERTY(EditAnywhere, Category = Options)
 	bool bCropAlongMirrorPlaneFirst = true;
 
+	/** Whether to locally simplify new edges created when cropping along the mirror plane. Will only simplify when doing so will not change the shape, UVs or PolyGroups. */
+	UPROPERTY(EditAnywhere, Category = Options, AdvancedDisplay, meta = (EditCondition = "bCropAlongMirrorPlaneFirst"))
+	bool bSimplifyAlongCrop = true;
+
 	/** Weld vertices that lie on the mirror plane. Vertices will not be welded if doing so would give an edge more than two faces, or if they are part of a face in the plane. */
 	UPROPERTY(EditAnywhere, Category = Options, meta = (
 		EditCondition = "OperationMode == EMirrorOperationMode::MirrorAndAppend", EditConditionHides))
diff --git a/Engine/Plugins/Experimental/MeshModelingToolsetExp/Source/MeshModelingToolsExp/Public/PlaneCutTool.h b/Engine/Plugins/Experimental/MeshModelingToolsetExp/Source/MeshModelingToolsExp/Public/PlaneCutTool.h
--- a/Engine/Plugins/Experimental/MeshModelingToolsetExp/Source/MeshModelingToolsExp/Public/PlaneCutTool.h
+++ b/Engine/Plugins/Experimental/MeshModelingToolsetExp/Source/MeshModelingToolsExp/Public/PlaneCutTool.h
@@ -80,6 +80,10 @@ public:
 	/** If true, will attempt to fill cut holes even if they're ill-formed (e.g. because they connect to pre-existing holes in the geometry) */
 	UPROPERTY(EditAnywhere, Category = Options, AdvancedDisplay, meta = (EditCondition = "bFillCutHole"))
 	bool bFillSpans = false;
+
+	/** If true, will simplify triangulation along plane cut when doing so will not affect the shape, UVs or PolyGroups */
+	UPROPERTY(EditAnywhere, Category = Options, AdvancedDisplay)
+	bool bSimplifyAlongCut = true;
 };
 
 
diff --git a/Engine/Plugins/Runtime/MeshModelingToolset/Source/ModelingOperators/Private/CompositionOps/MirrorOp.cpp b/Engine/Plugins/Runtime/MeshModelingToolset/Source/ModelingOperators/Private/CompositionOps/MirrorOp.cpp
--- a/Engine/Plugins/Runtime/MeshModelingToolset/Source/ModelingOperators/Private/CompositionOps/MirrorOp.cpp
+++ b/Engine/Plugins/Runtime/MeshModelingToolset/Source/ModelingOperators/Private/CompositionOps/MirrorOp.cpp
@@ -37,6 +37,7 @@ void FMirrorOp::CalculateResult(FProgressCancel* Progress)
 
 		FMeshPlaneCut Cutter(ResultMesh.Get(), LocalPlaneOrigin, LocalPlaneNormal);
 		Cutter.PlaneTolerance = PlaneTolerance;
+		Cutter.bSimplifyAlongNewEdges = bSimplifyAlongNewEdges;
 		Cutter.Cut();
 
 		if (Progress && Progress->Cancelled())
diff --git a/Engine/Plugins/Runtime/MeshModelingToolset/Source/ModelingOperators/Private/CuttingOps/PlaneCutOp.cpp b/Engine/Plugins/Runtime/MeshModelingToolset/Source/ModelingOperators/Private/CuttingOps/PlaneCutOp.cpp
--- a/Engine/Plugins/Runtime/MeshModelingToolset/Source/ModelingOperators/Private/CuttingOps/PlaneCutOp.cpp
+++ b/Engine/Plugins/Runtime/MeshModelingToolset/Source/ModelingOperators/Private/CuttingOps/PlaneCutOp.cpp
@@ -34,6 +34,7 @@ void FPlaneCutOp::CalculateResult(FProgressCancel* Progress)
 		return;
 	}
 	Cut.UVScaleFactor = UVScaleFactor;
+	Cut.bSimplifyAlongNewEdges = bSimplifyAlongNewEdges;
 
 
 	if (Progress && Progress->Cancelled())
diff --git a/Engine/Plugins/Runtime/MeshModelingToolset/Source/ModelingOperators/Public/CompositionOps/MirrorOp.h b/Engine/Plugins/Runtime/MeshModelingToolset/Source/ModelingOperators/Public/CompositionOps/MirrorOp.h
--- a/Engine/Plugins/Runtime/MeshModelingToolset/Source/ModelingOperators/Public/CompositionOps/MirrorOp.h
+++ b/Engine/Plugins/Runtime/MeshModelingToolset/Source/ModelingOperators/Public/CompositionOps/MirrorOp.h
@@ -27,6 +27,9 @@ public:
 	/** Whether to crop the result first along the plane. */
 	bool bCropFirst = true;
 
+	/** Whether to locally simplify the new edges created when cropping along the plane. Only relevant if bCropFirst is true. */
+	bool bSimplifyAlongNewEdges = true;
+
 	/** If true, the mirrored portion is appended to the original. If false, the result is just the mirrored portion. */
 	bool bAppendToOriginal = true;
 
diff --git a/Engine/Plugins/Runtime/MeshModelingToolset/Source/ModelingOperators/Public/CuttingOps/PlaneCutOp.h b/Engine/Plugins/Runtime/MeshModelingToolset/Source/ModelingOperators/Public/CuttingOps/PlaneCutOp.h
--- a/Engine/Plugins/Runtime/MeshModelingToolset/Source/ModelingOperators/Public/CuttingOps/PlaneCutOp.h
+++ b/Engine/Plugins/Runtime/MeshModelingToolset/Source/ModelingOperators/Public/CuttingOps/PlaneCutOp.h
@@ -23,6 +23,7 @@ public:
 	bool bFillCutHole = true;
 	bool bFillSpans = false;
 	bool bKeepBothHalves = false;
+	bool bSimplifyAlongNewEdges = true;
 	double CutPlaneLocalThickness = 0; // plane thickness in the local space of the mesh
 	double UVScaleFactor = 0;
 	static const FName ObjectIndexAttribute;
