ModelingTools: Update the Volume to Mesh conversion tool's option tooltips to better reflect the behavior behind them.

diff --git a/Engine/Plugins/Experimental/MeshModelingToolsetExp/Source/MeshModelingToolsExp/Public/VolumeToMeshTool.h b/Engine/Plugins/Experimental/MeshModelingToolsetExp/Source/MeshModelingToolsExp/Public/VolumeToMeshTool.h
--- a/Engine/Plugins/Experimental/MeshModelingToolsetExp/Source/MeshModelingToolsExp/Public/VolumeToMeshTool.h
+++ b/Engine/Plugins/Experimental/MeshModelingToolsetExp/Source/MeshModelingToolsExp/Public/VolumeToMeshTool.h
@@ -34,15 +34,19 @@ class MESHMODELINGTOOLSEXP_API UVolumeToMeshToolProperties : public UInteractive
 	GENERATED_BODY()
 
 public:
+	/** Weld coincident vertices and edges together in the resulting mesh to form a closed mesh surface. */
 	UPROPERTY(EditAnywhere, Category = Options)
 	bool bWeldEdges = true;
 
-	UPROPERTY(EditAnywhere, Category = Options)
+	/** If WeldEdges is enabled, attempt to fill any small holes or cracks in the resulting mesh to form a closed surface. */
+	UPROPERTY(EditAnywhere, Category = Options, meta = (EditCondition = "bWeldEdges") )
 	bool bAutoRepair = true;
 
-	UPROPERTY(EditAnywhere, Category = Options)
+	/** If WeldEdges is enabled, and after mesh generation is complete, flip edges in planar regions to improve triangle quality. */
+	UPROPERTY(EditAnywhere, Category = Options, meta = (EditCondition = "bWeldEdges"))
 	bool bOptimizeMesh = true;
 
+	/** Show the wireframe of the resulting converted mesh geometry.*/
 	UPROPERTY(EditAnywhere, Category = Options)
 	bool bShowWireframe = true;
 };
