ModelingTools: Disable the rendering of the BaseSculpt tool's indicator circle in the SeamSculptTool, since it's not needed for the tool's functioning and causes visual confusion.

diff --git a/Engine/Plugins/Experimental/MeshModelingToolsetExp/Source/MeshModelingToolsExp/Private/SeamSculptTool.cpp b/Engine/Plugins/Experimental/MeshModelingToolsetExp/Source/MeshModelingToolsExp/Private/SeamSculptTool.cpp
--- a/Engine/Plugins/Experimental/MeshModelingToolsetExp/Source/MeshModelingToolsExp/Private/SeamSculptTool.cpp
+++ b/Engine/Plugins/Experimental/MeshModelingToolsetExp/Source/MeshModelingToolsExp/Private/SeamSculptTool.cpp
@@ -27,6 +27,9 @@
 
 #include "ModelingToolTargetUtil.h"
 
+#include "BaseGizmos/BrushStampIndicator.h"
+
+
 #include UE_INLINE_GENERATED_CPP_BY_NAME(SeamSculptTool)
 
 using namespace UE::Geometry;
@@ -102,6 +105,9 @@ void USeamSculptTool::Setup()
 	GetToolManager()->DisplayMessage(
 		LOCTEXT("OnStartSeamSculptTool", "Draw UV seams on the mesh by left click-dragging between mesh vertices."),
 		EToolMessageLevel::UserNotification);
+
+	// We don't need this to be visible for the seam editing process.
+	BrushStampIndicator->bVisible = false;
 }
 
 
