Fix an issue in the PivotActor tool where changing the transform using the gizmo did not cause the reset to default arrow to appear in the details panel

diff --git a/Engine/Plugins/Experimental/MeshModelingToolsetExp/Source/MeshModelingToolsEditorOnlyExp/Private/AddPivotActorTool.cpp b/Engine/Plugins/Experimental/MeshModelingToolsetExp/Source/MeshModelingToolsEditorOnlyExp/Private/AddPivotActorTool.cpp
--- a/Engine/Plugins/Experimental/MeshModelingToolsetExp/Source/MeshModelingToolsEditorOnlyExp/Private/AddPivotActorTool.cpp
+++ b/Engine/Plugins/Experimental/MeshModelingToolsetExp/Source/MeshModelingToolsEditorOnlyExp/Private/AddPivotActorTool.cpp
@@ -185,6 +185,7 @@ void UAddPivotActorTool::GizmoTransformChanged(UTransformProxy* Proxy, FTransfor
 {
 	TransformProperties->Position = Transform.GetTranslation();
 	TransformProperties->Rotation = Transform.GetRotation();
+	NotifyOfPropertyChangeByTool(TransformProperties);
 	GizmoPositionWatcher.SilentUpdate();
 	GizmoRotationWatcher.SilentUpdate();
 }
