fix the 'showflags' collision rendering not updating after accepting new collision shapes in the Mesh to Collision tool

diff --git a/Engine/Plugins/Experimental/MeshModelingToolsetExp/Source/MeshModelingToolsExp/Private/Physics/SetCollisionGeometryTool.cpp b/Engine/Plugins/Experimental/MeshModelingToolsetExp/Source/MeshModelingToolsExp/Private/Physics/SetCollisionGeometryTool.cpp
--- a/Engine/Plugins/Experimental/MeshModelingToolsetExp/Source/MeshModelingToolsExp/Private/Physics/SetCollisionGeometryTool.cpp
+++ b/Engine/Plugins/Experimental/MeshModelingToolsetExp/Source/MeshModelingToolsExp/Private/Physics/SetCollisionGeometryTool.cpp
@@ -382,6 +382,8 @@ void USetCollisionGeometryTool::OnShutdown(EToolShutdownType ShutdownType)
 					{
 						SMComponent->RecreatePhysicsState();
 					}
+					// Mark the render state dirty to make sure any CollisionTraceFlag changes get picked up
+					SMComponent->MarkRenderStateDirty();
 				}
 			}
 
