MeshModelingTools: Make mirror tool use the standard input/output properties, and standard ctrl+shift+click behavior. Make the "export both halves" setting in plane cut tool be controlled by the "keep both halves" setting, and be closer to it.

diff --git a/Engine/Plugins/Experimental/MeshModelingToolsetExp/Source/MeshModelingToolsExp/Private/MirrorTool.cpp b/Engine/Plugins/Experimental/MeshModelingToolsetExp/Source/MeshModelingToolsExp/Private/MirrorTool.cpp
--- a/Engine/Plugins/Experimental/MeshModelingToolsetExp/Source/MeshModelingToolsExp/Private/MirrorTool.cpp
+++ b/Engine/Plugins/Experimental/MeshModelingToolsetExp/Source/MeshModelingToolsExp/Private/MirrorTool.cpp
@@ -12,6 +12,8 @@
 #include "MeshDescriptionToDynamicMesh.h"
 #include "ModelingToolTargetUtil.h"
 #include "Misc/MessageDialog.h"
+#include "PropertySets/CreateMeshObjectTypeProperties.h"
+#include "PropertySets/OnAcceptProperties.h" //UOnAcceptHandleSourcesProperties
 #include "ToolBuilderUtil.h"
 #include "ToolSetupUtil.h"
 #include "DynamicMesh/MeshTransforms.h"
@@ -39,6 +41,7 @@ namespace MirrorTool_Local
 		T.SetScale3D(TransformIn.GetScale3D());
 		return T;
 	}
+	const FString MirrorToolIdentifier = TEXT("MirrorTool");
 }
 
 // Tool builder functions
@@ -157,11 +160,13 @@ void UMirrorTool::Render(IToolsContextRenderAPI* RenderAPI)
 
 void UMirrorTool::Setup()
 {
+	using namespace MirrorTool_Local;
+
 	UInteractiveTool::Setup();
 
 	SetToolDisplayName(LOCTEXT("ToolName", "Mirror"));
 	GetToolManager()->DisplayMessage(
-		LOCTEXT("OnStartMirrorTool", "Mirror one or more meshes across a plane. The plane can be set by using the preset buttons, moving the gizmo, or ctrl+clicking on a spot on the original mesh."),
+		LOCTEXT("OnStartMirrorTool", "Mirror one or more meshes across a plane. Ctrl+[shift+]clicking on a spot on the original mesh will place the plane there. Buttons will reset it."),
 		EToolMessageLevel::UserNotification);
 
 	// Set up the properties
@@ -169,10 +174,27 @@ void UMirrorTool::Setup()
 	Settings->RestoreProperties(this);
 	AddToolPropertySource(Settings);
 
+	OutputTypeProperties = NewObject<UCreateMeshObjectTypeProperties>(this);
+	OutputTypeProperties->InitializeDefaultWithAuto();
+	OutputTypeProperties->OutputType = UCreateMeshObjectTypeProperties::AutoIdentifier;
+	OutputTypeProperties->RestoreProperties(this, MirrorToolIdentifier);
+	OutputTypeProperties->WatchProperty(OutputTypeProperties->OutputType, [this](FString) { OutputTypeProperties->UpdatePropertyVisibility(); });
+	AddToolPropertySource(OutputTypeProperties);
+
+	HandleSourcesProperties = NewObject<UOnAcceptHandleSourcesProperties>(this);
+	HandleSourcesProperties->RestoreProperties(this, MirrorToolIdentifier);
+	AddToolPropertySource(HandleSourcesProperties);
+
 	ToolActions = NewObject<UMirrorToolActionPropertySet>(this);
 	ToolActions->Initialize(this);
 	AddToolPropertySource(ToolActions);
 
+	Settings->WatchProperty(Settings->WriteTo, [this](EMirrorSaveMode SaveMode) {
+		bool bShowOutputOptions = SaveMode == EMirrorSaveMode::NewObjects;
+		SetToolPropertySourceEnabled(OutputTypeProperties, bShowOutputOptions);
+		SetToolPropertySourceEnabled(HandleSourcesProperties, bShowOutputOptions);
+	});
+
 	CheckAndDisplayWarnings();
 
 	// Fill in the MeshesToMirror array with suitably converted meshes.
@@ -236,13 +258,7 @@ void UMirrorTool::Setup()
 		}
 		});
 
-	// Modify the Ctrl+click set plane behavior to respond to our CtrlClickBehavior property
-	PlaneMechanic->SetPlaneCtrlClickBehaviorTarget->OnClickedPositionFunc = [this](const FHitResult& Hit)
-	{
-		bool bIgnoreNormal = (Settings->CtrlClickBehavior == EMirrorCtrlClickBehavior::Reposition);
-		PlaneMechanic->SetDrawPlaneFromWorldPos((FVector3d)Hit.ImpactPoint, (FVector3d)Hit.ImpactNormal, bIgnoreNormal);
-	};
-	// Also include the original components in the ctrl+click hit testing even though we made them 
+	// Include the original components in the ctrl+click hit testing even though we made them 
 	// invisible, since we want to be able to reposition the plane onto the original mesh.
 	for (int32 ComponentIdx = 0; ComponentIdx < Targets.Num(); ComponentIdx++)
 	{
@@ -350,7 +366,11 @@ void UMirrorTool::CheckAndDisplayWarnings()
 
 void UMirrorTool::OnShutdown(EToolShutdownType ShutdownType)
 {
+	using namespace MirrorTool_Local;
+
 	Settings->SaveProperties(this);
+	OutputTypeProperties->SaveProperties(this, MirrorToolIdentifier);
+	HandleSourcesProperties->SaveProperties(this, MirrorToolIdentifier);
 
 	PlaneMechanic->Shutdown();
 
@@ -363,6 +383,8 @@ void UMirrorTool::OnShutdown(EToolShutdownType ShutdownType)
 	// Swap in results, if appropriate
 	if (ShutdownType == EToolShutdownType::Accept)
 	{
+		GetToolManager()->BeginUndoTransaction(LOCTEXT("MirrorToolTransactionName", "Mirror Tool"));
+
 		// Gather results
 		TArray<FDynamicMeshOpResult> Results;
 		for (int32 PreviewIndex = 0; PreviewIndex < Previews.Num(); ++PreviewIndex)
@@ -375,6 +397,19 @@ void UMirrorTool::OnShutdown(EToolShutdownType ShutdownType)
 
 		// Convert to output. This will also edit the selection.
 		GenerateAsset(Results);
+
+		// If we created new objects, deal with the input objects according to the user's settings.
+		if (Settings->WriteTo != EMirrorSaveMode::InputObjects)
+		{
+			TArray<AActor*> Actors;
+			for (int32 ComponentIdx = 0; ComponentIdx < Targets.Num(); ComponentIdx++)
+			{
+				Actors.Add(UE::ToolTarget::GetTargetActor(Targets[ComponentIdx]));
+			}
+			HandleSourcesProperties->ApplyMethod(Actors, GetToolManager());
+		}
+
+		GetToolManager()->EndUndoTransaction();
 	}
 	else
 	{
@@ -392,10 +427,6 @@ void UMirrorTool::GenerateAsset(const TArray<FDynamicMeshOpResult>& Results)
 		return;
 	}
 
-	GetToolManager()->BeginUndoTransaction(LOCTEXT("MirrorToolTransactionName", "Mirror Tool"));
-
-	ensure(Results.Num() > 0);
-
 	int32 NumSourceMeshes = MeshesToMirror.Num();
 
 	// check if we entirely cut away any meshes
@@ -436,7 +467,7 @@ void UMirrorTool::GenerateAsset(const TArray<FDynamicMeshOpResult>& Results)
 			}
 			continue;
 		}
-		else if (Settings->SaveMode == EMirrorSaveMode::UpdateAssets)
+		else if (Settings->WriteTo == EMirrorSaveMode::InputObjects)
 		{
 			NewSelection.Actors.Add(UE::ToolTarget::GetTargetActor(Targets[OrigMeshIdx]));
 
@@ -458,14 +489,19 @@ void UMirrorTool::GenerateAsset(const TArray<FDynamicMeshOpResult>& Results)
 			NewMeshObjectParams.BaseName = TEXT("Mirror");
 			NewMeshObjectParams.Materials = Materials;
 			NewMeshObjectParams.SetMesh(Mesh);
+			if (OutputTypeProperties->OutputType == UCreateMeshObjectTypeProperties::AutoIdentifier)
+			{
+				UE::ToolTarget::ConfigureCreateMeshObjectParams(Targets[OrigMeshIdx], NewMeshObjectParams);
+			}
+			else
+			{
+				OutputTypeProperties->ConfigureCreateMeshObjectParams(NewMeshObjectParams);
+			}
 			FCreateMeshObjectResult Result = UE::Modeling::CreateMeshObject(GetToolManager(), MoveTemp(NewMeshObjectParams));
 			if (Result.IsOK() && Result.NewActor != nullptr)
 			{
 				NewSelection.Actors.Add(Result.NewActor);
 			}
-
-			// Remove the original actor
-			UE::ToolTarget::GetTargetComponent(Targets[OrigMeshIdx])->DestroyComponent();
 		}
 	}
 
@@ -474,8 +510,6 @@ void UMirrorTool::GenerateAsset(const TArray<FDynamicMeshOpResult>& Results)
 	{
 		GetToolManager()->RequestSelectionChange(NewSelection);
 	}
-
-	GetToolManager()->EndUndoTransaction();
 }
 
 
@@ -531,7 +565,7 @@ void UMirrorTool::ApplyAction(EMirrorToolAction ActionType)
 		}
 
 		// The user can optionally have the button change the direction only
-		if (Settings->bButtonsOnlyChangeOrientation)
+		if (ToolActions->bButtonsOnlyChangeOrientation)
 		{
 			ShiftedPlaneOrigin = MirrorPlaneOrigin;	// Keeps the same
 		}
diff --git a/Engine/Plugins/Experimental/MeshModelingToolsetExp/Source/MeshModelingToolsExp/Private/PlaneCutTool.cpp b/Engine/Plugins/Experimental/MeshModelingToolsetExp/Source/MeshModelingToolsExp/Private/PlaneCutTool.cpp
--- a/Engine/Plugins/Experimental/MeshModelingToolsetExp/Source/MeshModelingToolsExp/Private/PlaneCutTool.cpp
+++ b/Engine/Plugins/Experimental/MeshModelingToolsetExp/Source/MeshModelingToolsExp/Private/PlaneCutTool.cpp
@@ -107,10 +107,6 @@ void UPlaneCutTool::Setup()
 	BasicProperties->RestoreProperties(this);
 	AddToolPropertySource(BasicProperties);
 
-	AcceptProperties = NewObject<UAcceptOutputProperties>(this, TEXT("Tool Accept Output Settings"));
-	AcceptProperties->RestoreProperties(this);
-	AddToolPropertySource(AcceptProperties);
-
 	ToolPropertyObjects.Add(this);
 
 	// initialize the PreviewMesh+BackgroundCompute object
@@ -248,7 +244,6 @@ void UPlaneCutTool::OnShutdown(EToolShutdownType ShutdownType)
 {
 	PlaneMechanic->Shutdown();
 	BasicProperties->SaveProperties(this);
-	AcceptProperties->SaveProperties(this);
 
 	// Restore (unhide) the source meshes
 	for (int Idx = 0; Idx < Targets.Num(); Idx++)
@@ -453,7 +448,7 @@ void UPlaneCutTool::GenerateAsset(const TArray<FDynamicMeshOpResult>& Results)
 			continue;
 		}
 
-		if (AcceptProperties->bExportSeparatedPiecesAsNewMeshAssets)
+		if (BasicProperties->bExportSeparatedPiecesAsNewMeshAssets)
 		{
 			TDynamicMeshScalarTriangleAttribute<int>* SubMeshIDs =
 				static_cast<TDynamicMeshScalarTriangleAttribute<int>*>(UseMesh->Attributes()->GetAttachedAttribute(
diff --git a/Engine/Plugins/Experimental/MeshModelingToolsetExp/Source/MeshModelingToolsExp/Public/MirrorTool.h b/Engine/Plugins/Experimental/MeshModelingToolsetExp/Source/MeshModelingToolsExp/Public/MirrorTool.h
--- a/Engine/Plugins/Experimental/MeshModelingToolsetExp/Source/MeshModelingToolsExp/Public/MirrorTool.h
+++ b/Engine/Plugins/Experimental/MeshModelingToolsetExp/Source/MeshModelingToolsExp/Public/MirrorTool.h
@@ -15,6 +15,9 @@
 
 #include "MirrorTool.generated.h"
 
+class UCreateMeshObjectTypeProperties;
+class UOnAcceptHandleSourcesProperties;
+
 UCLASS()
 class MESHMODELINGTOOLSEXP_API UMirrorToolBuilder : public UMultiSelectionMeshEditingToolBuilder
 {
@@ -26,11 +29,11 @@ public:
 UENUM()
 enum class EMirrorSaveMode : uint8
 {
-	/**  Save the results in place of the original assets. */
-	UpdateAssets = 0,
+	/**  Save the results in place of the original input objects. */
+	InputObjects = 0,
 
-	/** Save the results as new assets. */
-	CreateNewAssets = 1,
+	/** Save the results as new objects. */
+	NewObjects = 1,
 };
 
 UENUM()
@@ -43,19 +46,6 @@ enum class EMirrorOperationMode : uint8
 	MirrorExisting = 1,
 };
 
-UENUM()
-enum class EMirrorCtrlClickBehavior : uint8
-{
-	/** Move the mirror plane to clicked location without adjusting its normal. */
-	Reposition = 0,
-
-	/** Move the mirror plane and adjust its normal according to click location. */
-	RepositionAndReorient = 1,
-};
-
-/**
- * 
- */
 UCLASS()
 class MESHMODELINGTOOLSEXP_API UMirrorToolProperties : public UInteractiveToolPropertySet
 {
@@ -90,22 +80,14 @@ public:
 	UPROPERTY(EditAnywhere, Category = Options, AdvancedDisplay, meta = (
 		EditCondition = "bWeldVerticesOnMirrorPlane && OperationMode == EMirrorOperationMode::MirrorAndAppend", EditConditionHides))
 	bool bAllowBowtieVertexCreation = false;
-
-	/** What Ctrl + clicking does to the mirror plane. */
-	UPROPERTY(EditAnywhere, Category = RepositionOptions)
-	EMirrorCtrlClickBehavior CtrlClickBehavior = EMirrorCtrlClickBehavior::Reposition;
-
-	/** If true the "Preset Mirror Directions" buttons only change the plane orientation, not location. */
-	UPROPERTY(EditAnywhere, Category = RepositionOptions, AdvancedDisplay)
-	bool bButtonsOnlyChangeOrientation = false;
 	
 	/** Whether to show the preview. */
 	UPROPERTY(EditAnywhere, Category = Options)
 	bool bShowPreview = true;
 
 	/** How to save the result. */
-	UPROPERTY(EditAnywhere, Category = ToolOutputOptions)
-	EMirrorSaveMode SaveMode = EMirrorSaveMode::UpdateAssets;
+	UPROPERTY(EditAnywhere, Category = OutputOptions)
+	EMirrorSaveMode WriteTo = EMirrorSaveMode::InputObjects;
 };
 
 
@@ -155,6 +137,10 @@ public:
 	UFUNCTION(CallInEditor, Category = RepositionPlane)
 	void ShiftToCenter() { PostAction(EMirrorToolAction::ShiftToCenter); }
 
+	/** If true the "Preset Mirror Directions" buttons only change the plane orientation, not location. */
+	UPROPERTY(EditAnywhere, Category = PresetMirrorDirections)
+	bool bButtonsOnlyChangeOrientation = false;
+
 	/** Move the mirror plane and adjust its normal to mirror entire selection leftward. */
 	UFUNCTION(CallInEditor, Category = PresetMirrorDirections, meta = (DisplayPriority = 1))
 	void Left() { PostAction(EMirrorToolAction::Left); }
@@ -212,6 +198,12 @@ protected:
 	UPROPERTY()
 	TObjectPtr<UMirrorToolProperties> Settings = nullptr;
 
+	UPROPERTY()
+	TObjectPtr<UCreateMeshObjectTypeProperties> OutputTypeProperties = nullptr;
+
+	UPROPERTY()
+	TObjectPtr<UOnAcceptHandleSourcesProperties> HandleSourcesProperties = nullptr;
+
 	UPROPERTY()
 	TObjectPtr<UMirrorToolActionPropertySet> ToolActions = nullptr;
 
diff --git a/Engine/Plugins/Experimental/MeshModelingToolsetExp/Source/MeshModelingToolsExp/Public/PlaneCutTool.h b/Engine/Plugins/Experimental/MeshModelingToolsetExp/Source/MeshModelingToolsExp/Public/PlaneCutTool.h
--- a/Engine/Plugins/Experimental/MeshModelingToolsetExp/Source/MeshModelingToolsExp/Public/PlaneCutTool.h
+++ b/Engine/Plugins/Experimental/MeshModelingToolsetExp/Source/MeshModelingToolsExp/Public/PlaneCutTool.h
@@ -29,27 +29,6 @@ public:
 };
 
 
-
-/**
-* Properties controlling how changes are baked out to static meshes on tool accept
-*/
-UCLASS()
-class MESHMODELINGTOOLSEXP_API UAcceptOutputProperties : public UInteractiveToolPropertySet
-{
-	GENERATED_BODY()
-
-public:
-
-	/** If true, meshes cut into multiple pieces will be saved as separate assets on 'accept'. */
-	UPROPERTY(EditAnywhere, Category = ToolOutputOptions)
-	bool bExportSeparatedPiecesAsNewMeshAssets = true;
-};
-
-
-
-
-
-
 /**
  * Standard properties of the plane cut operation
  */
@@ -67,9 +46,13 @@ public:
 	bool bKeepBothHalves = false;
 
 	/** If keeping both halves, separate the two pieces by this amount */
-	UPROPERTY(EditAnywhere, Category = Options, meta = (EditCondition = "bKeepBothHalves == true", UIMin = "0", ClampMin = "0", Delta = 0.5, LinearDeltaSensitivity = 1))
+	UPROPERTY(EditAnywhere, Category = Options, meta = (EditCondition = "bKeepBothHalves", UIMin = "0", ClampMin = "0", Delta = 0.5, LinearDeltaSensitivity = 1))
 	float SpacingBetweenHalves = 0;
 
+	/** If true, meshes cut into multiple pieces will be saved as separate assets on 'accept'. */
+	UPROPERTY(EditAnywhere, Category = Options, meta = (EditCondition = "bKeepBothHalves"))
+	bool bExportSeparatedPiecesAsNewMeshAssets = true;
+
 	UPROPERTY(EditAnywhere, Category = Options)
 	bool bShowPreview = true;
 
@@ -151,9 +134,6 @@ protected:
 	UPROPERTY()
 	TObjectPtr<UPlaneCutToolProperties> BasicProperties;
 
-	UPROPERTY()
-	TObjectPtr<UAcceptOutputProperties> AcceptProperties;
-
 	UPROPERTY()
 	TArray<TObjectPtr<UMeshOpPreviewWithBackgroundCompute>> Previews;
 
