Various skinning tool improvements:
- interactive normalization
- Add/Remove/Multiply/Relax brush modes
- Surface/Volume falloff modes
- realtime deformation preview

diff --git a/Engine/Plugins/Experimental/MeshModelingToolsetExp/Source/MeshModelingToolsExp/Private/SkinWeightsPaintTool.cpp b/Engine/Plugins/Experimental/MeshModelingToolsetExp/Source/MeshModelingToolsExp/Private/SkinWeightsPaintTool.cpp
--- a/Engine/Plugins/Experimental/MeshModelingToolsetExp/Source/MeshModelingToolsExp/Private/SkinWeightsPaintTool.cpp
+++ b/Engine/Plugins/Experimental/MeshModelingToolsetExp/Source/MeshModelingToolsExp/Private/SkinWeightsPaintTool.cpp
@@ -15,6 +15,7 @@
 
 #include "MeshDescription.h"
 #include "DynamicMesh/NonManifoldMappingSupport.h"
+#include "Parameterization/MeshLocalParam.h"
 
 #include UE_INLINE_GENERATED_CPP_BY_NAME(SkinWeightsPaintTool)
 
@@ -22,94 +23,389 @@
 #define LOCTEXT_NAMESPACE "USkinWeightsPaintTool"
 
 // Any weight below this value is ignored, since it won't be representable in a uint8.
-const float MinimumWeightThreshold = 1.0f / 255.0f;
+constexpr float MinimumWeightThreshold = 1.0f / 255.0f;
 
+USkinWeightsPaintToolProperties::USkinWeightsPaintToolProperties(const FObjectInitializer& ObjectInitializer)
+		: Super(ObjectInitializer)
+{
+	ColorRamp.Add(FLinearColor::Blue);
+	ColorRamp.Add(FLinearColor::Green);
+	ColorRamp.Add(FLinearColor(FColor::Orange));
+	ColorRamp.Add(FLinearColor::Red);
 
-void FMeshSkinWeightsChange::Apply(UObject* Object)
+	MinColor = FLinearColor::Black;
+	MaxColor = FLinearColor::White;
+}
+
+void FSkinToolDeformer::Initialize(
+	const USkeletalMeshComponent* SkeletalMeshComponent,
+	const FMeshDescription* Mesh)
 {
-	USkinWeightsPaintTool* Tool = CastChecked<USkinWeightsPaintTool>(Object);
+	// get all bone transforms in the reference pose store a copy in component space
+	Component = SkeletalMeshComponent;
+	const FReferenceSkeleton& RefSkeleton = Component->GetSkeletalMeshAsset()->GetRefSkeleton();
+	const TArray<FTransform> &LocalSpaceBoneTransforms = RefSkeleton.GetRefBonePose();
+	const int32 NumBones = LocalSpaceBoneTransforms.Num();
+	InvCSRefPoseTransforms.SetNumUninitialized(NumBones);
+	for (int32 BoneIndex=0; BoneIndex<NumBones; ++BoneIndex)
+	{
+		const int32 ParentBoneIndex = RefSkeleton.GetParentIndex(BoneIndex);
+		const FTransform& LocalTransform = LocalSpaceBoneTransforms[BoneIndex];
+		if (ParentBoneIndex != INDEX_NONE)
+		{
+			InvCSRefPoseTransforms[BoneIndex] = LocalTransform * InvCSRefPoseTransforms[ParentBoneIndex];
+		}
+		else
+		{
+			InvCSRefPoseTransforms[BoneIndex] = LocalTransform;
+		}
+	}
+	
+	for (int32 BoneIndex=0; BoneIndex<NumBones; ++BoneIndex)
+	{
+		// pre-invert the transforms so we don't have to at runtime
+		InvCSRefPoseTransforms[BoneIndex] = InvCSRefPoseTransforms[BoneIndex].Inverse();
+
+		// store map of bone names to bone indices
+		FName BoneName = RefSkeleton.GetBoneName(BoneIndex);
+		BoneNameToIndexMap.Add(BoneName, BoneIndex);
+	}
+
+	// store reference pose vertex positions
+	const TArrayView<const UE::Math::TVector<float>> VertexPositions = Mesh->GetVertexPositions().GetRawArray();
+	RefPoseVertexPositions = VertexPositions;
 
-	Tool->ExternalUpdateValues(BoneName, NewWeights);
+	// force all vertices to be updated initially
+	VerticesWithModifiedWeights.Reserve(RefPoseVertexPositions.Num());
+	for (int32 VertexID=0; VertexID<RefPoseVertexPositions.Num(); ++VertexID)
+	{
+		VerticesWithModifiedWeights.Add(VertexID);
+	}
 }
 
-void FMeshSkinWeightsChange::Revert(UObject* Object)
+void FSkinToolDeformer::UpdateVertexDeformation(
+	PerBoneWeightMap& Weights,
+	UPreviewMesh* PreviewMesh,
+	UE::Geometry::TDynamicVerticesOctree3<FDynamicMesh3>& OctreeToUpdate)
 {
-	USkinWeightsPaintTool* Tool = CastChecked<USkinWeightsPaintTool>(Object);
+	if (VerticesWithModifiedWeights.IsEmpty())
+	{
+		return;
+	}
+	
+	// update vertex positions
+	PreviewMesh->EditMesh([this, &Weights](FDynamicMesh3& Mesh)
+	{
+		const TArray<FTransform>& CurrentBoneTransforms = Component->GetComponentSpaceTransforms();
 
-	Tool->ExternalUpdateValues(BoneName, OldWeights);
+		// accumulate new skinned position for each vertex that has had it's weight's modified
+		for (const int32 VertexID : VerticesWithModifiedWeights)
+		{
+			FVector VertexNewPosition = FVector::ZeroVector;
+			for (TTuple<FName, TArray<float>> Pair : Weights)
+			{
+				const float VertexWeight = Pair.Value[VertexID];
+				if (VertexWeight <= MinimumWeightThreshold)
+				{
+					continue;
+				}
+				const FName BoneName = Pair.Key;
+				const int32 BoneIndex = BoneNameToIndexMap[BoneName];
+				const FVector& RefPoseVertexPosition = RefPoseVertexPositions[VertexID];
+				const FTransform& InvRefPoseTransform = InvCSRefPoseTransforms[BoneIndex];
+				const FVector& BoneLocalPositionInRefPose = InvRefPoseTransform.TransformPosition(RefPoseVertexPosition);
+				const FTransform& CurrentTransform = CurrentBoneTransforms[BoneIndex];
+				VertexNewPosition += CurrentTransform.TransformPosition(BoneLocalPositionInRefPose) * VertexWeight;
+			}
+			 
+			Mesh.SetVertex(VertexID, VertexNewPosition, false);
+		}
+	});
+
+	// update acceleration structure
+	OctreeToUpdate.RemoveVertices(VerticesWithModifiedWeights);
+	OctreeToUpdate.InsertVertices(VerticesWithModifiedWeights);
+
+	// empty queue of vertices to update
+	VerticesWithModifiedWeights.Reset();
 }
 
-void FMeshSkinWeightsChange::UpdateValues(const TArray<int32>& Indices, const TArray<float>& OldValues, const TArray<float>& NewValues)
+void FSkinToolDeformer::SetVerticesNeedUpdated(TArray<int32> VertexIndices)
 {
-	const int32 NumIndices = Indices.Num();
-	for (int32 i = 0; i < NumIndices; i++)
+	for (int32 VertexID : VertexIndices)
 	{
-		NewWeights.Add(Indices[i], NewValues[i]);
-		OldWeights.FindOrAdd(Indices[i], OldValues[i]);
+		VerticesWithModifiedWeights.Add(VertexID);
 	}
 }
 
+void FSkinToolDeformer::SetVertexNeedsUpdated(int32 VertexIndex)
+{
+	VerticesWithModifiedWeights.Add(VertexIndex);
+}
 
-/*
- * ToolBuilder
- */
+void FSkinToolWeights::InitializeSkinWeights(const FReferenceSkeleton& RefSkeleton, FMeshDescription* EditedMesh)
+{
+	const FSkeletalMeshConstAttributes MeshAttribs(*EditedMesh);
+	const FSkinWeightsVertexAttributesConstRef VertexSkinWeights = MeshAttribs.GetVertexSkinWeights();
+	const int32 NumVertices = EditedMesh->Vertices().Num();
 
-UMeshSurfacePointTool* USkinWeightsPaintToolBuilder::CreateNewTool(const FToolBuilderState& SceneState) const
+	// Create a map of all bones to their per-vertex weights.
+	CurrentWeightsMap.Reset();
+	for (const FMeshBoneInfo& BoneInfo : RefSkeleton.GetRefBoneInfo())
+	{
+		CurrentWeightsMap.Add(BoneInfo.Name, {}).AddZeroed(NumVertices);
+	}
+
+	for (int32 VertexIndex = 0; VertexIndex < NumVertices; VertexIndex++)
+	{
+		const FVertexID VertexID(VertexIndex);
+
+		for (UE::AnimationCore::FBoneWeight BoneWeight: VertexSkinWeights.Get(VertexID))
+		{
+			FName BoneName = RefSkeleton.GetBoneName(static_cast<int32>(BoneWeight.GetBoneIndex()));
+			const float Weight = BoneWeight.GetWeight();
+
+			if (Weight >= MinimumWeightThreshold)
+			{
+				// If the source mesh has a bone that we don't recognize, we ignore it. It's
+				// weight will get cleared when the new weights are updated back to the 
+				// source mesh.
+				TArray<float>* PerVertexWeights = CurrentWeightsMap.Find(BoneName);
+
+				if (PerVertexWeights)
+				{
+					PerVertexWeights->GetData()[VertexIndex] = Weight;
+				}
+			}
+		}
+	}
+	
+	// maintain duplicate weight map
+	PreStrokeWeightsMap = CurrentWeightsMap;
+
+	// maintain relax-per stroke map
+	MaxFalloffPerVertexThisStroke.SetNumZeroed(NumVertices);
+}
+
+TArray<float>* FSkinToolWeights::GetWeightsForBone(const FName BoneName)
 {
-	return NewObject<USkinWeightsPaintTool>(SceneState.ToolManager);
+	return CurrentWeightsMap.Find(BoneName);
 }
 
+void FSkinToolWeights::EditVertexWeightAndNormalize(
+	const FName BoneToHoldConstant,
+	const int32 VertexId,
+	const float NewWeightValue,
+	FMultiBoneWeightEdits& WeightEdits)
+{
+	// calculate the sum of all the weights on this vertex (not including the one we currently applied)
+	TArray<FName> NamesOfBonesAffectingVertex;
+	TArray<float> ValuesToNormalize;
+	float Total = 0.0f;
+	for (const TTuple<FName, TArray<float>>& Pair : CurrentWeightsMap)
+	{
+		const FName BoneName = Pair.Key;
+		if (BoneName == BoneToHoldConstant)
+		{
+			continue;
+		}
+		
+		const float WeightOnThisBone = CurrentWeightsMap[BoneName][VertexId];
+		if (WeightOnThisBone < MinimumWeightThreshold)
+		{
+			continue;
+		}
+		
+		NamesOfBonesAffectingVertex.Add(BoneName);
+		ValuesToNormalize.Add(WeightOnThisBone);
+		Total += WeightOnThisBone;
+	}
+
+	// if user applied FULL weight to this vertex OR there's no other weights of any significance,
+	// then simply set everything else to zero and return
+	if (NewWeightValue >= (1.0f - MinimumWeightThreshold) || Total <= MinimumWeightThreshold)
+	{
+		// set all other influences to 0.0f
+		for (int32 i=0; i<ValuesToNormalize.Num(); ++i)
+		{
+			const FName BoneName = NamesOfBonesAffectingVertex[i];
+			const float OldWeight = ValuesToNormalize[i];
+			constexpr float NewWeight = 0.0f;
+			WeightEdits.MergeSingleEdit(BoneName, VertexId, OldWeight, NewWeight);
+		}
 
-void USkinWeightsPaintTool::Setup()
+		// set current bone value to 1.0f
+		WeightEdits.MergeSingleEdit(
+			BoneToHoldConstant,
+			VertexId,
+			PreStrokeWeightsMap[BoneToHoldConstant][VertexId],
+			1.0f);
+		
+		return;
+	}
+
+	// calculate amount we have to spread across the other bones affecting this vertex
+	const float AvailableTotal = 1.0f - NewWeightValue;
+
+	// normalize weights into available space not set by current bone
+	for (int32 i=0; i<ValuesToNormalize.Num(); ++i)
+	{
+		float NormalizedValue = 0.f;
+		if (AvailableTotal > MinimumWeightThreshold && Total > KINDA_SMALL_NUMBER)
+		{
+			NormalizedValue = (ValuesToNormalize[i] / Total) * AvailableTotal;	
+		}
+		const FName BoneName = NamesOfBonesAffectingVertex[i];
+		const float OldWeight = ValuesToNormalize[i];
+		const float NewWeight = NormalizedValue;
+		WeightEdits.MergeSingleEdit(BoneName, VertexId, OldWeight, NewWeight);
+	}
+
+	// record current bone edit
+	WeightEdits.MergeSingleEdit(
+		BoneToHoldConstant,
+		VertexId,
+		PreStrokeWeightsMap[BoneToHoldConstant][VertexId],
+		NewWeightValue);
+}
+
+void FSkinToolWeights::ApplyCurrentWeightsToMeshDescription(
+	const FReferenceSkeleton& RefSkeleton,
+	FMeshDescription* EditedMesh)
 {
-	UDynamicMeshBrushTool::Setup();
+	using namespace UE::AnimationCore;
+	
+	FSkeletalMeshAttributes MeshAttribs(*EditedMesh);
+	FSkinWeightsVertexAttributesRef VertexSkinWeights = MeshAttribs.GetVertexSkinWeights();
 
-	IPrimitiveComponentBackedTarget* TargetComponent = Cast<IPrimitiveComponentBackedTarget>(Target);
-	USkeletalMeshComponent* Component = Cast<USkeletalMeshComponent>(TargetComponent->GetOwnerComponent());
+	TMap<FBoneIndexType, const TArray<float>*> BoneIndexWeightMap;
+	for (const TTuple<FName, TArray<float>>& BoneNameAndWeights : CurrentWeightsMap)
+	{
+		const int32 BoneIndex = RefSkeleton.FindBoneIndex(BoneNameAndWeights.Key);
+		if (BoneIndex != INDEX_NONE)
+		{
+			BoneIndexWeightMap.Add(static_cast<FBoneIndexType>(BoneIndex), &BoneNameAndWeights.Value);
+		}
+	}
 
-	EditedMesh = MakeUnique<FMeshDescription>();
-	*EditedMesh = *UE::ToolTarget::GetMeshDescription(Target);
+	FBoneWeightsSettings Settings;
+	Settings.SetNormalizeType(EBoneWeightNormalizeType::None);
 
-	// hide strength and falloff
-	BrushProperties->RestoreProperties(this);
+	TArray<FBoneWeight> SourceBoneWeights;
+	SourceBoneWeights.Reserve(MaxInlineBoneWeightCount);
 
-	ToolProps = NewObject<USkinWeightsPaintToolProperties>(this);
-	ToolProps->RestoreProperties(this);
-	if (Component && Component->GetSkeletalMeshAsset())
+	const int32 NumVertices = EditedMesh->Vertices().Num();
+	for (int32 VertexIndex = 0; VertexIndex < NumVertices; VertexIndex++)
 	{
-		// Get all non-virtual bones
-		// TArray<int32> BoneIndices;
-		USkeletalMesh& SkeletalMesh = *Component->GetSkeletalMeshAsset();
-		// SkeletalMesh.RefSkeleton.GetRawRefBoneInfo();
-		
+		SourceBoneWeights.Reset();
+
+		for (const TTuple<unsigned short, const TArray<float>*>& BoneIndexWeight : BoneIndexWeightMap)
+		{
+			SourceBoneWeights.Add(FBoneWeight(BoneIndexWeight.Key, (*BoneIndexWeight.Value)[VertexIndex]));
+		}
 
-		// Initialize the bone browser
-		FCurveEvaluationOption CurveEvalOption(
-				Component->GetAllowedAnimCurveEvaluate(),
-				&Component->GetDisallowedAnimCurvesEvaluation(),
-				0 /* Always use the highest LOD */
-				);
-		BoneContainer.InitializeTo(Component->RequiredBones, CurveEvalOption, SkeletalMesh);
+		VertexSkinWeights.Set(FVertexID(VertexIndex), FBoneWeights::Create(SourceBoneWeights, Settings));
+	}
+}
 
-		ToolProps->SkeletalMesh = Component->GetSkeletalMeshAsset();
-		ToolProps->CurrentBone.Initialize(BoneContainer);
+void FSkinToolWeights::ResetAfterStroke()
+{
+	PreStrokeWeightsMap = CurrentWeightsMap;
 
-		// Pick the first root bone as the initial selection.
-		PendingCurrentBone = Component->GetSkeletalMeshAsset()->GetRefSkeleton().GetBoneName(0);
-		ToolProps->CurrentBone.BoneName = PendingCurrentBone.GetValue();
+	for (int32 i=0; i<MaxFalloffPerVertexThisStroke.Num(); ++i)
+	{
+		MaxFalloffPerVertexThisStroke[i] = 0.f;
+	}
+}
 
-		// Update the skeleton drawing information from the original bind pose
-		MaxDrawRadius = Component->Bounds.SphereRadius * 0.0025f;
+float FSkinToolWeights::SetCurrentFalloffAndGetMaxFalloffThisStroke(int32 VertexID, float CurrentStength)
+{
+	float& MaxFalloffThisStroke = MaxFalloffPerVertexThisStroke[VertexID];
+	if (MaxFalloffThisStroke < CurrentStength)
+	{
+		MaxFalloffThisStroke = CurrentStength;
 	}
-	AddToolPropertySource(ToolProps);
+	return MaxFalloffThisStroke;
+}
+
+void FMeshSkinWeightsChange::Apply(UObject* Object)
+{
+	USkinWeightsPaintTool* Tool = CastChecked<USkinWeightsPaintTool>(Object);
+	
+	for (TTuple<FName, FSingleBoneWeightEdits>& Pair : AllWeightEdits.PerBoneWeightEdits)
+	{
+		Tool->ExternalUpdateWeights(Pair.Key, Pair.Value.NewWeights);
+	}
+}
+
+void FMeshSkinWeightsChange::Revert(UObject* Object)
+{
+	USkinWeightsPaintTool* Tool = CastChecked<USkinWeightsPaintTool>(Object);
+
+	for (TTuple<FName, FSingleBoneWeightEdits>& Pair : AllWeightEdits.PerBoneWeightEdits)
+	{
+		Tool->ExternalUpdateWeights(Pair.Key, Pair.Value.OldWeights);
+	}
+}
+
+void FMeshSkinWeightsChange::AddBoneWeightEdit(const FSingleBoneWeightEdits& BoneWeightEdit)
+{
+	AllWeightEdits.MergeEdits(BoneWeightEdit);
+}
+
+/*
+ * ToolBuilder
+ */
+
+UMeshSurfacePointTool* USkinWeightsPaintToolBuilder::CreateNewTool(const FToolBuilderState& SceneState) const
+{
+	return NewObject<USkinWeightsPaintTool>(SceneState.ToolManager);
+}
 
+void USkinWeightsPaintTool::Setup()
+{
+	UDynamicMeshBrushTool::Setup();
+
+	const IPrimitiveComponentBackedTarget* TargetComponent = Cast<IPrimitiveComponentBackedTarget>(Target);
+	check(TargetComponent);
+	const USkeletalMeshComponent* Component = Cast<USkeletalMeshComponent>(TargetComponent->GetOwnerComponent());
+	check(Component && Component->GetSkeletalMeshAsset())
+
+	// initialize the bone browser
+	// TODO replace this and synchronize selection & rendering with skeletal mesh editor
+	FCurveEvaluationOption CurveEvalOption(
+			Component->GetAllowedAnimCurveEvaluate(),
+			&Component->GetDisallowedAnimCurvesEvaluation(),
+			0 /* Always use the highest LOD */
+			);
+	USkeletalMesh& SkeletalMesh = *Component->GetSkeletalMeshAsset();
+	BoneContainer.InitializeTo(Component->RequiredBones, CurveEvalOption, SkeletalMesh);
+	// Update the skeleton drawing information from the original bind pose
+	MaxDrawRadius = Component->Bounds.SphereRadius * 0.0025f;
 	UpdateBonePositionInfos(MaxDrawRadius);
 
+	// initialize the tool properties
+	BrushProperties->RestoreProperties(this); // hides strength and falloff
+	ToolProps = NewObject<USkinWeightsPaintToolProperties>(this);
+	ToolProps->RestoreProperties(this);
+	ToolProps->SkeletalMesh = Component->GetSkeletalMeshAsset();
+	ToolProps->CurrentBone.Initialize(BoneContainer);
+	AddToolPropertySource(ToolProps);
+	// attach callback to be informed when tool properties are modified
+	ToolProps->GetOnModified().AddUObject(this, &USkinWeightsPaintTool::OnToolPropertiesModified);
+
+	// create a mesh description for editing
+	EditedMesh = MakeUnique<FMeshDescription>();
+	*EditedMesh = *UE::ToolTarget::GetMeshDescription(Target);
+
+	// default to the root bone as current bone
+	PendingCurrentBone = CurrentBone = ToolProps->SkeletalMesh->GetRefSkeleton().GetBoneName(0);
+	ToolProps->CurrentBone.BoneName = PendingCurrentBone.GetValue();
+
 	// configure preview mesh
 	PreviewMesh->SetTangentsMode(EDynamicMeshComponentTangentsMode::AutoCalculated);
-	//PreviewMesh->EnableWireframe(SelectionProps->bShowWireframe);
+	PreviewMesh->EnableWireframe(true);
 	PreviewMesh->SetShadowsEnabled(false);
-
 	// enable vtx colors on preview mesh
 	PreviewMesh->EditMesh([](FDynamicMesh3& Mesh)
 	{
@@ -119,44 +415,39 @@ void USkinWeightsPaintTool::Setup()
 		// Create an overlay that has no split elements, init with zero value.
 		Mesh.Attributes()->PrimaryColors()->CreateFromPredicate([](int ParentVID, int TriIDA, int TriIDB){return true;}, 0.f);
 	});
-
-	// build octree
-	VerticesOctree.Initialize(PreviewMesh->GetMesh(), true);
-
 	UMaterialInterface* VtxColorMaterial = GetToolManager()->GetContextQueriesAPI()->GetStandardMaterial(EStandardToolContextMaterials::VertexColorMaterial);
 	if (VtxColorMaterial != nullptr)
 	{
 		PreviewMesh->SetOverrideRenderMaterial(VtxColorMaterial);
 	}
 
-	RecalculateBrushRadius();
+	// build octree for brush
+	VerticesOctree.Initialize(PreviewMesh->GetMesh(), true);
 
+	// record data needed to compute deformations
+	Deformer.Initialize(Component, EditedMesh.Get());
+	
+	// copy weight maps
+	const FReferenceSkeleton& RefSkeleton = ToolProps->SkeletalMesh->GetRefSkeleton();
+	Weights.InitializeSkinWeights(RefSkeleton, EditedMesh.Get());
+	bVisibleWeightsValid = false;
+
+	RecalculateBrushRadius();
+	
+	// inform user of tool keys
+	// TODO talk with UX team about viewport overlay to show hotkeys
 	GetToolManager()->DisplayMessage(
 		LOCTEXT("OnStartSkinWeightsPaint", "Paint per-bone skin weights. [ and ] change brush size, Ctrl to Erase/Subtract, Shift to Smooth"),
 		EToolMessageLevel::UserNotification);
-
-	InitializeSkinWeights();
-
-	CurrentBoneWatcher.Initialize([this]() { return ToolProps->CurrentBone; },
-		[this](FBoneReference NewValue) { PendingCurrentBone = NewValue.BoneName; }, ToolProps->CurrentBone);
-
-	bVisibleWeightsValid = false;
 }
 
-
-
 void USkinWeightsPaintTool::RegisterActions(FInteractiveToolActionSet& ActionSet)
 {
 	UDynamicMeshBrushTool::RegisterActions(ActionSet);
 }
 
-
-
-
 void USkinWeightsPaintTool::OnTick(float DeltaTime)
 {
-	CurrentBoneWatcher.CheckAndUpdate();
-
 	if (bStampPending)
 	{
 		ApplyStamp(LastStamp);
@@ -169,61 +460,38 @@ void USkinWeightsPaintTool::OnTick(float DeltaTime)
 		PendingCurrentBone.Reset();
 	}
 
-	if (bVisibleWeightsValid == false)
+	if (bVisibleWeightsValid == false || ToolProps->bColorModeChanged)
 	{
-		UpdateBoneVisualization();
+		UpdateCurrentBoneVertexColors();
 		bVisibleWeightsValid = true;
+		ToolProps->bColorModeChanged = false;
 	}
-}
-
-
 
+	// sparsely updates vertex positions (only when needed)
+	Deformer.UpdateVertexDeformation(Weights.CurrentWeightsMap, PreviewMesh, VerticesOctree);
+}
 
 void USkinWeightsPaintTool::Render(IToolsContextRenderAPI* RenderAPI)
 {
 	UDynamicMeshBrushTool::Render(RenderAPI);
 
 	// FIXME: Make selective.
-	RenderBonePositions(RenderAPI->GetPrimitiveDrawInterface());
+	//RenderBonePositions(RenderAPI->GetPrimitiveDrawInterface());
 }
 
-
-
-
-
-
-
 bool USkinWeightsPaintTool::HitTest(const FRay& Ray, FHitResult& OutHit)
 {
-	bool bHit = UDynamicMeshBrushTool::HitTest(Ray, OutHit);
-	//if (bHit && SelectionProps->bHitBackFaces == false)
-	//{
-	//	const FDynamicMesh3* SourceMesh = PreviewMesh->GetPreviewDynamicMesh();
-	//	FVector3d Normal, Centroid;
-	//	double Area;
-	//	SourceMesh->GetTriInfo(OutHit.FaceIndex, Normal, Area, Centroid);
-	//	FViewCameraState StateOut;
-	//	GetToolManager()->GetContextQueriesAPI()->GetCurrentViewState(StateOut);
-	//	FVector3d LocalEyePosition(ComponentTarget->GetWorldTransform().InverseTransformPosition(StateOut.Position));
-
-	//	if (Normal.Dot((Centroid - LocalEyePosition)) > 0)
-	//	{
-	//		bHit = false;
-	//	}
-	//}
-	return bHit;
+	return UDynamicMeshBrushTool::HitTest(Ray, OutHit);
 }
 
-
 void USkinWeightsPaintTool::OnBeginDrag(const FRay& WorldRay)
 {
 	UDynamicMeshBrushTool::OnBeginDrag(WorldRay);
 
-	PreviewBrushROI.Reset();
 	if (IsInBrushStroke())
 	{
-		bInRemoveStroke = GetCtrlToggle();
-		bInSmoothStroke = GetShiftToggle();
+		bInvertStroke = GetCtrlToggle();
+		bSmoothStroke = GetShiftToggle();
 		BeginChange();
 		StartStamp = UBaseBrushTool::LastBrushStamp;
 		LastStamp = StartStamp;
@@ -231,8 +499,6 @@ void USkinWeightsPaintTool::OnBeginDrag(const FRay& WorldRay)
 	}
 }
 
-
-
 void USkinWeightsPaintTool::OnUpdateDrag(const FRay& WorldRay)
 {
 	UDynamicMeshBrushTool::OnUpdateDrag(WorldRay);
@@ -243,16 +509,16 @@ void USkinWeightsPaintTool::OnUpdateDrag(const FRay& WorldRay)
 	}
 }
 
-
-
 void USkinWeightsPaintTool::OnEndDrag(const FRay& Ray)
 {
 	UDynamicMeshBrushTool::OnEndDrag(Ray);
 
-	bInRemoveStroke = false;
-	bInSmoothStroke = false;
+	bInvertStroke = false;
+	bSmoothStroke = false;
 	bStampPending = false;
 
+	Weights.ResetAfterStroke();
+
 	// close change record
 	TUniquePtr<FMeshSkinWeightsChange> Change = EndChange();
 
@@ -263,58 +529,154 @@ void USkinWeightsPaintTool::OnEndDrag(const FRay& Ray)
 	GetToolManager()->EndUndoTransaction();
 }
 
-
 bool USkinWeightsPaintTool::OnUpdateHover(const FInputDeviceRay& DevicePos)
 {
 	UDynamicMeshBrushTool::OnUpdateHover(DevicePos);
-
-	// todo get rid of this redundant hit test!
-	FHitResult OutHit;
-	if (UDynamicMeshBrushTool::HitTest(DevicePos.WorldRay, OutHit))
-	{
-		PreviewBrushROI.Reset();
-		CalculateVertexROI(LastBrushStamp, PreviewBrushROI);
-	}
-
 	return true;
 }
 
-
-
-void USkinWeightsPaintTool::CalculateVertexROI(const FBrushStampData& Stamp, TArray<int>& VertexROI)
+void USkinWeightsPaintTool::CalculateVertexROI(
+	const FBrushStampData& Stamp,
+	TArray<int32>& VertexROI,
+	TArray<float>& VertexSqDistances) const
 {
 	using namespace UE::Geometry;
-	IPrimitiveComponentBackedTarget* TargetComponent = Cast<IPrimitiveComponentBackedTarget>(Target);
-	FTransform3d Transform(TargetComponent->GetWorldTransform());
-	FVector3d StampPosLocal = Transform.InverseTransformPosition((FVector3d)Stamp.WorldPosition);
 
-	float RadiusSqr = CurrentBrushRadius * CurrentBrushRadius;
-	const FDynamicMesh3* Mesh = PreviewMesh->GetPreviewDynamicMesh();
-	FAxisAlignedBox3d QueryBox(StampPosLocal, CurrentBrushRadius);
-	VerticesOctree.RangeQuery(QueryBox,
-		[&](int32 VertexID) { return FVector3d::DistSquared(Mesh->GetVertex(VertexID), StampPosLocal) < RadiusSqr; },
-		VertexROI);
+	if (ToolProps->FalloffMode == EWeightBrushFalloffMode::Volume)
+	{
+		IPrimitiveComponentBackedTarget* TargetComponent = Cast<IPrimitiveComponentBackedTarget>(Target);
+		FTransform3d Transform(TargetComponent->GetWorldTransform());
+		FVector3d StampPosLocal = Transform.InverseTransformPosition((FVector3d)Stamp.WorldPosition);
+
+		float RadiusSqr = CurrentBrushRadius * CurrentBrushRadius;
+		const FDynamicMesh3* Mesh = PreviewMesh->GetPreviewDynamicMesh();
+		FAxisAlignedBox3d QueryBox(StampPosLocal, CurrentBrushRadius);
+		VerticesOctree.RangeQuery(QueryBox,
+			[&](int32 VertexID) { return FVector3d::DistSquared(Mesh->GetVertex(VertexID), StampPosLocal) < RadiusSqr; },
+			VertexROI);
+
+		for (int32 VertexID : VertexROI)
+		{
+			VertexSqDistances.Add(FVector3d::DistSquared(Mesh->GetVertex(VertexID), StampPosLocal));
+		}
+	}
+	else
+	{
+		auto GetFrameFromStamp = [](const FBrushStampData& InStamp) -> FFrame3d
+		{
+			const FVector3d Origin = InStamp.WorldPosition;
+			const FVector3d Normal = InStamp.WorldNormal;
+			FVector3d NonCollinear = Normal;
+			// get a guaranteed non collinear vector to the normal
+			// doesn't matter where in the plane, stamp is radially symmetric
+			do 
+			{
+				NonCollinear.X = FMath::RandRange(-1.0f, 1.0f);
+				NonCollinear.Y = FMath::RandRange(-1.0f, 1.0f);
+				NonCollinear.Z = FMath::RandRange(-1.0f, 1.0f);
+				NonCollinear.Normalize();
+				
+			} while (FMath::Abs(NonCollinear.Dot(Normal)) > 0.8f);
+
+			const FVector3d Plane = Normal.Cross(NonCollinear);
+			const FVector3d Cross = Plane.Cross(Normal);
+			return FFrame3d(Origin, Cross, Plane, Normal);
+			
+		};
+
+		// get coordinate frame from stamp
+		FFrame3d SeedFrame = GetFrameFromStamp(Stamp);
+		
+		const FDynamicMesh3* Mesh = PreviewMesh->GetPreviewDynamicMesh();
+
+		// get triangle under the stamp
+		// we pass it's vertices to the ExpMap generator to act as seed nodes to start the search from
+		// this is necessary because the stamp may be in a large face center, from from any vertex on the mesh.
+		FDynamicMeshAABBTree3 Spatial(Mesh, true);
+		double NearDistSqr;
+		int32 SeedTID = Spatial.FindNearestTriangle(SeedFrame.Origin, NearDistSqr);
+		//FDistPoint3Triangle3d Query = TMeshQueries<FDynamicMesh3>::TriangleDistance(*Mesh, SeedTID, SeedFrame.Origin);
+		FIndex3i TriVerts = Mesh->GetTriangle(SeedTID);
+
+		// create the ExpMap generator, computes vertex polar coordinates in a plane tangent to the surface
+		TMeshLocalParam<FDynamicMesh3> Param(Mesh);
+		Param.ParamMode = ELocalParamTypes::ExponentialMapUpwindAvg;
+		Param.ComputeToMaxDistance(SeedFrame, TriVerts, Stamp.Radius);
+
+		// store vertices under the brush and their distances from the stamp
+		const float StampRadSq = FMath::Pow(Stamp.Radius, 2);
+		for (int32 VertexID : Mesh->VertexIndicesItr())
+		{
+			if (!Param.HasUV(VertexID))
+			{
+				continue;
+			}
+			
+			FVector2d UV = Param.GetUV(VertexID);
+			const float DistSq = UV.SizeSquared();
+			if (DistSq >= StampRadSq)
+			{
+				continue;
+			}
+			
+			VertexSqDistances.Add(DistSq);
+			VertexROI.Add(VertexID);
+		}
+	}
 }
 
-
 FVector4f USkinWeightsPaintTool::WeightToColor(float Value)
 {
-	Value = FMath::Clamp(Value, 0.0f, 1.0f);
+	// optional greyscale mode
+	if (ToolProps->ColorMode == EWeightColorMode::Greyscale)
+	{
+		return FLinearColor::LerpUsingHSV(FLinearColor::Black, FLinearColor::White, Value);	
+	}
+	
+	// early out zero weights to min color
+	if (Value <= MinimumWeightThreshold)
+	{
+		return ToolProps->MinColor;
+	}
 
+	// early out full weights to max color
+	if (FMath::IsNearlyEqual(Value, 1.0f))
 	{
-		// A close approximation of the skeletal mesh editor's bone weight ramp. 
-		const FLinearColor HSV((1.0f - Value) * 285.0f, 100.0f, 85.0f);
-		return UE::Geometry::ToVector4<float>(HSV.HSVToLinearRGB());
+		return ToolProps->MaxColor;
 	}
+
+	// get user-specified color ramp for intermediate colors
+	const TArray<FLinearColor>& Colors = ToolProps->ColorRamp;
+	
+	// revert back to simple Lerp(min,max) if user supplied color ramp doesn't have enough colors
+	if (Colors.Num() < 2)
+	{
+		const FLinearColor FinalColor = FLinearColor::LerpUsingHSV(ToolProps->MinColor, ToolProps->MaxColor, Value);
+		return UE::Geometry::ToVector4<float>(FinalColor);
+	}
+
+	// otherwise, interpolate within two nearest ramp colors
+	Value = FMath::Clamp(Value, 0.0f, 1.0f);
+	const float PerColorRange = 1.0f / (Colors.Num() - 1);
+	const int ColorIndex = static_cast<int>(Value / PerColorRange);
+	const float RangeStart = ColorIndex * PerColorRange;
+	const float RangeEnd = (ColorIndex + 1) * PerColorRange;
+	const float Param = (Value - RangeStart) / (RangeEnd - RangeStart);
+	const FLinearColor& StartColor = Colors[ColorIndex];
+	const FLinearColor& EndColor = Colors[ColorIndex+1];
+	const FLinearColor FinalColor = FLinearColor::LerpUsingHSV(StartColor, EndColor, Param);
+	return UE::Geometry::ToVector4<float>(FinalColor);
 }
 
 
-void USkinWeightsPaintTool::UpdateBoneVisualization()
+void USkinWeightsPaintTool::UpdateCurrentBoneVertexColors()
 {
-	if (!SkinWeightsMap.Contains(CurrentBone))
+	if (!Weights.CurrentWeightsMap.Contains(CurrentBone))
+	{
 		return;
-
-	TArray<float>& SkinWeightsData = *SkinWeightsMap.Find(CurrentBone);
+	}
+	
+	TArray<float>& SkinWeightsData = *Weights.CurrentWeightsMap.Find(CurrentBone);
 
 	// update mesh with new value colors
 	PreviewMesh->EditMesh([&](FDynamicMesh3& Mesh)
@@ -325,14 +687,13 @@ void USkinWeightsPaintTool::UpdateBoneVisualization()
 		{
 			const int32 VertexId = ColorOverlay->GetParentVertex(ElementId);	
 			const int32 SrcVertexId = NonManifoldMappingSupport.GetOriginalNonManifoldVertexID(VertexId);
-               const float Value = SkinWeightsData.IsValidIndex(SrcVertexId) ? SkinWeightsData[SrcVertexId] : 0.f;
+			const float Value = SkinWeightsData.IsValidIndex(SrcVertexId) ? SkinWeightsData[SrcVertexId] : 0.f;
 			const FVector4f Color(WeightToColor(Value));
 			ColorOverlay->SetElement(ElementId, Color);
 		}
 	});
 }
 
-
 double USkinWeightsPaintTool::CalculateBrushFalloff(double Distance) const
 {
 	double f = FMathd::Clamp(1.0 - BrushProperties->BrushFalloffAmount, 0.0, 1.0);
@@ -347,131 +708,267 @@ double USkinWeightsPaintTool::CalculateBrushFalloff(double Distance) const
 	return w;
 }
 
-
-
 void USkinWeightsPaintTool::ApplyStamp(const FBrushStampData& Stamp)
 {
-	using namespace UE::Geometry;
-	
-	// FIXME: Move to earlier.
-	if (!SkinWeightsMap.Contains(CurrentBone))
-		return;
+	// get the vertices under the brush, and their squared distances to the brush center
+	// when using "Volume" brush, distances are straight line
+	// when using "Surface" brush, distances are geodesics
+	TArray<int32> VerticesInStamp;
+	TArray<float> VertexSqDistances;
+	CalculateVertexROI(Stamp, VerticesInStamp, VertexSqDistances);
+
+	// gather sparse set of modifications made from this stamp, these edits are merged throughout
+	// the lifetime of a single brush stroke in the "ActiveChange" allowing for undo/redo
+	FMultiBoneWeightEdits WeightEditsFromStamp;
+
+	// generate a weight edit from this stamp (includes modifications caused by normalization)
+	if (bSmoothStroke || ToolProps->BrushMode == EBrushBehaviorMode::Relax)
+	{
+		// use mesh topology to iteratively smooth weights across neighboring vertices
+		RelaxWeightOnVertices(VerticesInStamp, VertexSqDistances, WeightEditsFromStamp);
+	}
+	else
+	{
+		// edit weight; either by "Add", "Remove", "Replace", "Multiply"
+		EditWeightOfVerticesInStamp(
+			ToolProps->BrushMode,
+			VerticesInStamp,
+			VertexSqDistances,
+			WeightEditsFromStamp);
+	}
 
-	IPrimitiveComponentBackedTarget* TargetComponent = Cast<IPrimitiveComponentBackedTarget>(Target);
-	FTransform3d Transform(TargetComponent->GetWorldTransform());
-	FVector3d StampPosLocal = Transform.InverseTransformPosition((FVector3d)Stamp.WorldPosition);
+	// store weight edits from all stamps made during a single stroke (1 transaction per stroke)
+	for (const TTuple<FName, FSingleBoneWeightEdits>& BoneWeightEdits : WeightEditsFromStamp.PerBoneWeightEdits)
+	{
+		ActiveChange->AddBoneWeightEdit(BoneWeightEdits.Value);
+	}
 
-	TArray<int32> ROIVertices;
-	CalculateVertexROI(Stamp, ROIVertices);
-	const int32 NumROIVertices = ROIVertices.Num();
+	// apply weights to weight map
+	WeightEditsFromStamp.ApplyEditsToWeightMap(Weights.CurrentWeightsMap);
 
-	TArray<float> ROIBefore, ROIAfter;
-	ROIBefore.SetNum(NumROIVertices);
-	ROIAfter.SetNum(NumROIVertices);
+	// weights have been modified, so update deformations
+	Deformer.SetVerticesNeedUpdated(VerticesInStamp);
 
-	TArray<float>& SkinWeightsData = *SkinWeightsMap.Find(CurrentBone);
+	// update vertex colors
+	PreviewMesh->DeferredEditMesh([&](FDynamicMesh3& Mesh)
+	{
+		TArray<int> ElementIds;
+		UE::Geometry::FDynamicMeshColorOverlay* ColorOverlay = Mesh.Attributes()->PrimaryColors();
+		const TArray<float>& VertexWeights = Weights.CurrentWeightsMap[CurrentBone];
+		const int32 NumVerticesInStamp = VerticesInStamp.Num();
+		for (int32 Index = 0; Index < NumVerticesInStamp; ++Index)
+		{
+			const int32 VertexId = VerticesInStamp[Index];
+			FVector4f NewColor(WeightToColor(VertexWeights[VertexId]));
+			ColorOverlay->GetVertexElements(VertexId, ElementIds);
+			for (const int32 ElementId : ElementIds)
+			{
+				ColorOverlay->SetElement(ElementId, NewColor);
+			}
+			ElementIds.Reset();
+		}
+	}, false);
+	PreviewMesh->NotifyDeferredEditCompleted(UPreviewMesh::ERenderUpdateMode::FastUpdate, EMeshRenderAttributeFlags::VertexColors, false);
+}
 
+void USkinWeightsPaintTool::RelaxWeightOnVertices(
+	TArray<int32> VerticesInStamp,
+	TArray<float> VertexSqDistances,
+	FMultiBoneWeightEdits& AllBoneWeightEditsFromStamp)
+{
 	const FDynamicMesh3* CurrentMesh = PreviewMesh->GetMesh();
-	UE::Geometry::FNonManifoldMappingSupport NonManifoldMappingSupport(*CurrentMesh);
-	if (bInSmoothStroke)
+	const UE::Geometry::FNonManifoldMappingSupport NonManifoldMappingSupport(*CurrentMesh);
+
+	auto NormalizeWeights = [](TMap<FName, float>& InOutWeights)
 	{
-		const float SmoothSpeed = 0.25f;
+		float TotalWeight = 0.f;
+		for (const TTuple<FName, float>& Weight : InOutWeights)
+		{
+			TotalWeight += Weight.Value;
+		}
 
-		for (int32 Index = 0; Index < NumROIVertices; ++Index)
+		for (TTuple<FName, float>& Weight : InOutWeights)
 		{
-			const int32 VertexId = ROIVertices[Index];
-			const int32 SrcVertexId = NonManifoldMappingSupport.GetOriginalNonManifoldVertexID(VertexId);
-			FVector3d Position = CurrentMesh->GetVertex(VertexId);
-			float ValueSum = 0, WeightSum = 0;
-			for (int32 NeighborVertexId : CurrentMesh->VtxVerticesItr(VertexId))
+			Weight.Value /= TotalWeight;
+		}
+	};
+
+	// for each vertex in the stamp...
+	TArray<int32> AllNeighborVertices;
+	TMap<FName, TArray<float>> WeightsOnAllNeighbors;
+	TMap<FName, float> FinalWeights;
+	for (int32 Index = 0; Index < VerticesInStamp.Num(); ++Index)
+	{
+		const int32 VertexId = VerticesInStamp[Index];
+		const int32 SrcVertexId = NonManifoldMappingSupport.GetOriginalNonManifoldVertexID(VertexId);
+
+		// get list of all neighboring vertices, AND this vertex
+		AllNeighborVertices.Reset();
+		AllNeighborVertices.Add(VertexId);
+		for (const int32 NeighborVertexId : CurrentMesh->VtxVerticesItr(SrcVertexId))
+		{
+			AllNeighborVertices.Add(NeighborVertexId);
+		}
+
+		// get all weights above a given threshold across ALL neighbors (including self)
+		WeightsOnAllNeighbors.Reset();
+		for (const int32 VertexID : AllNeighborVertices)
+		{
+			for (TTuple<FName, TArray<float>> Pair : Weights.CurrentWeightsMap)
 			{
-				FVector3d NbrPos = CurrentMesh->GetVertex(NeighborVertexId);
-				const float Weight = FMathf::Clamp(1.0f / FVector3d::DistSquared(NbrPos, Position), 0.0001f, 1000.0f);
-				const int32 SrcNeighborVertexId = NonManifoldMappingSupport.GetOriginalNonManifoldVertexID(NeighborVertexId);
-				ValueSum += Weight * SkinWeightsData[SrcNeighborVertexId];
-				WeightSum += Weight;
+				const float NeighborWeightValue = Pair.Value[VertexID];
+				if (NeighborWeightValue > MinimumWeightThreshold)
+				{
+					const FName BoneName = Pair.Key;
+					TArray<float>& BoneWeights = WeightsOnAllNeighbors.FindOrAdd(BoneName);
+					BoneWeights.Add(NeighborWeightValue);
+				}
 			}
-			ValueSum /= WeightSum;
+		}
 
-			const float Falloff = float(CalculateBrushFalloff(FVector3d::Dist(Position, StampPosLocal)));
-			const float NewValue = FMathf::Lerp(SkinWeightsData[SrcVertexId], ValueSum, SmoothSpeed*Falloff);
+		// calculate single average weight of each bone on all the neighbors
+		FinalWeights.Reset();
+		for (const TTuple<FName, TArray<float>>& NeighborWeights : WeightsOnAllNeighbors)
+		{
+			float TotalWeightOnThisBone = 0.f;
+			for (const float& Value : NeighborWeights.Value)
+			{
+				TotalWeightOnThisBone += Value;
+			}
+			FinalWeights.Add(NeighborWeights.Key, TotalWeightOnThisBone / NeighborWeights.Value.Num());
+		}
 
-			ROIBefore[Index] = SkinWeightsData[SrcVertexId];
-			ROIAfter[Index] = FMath::Clamp(NewValue, 0.0f, 1.0f);
+		// normalize the weights
+		NormalizeWeights(FinalWeights);
+
+		// lerp weights from previous values, to fully relaxed values by brush strength scaled by falloff
+		const float CurrentFalloff = static_cast<float>(CalculateBrushFalloff(FMath::Sqrt(VertexSqDistances[Index])));
+		const float UseFalloff = Weights.SetCurrentFalloffAndGetMaxFalloffThisStroke(VertexId, CurrentFalloff);
+		const float UseStrength = BrushProperties->BrushStrength * UseFalloff;
+		for (TTuple<FName, float>& FinalWeight : FinalWeights)
+		{
+			FName BoneName = FinalWeight.Key;
+			float NewWeight = FinalWeight.Value;
+			float OldWeight = Weights.PreStrokeWeightsMap[BoneName][VertexId];
+			FinalWeight.Value = FMath::Lerp(OldWeight, NewWeight, UseStrength);
 		}
-	}
-	else
-	{
-		const bool bInvert = bInRemoveStroke;
-		const float Sign = (bInvert) ? -1.0f : 1.0f;
-		const float UseStrength = Sign * BrushProperties->BrushStrength;
-		for (int32 Index = 0; Index < NumROIVertices; ++Index)
+
+		// normalize again
+		NormalizeWeights(FinalWeights);
+
+		// apply weight edits
+		for (const TTuple<FName, float>& FinalWeight : FinalWeights)
 		{
-			const int32 VertexId = ROIVertices[Index];
-			const int32 SrcVertexId = NonManifoldMappingSupport.GetOriginalNonManifoldVertexID(VertexId);
-			const FVector3d Position = CurrentMesh->GetVertex(VertexId);
-			const float Falloff = (float)CalculateBrushFalloff(FVector3d::Dist(Position, StampPosLocal));
-			ROIBefore[Index] = SkinWeightsData[SrcVertexId];
-			ROIAfter[Index] = FMath::Clamp(ROIBefore[Index] + UseStrength * Falloff, 0.0f, 1.0f);
+			// record an edit for this vertex, for this bone
+			FName BoneName = FinalWeight.Key;
+			float NewWeight = FinalWeight.Value;
+			float OldWeight = Weights.PreStrokeWeightsMap[BoneName][VertexId];
+			AllBoneWeightEditsFromStamp.MergeSingleEdit(BoneName, VertexId, OldWeight, NewWeight);
 		}
 	}
+}
+
+void USkinWeightsPaintTool::EditWeightOfVerticesInStamp(
+	EBrushBehaviorMode EditMode,
+	const TArray<int32>& VerticesInStamp,
+	const TArray<float>& VertexSqDistances,
+	FMultiBoneWeightEdits& AllBoneWeightEditsFromStamp)
+{
+	const UE::Geometry::FNonManifoldMappingSupport NonManifoldMappingSupport(*PreviewMesh->GetMesh());
+	TArray<float>& PreWeightData = *Weights.PreStrokeWeightsMap.Find(CurrentBone);
 
-	// track changes
-	if (ActiveChange)
+	// invert brush strength differently depending on brush mode
+	float UseStrength = BrushProperties->BrushStrength;
+	switch (EditMode)
 	{
-		ActiveChange->UpdateValues(ROIVertices, ROIBefore, ROIAfter);
+	case EBrushBehaviorMode::Add:
+		{
+			UseStrength *= bInvertStroke ? -1.0f : 1.0f;
+			break;
+		}
+	case EBrushBehaviorMode::Replace:
+		{
+			UseStrength = bInvertStroke ? 1.0f - UseStrength : UseStrength;
+			break;
+		}
+	case EBrushBehaviorMode::Multiply:
+		{
+			UseStrength = bInvertStroke ? 1.0f + UseStrength : UseStrength;
+			break;
+		}		
+	default:
+		checkNoEntry();
 	}
-
-	// update values and colors
-	PreviewMesh->DeferredEditMesh([&](FDynamicMesh3& Mesh)
+	
+	// spin through the vertices in the stamp and store new weight values in NewValuesFromStap
+	// afterwards, these values are normalized while taking into consideration the users desired changes
+	const int32 NumVerticesInStamp = VerticesInStamp.Num();
+	for (int32 Index = 0; Index < NumVerticesInStamp; ++Index)
 	{
-		TArray<int> ElementIds;
-		FDynamicMeshColorOverlay* ColorOverlay = Mesh.Attributes()->PrimaryColors();
-		for (int32 Index = 0; Index < NumROIVertices; ++Index)
+		const int32 VertexId = VerticesInStamp[Index];
+		const int32 SrcVertexId = NonManifoldMappingSupport.GetOriginalNonManifoldVertexID(VertexId);
+		const float CurrentFalloff = static_cast<float>(CalculateBrushFalloff(FMath::Sqrt(VertexSqDistances[Index])));
+		const float UseFalloff = Weights.SetCurrentFalloffAndGetMaxFalloffThisStroke(SrcVertexId, CurrentFalloff);
+		const float ValueBeforeStroke = PreWeightData[SrcVertexId];
+
+		// calculate new weight value
+		float NewValueAfterStamp = ValueBeforeStroke;
+		switch (EditMode)
 		{
-			const int32 VertexId = ROIVertices[Index];
-			const int32 SrcVertexId = NonManifoldMappingSupport.GetOriginalNonManifoldVertexID(VertexId);
-			SkinWeightsData[SrcVertexId] = ROIAfter[Index];
-			FVector4f NewColor(WeightToColor(ROIAfter[Index]));
-			ColorOverlay->GetVertexElements(VertexId, ElementIds);
-			for (int ElementId : ElementIds)
+		case EBrushBehaviorMode::Add:
 			{
-				ColorOverlay->SetElement(ElementId, NewColor);
+				NewValueAfterStamp = ValueBeforeStroke + (UseStrength * UseFalloff);
+				break;
 			}
-			ElementIds.Reset();
+		case EBrushBehaviorMode::Replace:
+			{
+				NewValueAfterStamp = FMath::Lerp(ValueBeforeStroke, UseStrength, UseFalloff);
+				break;
+			}
+		case EBrushBehaviorMode::Multiply:
+			{
+				const float DeltaFromThisStamp = ((ValueBeforeStroke * UseStrength) - ValueBeforeStroke) * UseFalloff;
+				NewValueAfterStamp = ValueBeforeStroke + DeltaFromThisStamp;
+				break;
+			}
+		default:
+			checkNoEntry();
 		}
-	}, false);
-	PreviewMesh->NotifyDeferredEditCompleted(UPreviewMesh::ERenderUpdateMode::FastUpdate, EMeshRenderAttributeFlags::VertexColors, false);
 
+		// normalize the values across all bones affecting the vertices in the stamp, and record the bone edits
+		// normalization is done while holding all weights on the current bone constant so that user edits are not overwritten
+		NewValueAfterStamp = FMath::Clamp(NewValueAfterStamp, 0.0f, 1.0f);
+		Weights.EditVertexWeightAndNormalize(
+			CurrentBone,
+			VertexId,
+			NewValueAfterStamp,
+			AllBoneWeightEditsFromStamp);
+	}
 }
 
-
 void USkinWeightsPaintTool::UpdateCurrentBone(const FName& BoneName)
 {
 	CurrentBone = BoneName;
-
 	bVisibleWeightsValid = false;
 }
 
-
 void USkinWeightsPaintTool::OnShutdown(EToolShutdownType ShutdownType)
 {
 	BrushProperties->SaveProperties(this);
 
 	if (ShutdownType == EToolShutdownType::Accept)
 	{
-		UpdateEditedSkinWeightsMesh();
+		// apply the weights to the mesh description
+		const FReferenceSkeleton& RefSkeleton = ToolProps->SkeletalMesh->GetRefSkeleton();
+		Weights.ApplyCurrentWeightsToMeshDescription(RefSkeleton, EditedMesh.Get());
 
 		// this block bakes the modified DynamicMeshComponent back into the StaticMeshComponent inside an undo transaction
 		GetToolManager()->BeginUndoTransaction(LOCTEXT("SkinWeightsPaintTool", "Paint Skin Weights"));
-
 		UE::ToolTarget::CommitMeshDescriptionUpdate(Target, EditedMesh.Get());
-
 		GetToolManager()->EndUndoTransaction();
 	}
 }
 
-
 void USkinWeightsPaintTool::UpdateBonePositionInfos(float MinRadius)
 {
 	const FReferenceSkeleton& RefSkeleton = BoneContainer.GetReferenceSkeleton();
@@ -619,7 +1116,7 @@ void USkinWeightsPaintTool::RenderBonePositions(FPrimitiveDrawInterface* PDI)
 
 void USkinWeightsPaintTool::BeginChange()
 {
-	ActiveChange = MakeUnique<FMeshSkinWeightsChange>(CurrentBone);
+	ActiveChange = MakeUnique<FMeshSkinWeightsChange>();
 }
 
 
@@ -629,105 +1126,45 @@ TUniquePtr<FMeshSkinWeightsChange> USkinWeightsPaintTool::EndChange()
 }
 
 
-void USkinWeightsPaintTool::ExternalUpdateValues(const FName& BoneName, const TMap<int32, float>& NewValues)
+void USkinWeightsPaintTool::ExternalUpdateWeights(const FName& BoneName, const TMap<int32, float>& NewValues)
 {
-	TArray<float>* SkinWeightValues = SkinWeightsMap.Find(BoneName);
-	if (SkinWeightValues == nullptr)
+	TArray<float>* PostSkinWeightValues = Weights.CurrentWeightsMap.Find(BoneName);
+	TArray<float>* PreSkinWeightValues = Weights.PreStrokeWeightsMap.Find(BoneName);
+	if (PostSkinWeightValues == nullptr)
 	{
 		return;
 	}
 
-	for (const auto& IV : NewValues)
+	for (const TTuple<int, float>& IV : NewValues)
 	{
-		SkinWeightValues->GetData()[IV.Key] = IV.Value;
+		PostSkinWeightValues->GetData()[IV.Key] = IV.Value;
+		PreSkinWeightValues->GetData()[IV.Key] = IV.Value;
+		Deformer.SetVertexNeedsUpdated(IV.Key);
 	}
 
 	if (BoneName == CurrentBone)
-		UpdateBoneVisualization();
-}
-
-
-void USkinWeightsPaintTool::InitializeSkinWeights()
-{
-	const FReferenceSkeleton& RefSkeleton = BoneContainer.GetReferenceSkeleton();
-
-	const FSkeletalMeshConstAttributes MeshAttribs(*EditedMesh);
-	const FSkinWeightsVertexAttributesConstRef VertexSkinWeights = MeshAttribs.GetVertexSkinWeights();
-	const int32 NumVertices = EditedMesh->Vertices().Num();
-
-	// Create a map of all bones to their per-vertex weights.
-	SkinWeightsMap.Reset();
-	for (const FMeshBoneInfo& BoneInfo : RefSkeleton.GetRefBoneInfo())
 	{
-		SkinWeightsMap.Add(BoneInfo.Name, {}).AddZeroed(NumVertices);
+		UpdateCurrentBoneVertexColors();
 	}
-
-	for (int32 VertexIndex = 0; VertexIndex < NumVertices; VertexIndex++)
-	{
-		const FVertexID VertexID(VertexIndex);
-
-		for (UE::AnimationCore::FBoneWeight BoneWeight: VertexSkinWeights.Get(VertexID))
-		{
-			FName BoneName = RefSkeleton.GetBoneName(static_cast<int32>(BoneWeight.GetBoneIndex()));
-			const float Weight = BoneWeight.GetWeight();
-
-			if (Weight >= MinimumWeightThreshold)
-			{
-				// If the source mesh has a bone that we don't recognize, we ignore it. It's
-				// weight will get cleared when the new weights are updated back to the 
-				// source mesh.
-				TArray<float>* PerVertexWeights = SkinWeightsMap.Find(BoneName);
-
-				if (PerVertexWeights)
-				{
-					PerVertexWeights->GetData()[VertexIndex] = Weight;
-				}
-			}
-		}
-	}
-
-	// Pick a root bone.
-	CurrentBone = RefSkeleton.GetBoneName(0);
-	PendingCurrentBone.Reset();
 }
 
-
-void USkinWeightsPaintTool::UpdateEditedSkinWeightsMesh()
+void USkinWeightsPaintTool::OnToolPropertiesModified(UObject* ModifiedObject, FProperty* ModifiedProperty)
 {
-	using namespace UE::AnimationCore;
-	
-	FSkeletalMeshAttributes MeshAttribs(*EditedMesh);
-	FSkinWeightsVertexAttributesRef VertexSkinWeights = MeshAttribs.GetVertexSkinWeights();
-	
-	const FReferenceSkeleton& RefSkeleton = BoneContainer.GetReferenceSkeleton();
-
-	TMap<FBoneIndexType, const TArray<float>*> BoneIndexWeightMap;
-	for (const auto& BoneNameAndWeights : SkinWeightsMap)
+	// bone changed?
+	const bool bCurrentBoneModified = ModifiedProperty->GetNameCPP() == GET_MEMBER_NAME_STRING_CHECKED(FBoneReference, BoneName);
+	if (bCurrentBoneModified)
 	{
-		const int32 BoneIndex = RefSkeleton.FindBoneIndex(BoneNameAndWeights.Key);
-		if (BoneIndex != INDEX_NONE)
-		{
-			BoneIndexWeightMap.Add(static_cast<FBoneIndexType>(BoneIndex), &BoneNameAndWeights.Value);
-		}
+		PendingCurrentBone = ToolProps->CurrentBone.BoneName;
 	}
-
-	FBoneWeightsSettings Settings;
-	Settings.SetNormalizeType(EBoneWeightNormalizeType::AboveOne);
-
-	TArray<FBoneWeight> SourceBoneWeights;
-	SourceBoneWeights.Reserve(MaxInlineBoneWeightCount);
-
-	const int32 NumVertices = EditedMesh->Vertices().Num();
-	for (int32 VertexIndex = 0; VertexIndex < NumVertices; VertexIndex++)
+	
+	// invalidate vertex color cache when weight color properties are modified
+	const bool bColorModeModified = ModifiedProperty->GetNameCPP() == GET_MEMBER_NAME_STRING_CHECKED(USkinWeightsPaintToolProperties, ColorMode);
+	const bool bColorRampModified = ModifiedProperty->GetNameCPP() == GET_MEMBER_NAME_STRING_CHECKED(USkinWeightsPaintToolProperties, ColorRamp);
+	const bool bMinColorModified = ModifiedProperty->GetNameCPP() == GET_MEMBER_NAME_STRING_CHECKED(USkinWeightsPaintToolProperties, MinColor);
+	const bool bMaxColorModified = ModifiedProperty->GetNameCPP() == GET_MEMBER_NAME_STRING_CHECKED(USkinWeightsPaintToolProperties, MaxColor);
+	if (bColorModeModified || bColorRampModified || bMinColorModified || bMaxColorModified)
 	{
-		SourceBoneWeights.Reset();
-
-		for (const auto& BoneIndexWeight : BoneIndexWeightMap)
-		{
-			SourceBoneWeights.Add(FBoneWeight(BoneIndexWeight.Key, (*BoneIndexWeight.Value)[VertexIndex]));
-		}
-
-		VertexSkinWeights.Set(FVertexID(VertexIndex), FBoneWeights::Create(SourceBoneWeights, Settings));
+		bVisibleWeightsValid = false;
 	}
 }
 
diff --git a/Engine/Plugins/Experimental/MeshModelingToolsetExp/Source/MeshModelingToolsExp/Public/SkinWeightsPaintTool.h b/Engine/Plugins/Experimental/MeshModelingToolsetExp/Source/MeshModelingToolsExp/Public/SkinWeightsPaintTool.h
--- a/Engine/Plugins/Experimental/MeshModelingToolsetExp/Source/MeshModelingToolsExp/Public/SkinWeightsPaintTool.h
+++ b/Engine/Plugins/Experimental/MeshModelingToolsetExp/Source/MeshModelingToolsExp/Public/SkinWeightsPaintTool.h
@@ -18,15 +18,104 @@
 struct FMeshDescription;
 class USkinWeightsPaintTool;
 
+// store a map of bone names to weight arrays, where array length == num vertices.
+// this struct stores ALL weights for ALL bones in the entire mesh.
+using PerBoneWeightMap = TMap<FName, TArray<float>>;
 
+// store a sparse set of modifications to a set of vertex weights on a SINGLE bone
+struct FSingleBoneWeightEdits
+{
+	FName BoneName;
+	TMap<int32, float> OldWeights;
+	TMap<int32, float> NewWeights;
+};
+
+// store a sparse set of modifications to a set of vertex weights for a SET of bones
+// with support for merging edits. these are used for transaction history undo/redo.
+struct FMultiBoneWeightEdits
+{
+	void MergeSingleEdit(const FName BoneName, const int32 VertexID, const float OldWeight, const float NewWeight)
+	{
+		PerBoneWeightEdits.FindOrAdd(BoneName);
+		PerBoneWeightEdits[BoneName].BoneName = BoneName;
+		PerBoneWeightEdits[BoneName].NewWeights.Add(VertexID, NewWeight);
+		PerBoneWeightEdits[BoneName].OldWeights.FindOrAdd(VertexID, OldWeight);
+	}
+	
+	void MergeEdits(const FSingleBoneWeightEdits& BoneWeightEdits)
+	{
+		// make sure bone has an entry in the map of weight edits
+		const FName Bone = BoneWeightEdits.BoneName;
+		PerBoneWeightEdits.FindOrAdd(Bone);
+		PerBoneWeightEdits[Bone].BoneName = Bone;
+	
+		for (const TTuple<int32, float>& NewWeight : BoneWeightEdits.NewWeights)
+		{
+			int32 VertexIndex = NewWeight.Key;
+			PerBoneWeightEdits[Bone].NewWeights.Add(VertexIndex, NewWeight.Value);
+			PerBoneWeightEdits[Bone].OldWeights.FindOrAdd(VertexIndex, BoneWeightEdits.OldWeights[VertexIndex]);
+		}
+	}
+
+	float GetVertexDeltaFromEdits(FName BoneName, int32 VertexIndex)
+	{
+		PerBoneWeightEdits.FindOrAdd(BoneName);
+		if (const float* NewVertexWeight = PerBoneWeightEdits[BoneName].NewWeights.Find(VertexIndex))
+		{
+			return *NewVertexWeight - PerBoneWeightEdits[BoneName].OldWeights[VertexIndex];
+		}
+
+		return 0.0f;
+	}
+
+	void ApplyEditsToWeightMap(PerBoneWeightMap& InOutWeightMap)
+	{
+		for (TTuple<FName, FSingleBoneWeightEdits>& BoneWeightEdits : PerBoneWeightEdits)
+		{
+			const FSingleBoneWeightEdits& WeightEdits = BoneWeightEdits.Value;
+			TArray<float>& VertexWeights = InOutWeightMap[WeightEdits.BoneName];
+			for (const TTuple<int32, float>& NewWeight : WeightEdits.NewWeights)
+			{
+				VertexWeights[NewWeight.Key] = NewWeight.Value;
+			}
+		}
+	}
+	
+	TMap<FName, FSingleBoneWeightEdits> PerBoneWeightEdits;
+};
+
+// data required to preview the skinning deformations as you paint
+struct FSkinToolDeformer
+{
+	void Initialize(
+		const USkeletalMeshComponent* SkeletalMeshComponent,
+		const FMeshDescription* Mesh);
+
+	void UpdateVertexDeformation(
+		PerBoneWeightMap& Weights,
+		UPreviewMesh* PreviewMesh,
+		UE::Geometry::TDynamicVerticesOctree3<FDynamicMesh3>& OctreeToUpdate);
+
+	void SetVerticesNeedUpdated(TArray<int32> VertexIndices);
+
+	void SetVertexNeedsUpdated(int32 VertexIndex);
+	
+	// which vertices require updating (partially re-calculated skinning deformation while paintin)
+	TSet<int32> VerticesWithModifiedWeights;
+	// position of all vertices in the reference pose
+	TArray<FVector> RefPoseVertexPositions;
+	// inverted, component space ref pose transform of each bone
+	TArray<FTransform> InvCSRefPoseTransforms;
+	// bone name to bone index
+	TMap<FName, int32> BoneNameToIndexMap;
+	// the skeletal mesh to get the current pose from
+	const USkeletalMeshComponent* Component;
+};
 
 class MESHMODELINGTOOLSEXP_API FMeshSkinWeightsChange : public FToolCommandChange
 {
 public:
-	FMeshSkinWeightsChange(const FName& InBoneName) :
-		FToolCommandChange(),
-		BoneName(InBoneName)
-	{ }
+	FMeshSkinWeightsChange() : FToolCommandChange()	{ }
 
 	virtual FString ToString() const override
 	{
@@ -37,14 +126,51 @@ public:
 
 	void Revert(UObject* Object) override;
 
-	void UpdateValues(const TArray<int32>& Indices, const TArray<float>& OldValues, const TArray<float>& NewValues);
+	void AddBoneWeightEdit(const FSingleBoneWeightEdits& BoneWeightEdit);
 
 private:
-	FName BoneName;
-	TMap<int32, float> OldWeights;
-	TMap<int32, float> NewWeights;
+	FMultiBoneWeightEdits AllWeightEdits;
 };
 
+// intermediate storage of the weight maps for duration of tool
+struct FSkinToolWeights
+{
+	// copy the initial weight values from the skeletal mesh
+	void InitializeSkinWeights(
+		const FReferenceSkeleton& RefSkeleton,
+		FMeshDescription* EditedMesh);
+
+	// get an array of weights, per-vertex for the given bone
+	TArray<float>* GetWeightsForBone(const FName BoneName);
+
+	// applies an edit to a single vertex weight on a single bone, then normalizes the remaining weights while
+	// keeping the edited weight intact (ie, adapts OTHER influences to achieve normalization)
+	void EditVertexWeightAndNormalize(
+		const FName BoneToHoldConstant,
+		const int32 VertexId,
+		const float NewWeightValue,
+		FMultiBoneWeightEdits& WeightEdits);
+
+	void ApplyCurrentWeightsToMeshDescription(
+		const FReferenceSkeleton& RefSkeleton,
+		FMeshDescription* EditedMesh);
+
+	void ResetAfterStroke();
+
+	float SetCurrentFalloffAndGetMaxFalloffThisStroke(int32 VertexID, float CurrentStength);
+	
+	// double-buffer of the entire weight matrix
+	// "Pre" is state of weights at stroke start
+	// "Current" is state of weights during stroke
+	// When stroke is over, Pre weights are synchronized with current.
+	PerBoneWeightMap PreStrokeWeightsMap;
+	PerBoneWeightMap CurrentWeightsMap;
+
+	// record the current maximum amount of falloff applied to each vertex during the current stroke
+	// values range from 0-1, this allows brushes to sweep over the same vertex, and apply only the maximum amount
+	// of modification (add/replace/relax etc) that was encountered for the duration of the stroke.
+	TArray<float> MaxFalloffPerVertexThisStroke;
+};
 
 /**
  *
@@ -58,6 +184,32 @@ public:
 	UMeshSurfacePointTool* CreateNewTool(const FToolBuilderState& SceneState) const override;
 };
 
+// weight color mode
+UENUM()
+enum class EWeightColorMode : uint8
+{
+	Greyscale,
+	ColorRamp,
+};
+
+// brush falloff mode
+UENUM()
+enum class EWeightBrushFalloffMode : uint8
+{
+	Surface,
+	Volume,
+};
+
+// brush behavior mode
+UENUM()
+enum class EBrushBehaviorMode : uint8
+{
+	Add,
+	Replace,
+	Multiply,
+	Relax
+};
+
 
 UCLASS()
 class MESHMODELINGTOOLSEXP_API USkinWeightsPaintToolProperties : 
@@ -65,17 +217,39 @@ class MESHMODELINGTOOLSEXP_API USkinWeightsPaintToolProperties :
 	public IBoneReferenceSkeletonProvider
 {
 	GENERATED_BODY()
+
+	USkinWeightsPaintToolProperties(const FObjectInitializer& ObjectInitializer);
+	
 public:
+
+	UPROPERTY(EditAnywhere, Category = Brush)
+	EBrushBehaviorMode BrushMode;
+	
+	UPROPERTY(EditAnywhere, Category = Brush)
+	EWeightBrushFalloffMode FalloffMode;
+	
 	UPROPERTY(VisibleAnywhere, Category = Skeleton)
 	FBoneReference CurrentBone;
 
+	UPROPERTY(EditAnywhere, Category = Brush)
+	EWeightColorMode ColorMode;
+	
+	UPROPERTY(EditAnywhere, Category = WeightColors)
+	TArray<FLinearColor> ColorRamp;
+
+	UPROPERTY(EditAnywhere, Category = WeightColors)
+	FLinearColor MinColor;
+
+	UPROPERTY(EditAnywhere, Category = WeightColors)
+	FLinearColor MaxColor;
+	
 	// IBoneReferenceSkeletonProvider
 	USkeleton* GetSkeleton(bool& bInvalidSkeletonIsError, const IPropertyHandle* PropertyHandle) override;
 
 	TObjectPtr<USkeletalMesh> SkeletalMesh;
+	bool bColorModeChanged = false;
 };
 
-
 /**
  *
  */
@@ -101,48 +275,67 @@ public:
 	void OnEndDrag(const FRay& Ray) override;
 	bool OnUpdateHover(const FInputDeviceRay& DevicePos) override;
 
+	void ExternalUpdateWeights(const FName &BoneName, const TMap<int32, float>& IndexValues);
+
 protected:
 
 	virtual void ApplyStamp(const FBrushStampData& Stamp);
-
 	void OnShutdown(EToolShutdownType ShutdownType) override;
 	void OnTick(float DeltaTime) override;
 
+	// stamp
 	double CalculateBrushFalloff(double Distance) const;
-	void CalculateVertexROI(const FBrushStampData& Stamp, TArray<int>& VertexROI);
-
-	UPROPERTY()
-	TObjectPtr<USkinWeightsPaintToolProperties> ToolProps;
-
-	TValueWatcher<FBoneReference> CurrentBoneWatcher;
-
-	bool bInRemoveStroke = false;
-	bool bInSmoothStroke = false;
+	void CalculateVertexROI(
+		const FBrushStampData& Stamp,
+		TArray<int>& VertexROI,
+		TArray<float>& VertexSqDistances) const;
+	void EditWeightOfVerticesInStamp(
+		EBrushBehaviorMode EditMode,
+		const TArray<int32>& VerticesInStamp,
+		const TArray<float>& VertexSqDistances,
+		FMultiBoneWeightEdits& AllBoneWeightEditsFromStamp);
+	void RelaxWeightOnVertices(
+		TArray<int32> VerticesInStamp,
+		TArray<float> VertexSqDistances,
+		FMultiBoneWeightEdits& AllBoneWeightEditsFromStamp);
+	bool bInvertStroke = false;
+	bool bSmoothStroke = false;
 	FBrushStampData StartStamp;
 	FBrushStampData LastStamp;
 	bool bStampPending;
-
-	TUniquePtr<FMeshDescription> EditedMesh;
-
-	FBoneContainer BoneContainer;
-
 	UE::Geometry::TDynamicVerticesOctree3<FDynamicMesh3> VerticesOctree;
-	TArray<int> PreviewBrushROI;
 
-	using BoneInfluenceMapType = TMap<FName, TArray<float>>;
-	BoneInfluenceMapType SkinWeightsMap;
-	FName CurrentBone = NAME_None;
-	TOptional<FName> PendingCurrentBone;
+	// tool properties
+	UPROPERTY()
+	TObjectPtr<USkinWeightsPaintToolProperties> ToolProps;
+	void OnToolPropertiesModified(UObject* ModifiedObject, FProperty* ModifiedProperty);
+	
+	// the currently edited mesh description
+	TUniquePtr<FMeshDescription> EditedMesh;
 
+	// storage of vertex weights per bone 
+	FSkinToolWeights Weights;
+	
+	// update deformation when vertex weights are modified
+	FSkinToolDeformer Deformer;
 
+	// vertex colors updated when switching current bone or editing weights
+	void UpdateCurrentBoneVertexColors();
+	FVector4f WeightToColor(float Value);
 	bool bVisibleWeightsValid = false;
 
-	static FVector4f WeightToColor(float Value);
-	
-	void InitializeSkinWeights();
-	void UpdateBoneVisualization();
+	// weight editing transactions
+	void BeginChange();
+	TUniquePtr<FMeshSkinWeightsChange> EndChange();
+	TUniquePtr<FMeshSkinWeightsChange> ActiveChange;
+
+	// which bone are we currently painting?
 	void UpdateCurrentBone(const FName &BoneName);
+	FName CurrentBone = NAME_None;
+	TOptional<FName> PendingCurrentBone;
 
+	// bone rendering
+	// TODO, bring this up to new method, add support for rotating bones here
 	struct FBonePositionInfo
 	{
 		FName BoneName;
@@ -151,19 +344,9 @@ protected:
 		float Radius;
 		TMap<FName, int32> ChildBones;
 	};
+	FBoneContainer BoneContainer;
 	TArray<FBonePositionInfo> BonePositionInfos;
 	float MaxDrawRadius = 0.0f;
-
 	void UpdateBonePositionInfos(float MinRadius);
 	void RenderBonePositions(FPrimitiveDrawInterface *PDI);
-
-	TUniquePtr<FMeshSkinWeightsChange> ActiveChange;
-
-	void BeginChange();
-	TUniquePtr<FMeshSkinWeightsChange> EndChange();
-
-	friend class FMeshSkinWeightsChange;
-	void ExternalUpdateValues(const FName &BoneName, const TMap<int32, float>& IndexValues);
-
-	void UpdateEditedSkinWeightsMesh();
 };
diff --git a/Engine/Plugins/Experimental/MeshModelingToolsetExp/Source/ModelingEditorUI/ModelingEditorUI.build.cs b/Engine/Plugins/Experimental/MeshModelingToolsetExp/Source/ModelingEditorUI/ModelingEditorUI.build.cs
--- a/Engine/Plugins/Experimental/MeshModelingToolsetExp/Source/ModelingEditorUI/ModelingEditorUI.build.cs
+++ b/Engine/Plugins/Experimental/MeshModelingToolsetExp/Source/ModelingEditorUI/ModelingEditorUI.build.cs
@@ -38,7 +38,8 @@ public class ModelingEditorUI : ModuleRules
                 "Projects",
 				"ToolWidgets",
 				"EditorWidgets",
-				"DeveloperSettings"
+				"DeveloperSettings", 
+				"MeshModelingToolsExp"
 			}
 			);
 
diff --git a/Engine/Plugins/Experimental/MeshModelingToolsetExp/Source/ModelingEditorUI/Private/SkinWeightDetailCustomization.cpp b/Engine/Plugins/Experimental/MeshModelingToolsetExp/Source/ModelingEditorUI/Private/SkinWeightDetailCustomization.cpp
new file mode 100644
--- /dev/null
+++ b/Engine/Plugins/Experimental/MeshModelingToolsetExp/Source/ModelingEditorUI/Private/SkinWeightDetailCustomization.cpp
@@ -0,0 +1,148 @@
+﻿// Copyright Epic Games, Inc. All Rights Reserved.
+
+#include "SkinWeightDetailCustomization.h"
+
+#include "DetailCategoryBuilder.h"
+#include "DetailLayoutBuilder.h"
+#include "DetailWidgetRow.h"
+#include "Widgets/Input/SSegmentedControl.h"
+#include "SkinWeightsPaintTool.h"
+#include "SSkinWeightProfileImportOptions.h"
+
+#define LOCTEXT_NAMESPACE "SkinWeightToolSettingsEditor"
+
+void FSkinWeightDetailCustomization::CustomizeDetails(IDetailLayoutBuilder& DetailBuilder)
+{
+	TArray<TWeakObjectPtr<UObject>> DetailObjects;
+	DetailBuilder.GetObjectsBeingCustomized(DetailObjects);
+
+	// should be impossible to get multiple settings objects for a single tool
+	ensure(DetailObjects.Num()==1);
+	SkinToolSettings = Cast<USkinWeightsPaintToolProperties>(DetailObjects[0]);
+	
+	// custom display of falloff mode as segmented toggle buttons
+	IDetailCategoryBuilder& BrushCategory = DetailBuilder.EditCategory("Brush", FText::GetEmpty(), ECategoryPriority::Important);
+
+	// add segmented control toggle for brush behavior modes ("Add", "Replace", etc..)
+	BrushCategory.AddCustomRow(LOCTEXT("BrushModeCategory", "Brush Mode"), false)
+		.NameContent()
+		[
+			SNew(STextBlock)
+			.Text(LOCTEXT("BrushModeLabel", "Brush Mode"))
+			.Font(FAppStyle::GetFontStyle(TEXT("PropertyWindow.NormalFont")))
+			.ToolTipText(LOCTEXT("BrushModeTooltip", "Determines how the weights are affected by the brush."))
+		]
+		.ValueContent()
+		[
+			SNew(SHorizontalBox)
+			+ SHorizontalBox::Slot()
+			[
+				SNew(SBox)
+				.Padding(2.0f)
+				.HAlign(HAlign_Center)
+				[
+					SNew(SSegmentedControl<EBrushBehaviorMode>)
+					.Value_Lambda([this]()
+					{
+						return SkinToolSettings->BrushMode;
+					})
+					.OnValueChanged_Lambda([this](EBrushBehaviorMode Mode)
+					{
+						SkinToolSettings->BrushMode = Mode;
+					})
+					+SSegmentedControl<EBrushBehaviorMode>::Slot(EBrushBehaviorMode::Add)
+					.Text(LOCTEXT("AddMode", "Add"))
+					+ SSegmentedControl<EBrushBehaviorMode>::Slot(EBrushBehaviorMode::Replace)
+					.Text(LOCTEXT("ReplaceMode", "Replace"))
+					+ SSegmentedControl<EBrushBehaviorMode>::Slot(EBrushBehaviorMode::Multiply)
+					.Text(LOCTEXT("MultiplyMode", "Multiply"))
+					+ SSegmentedControl<EBrushBehaviorMode>::Slot(EBrushBehaviorMode::Relax)
+					.Text(LOCTEXT("RelaxMode", "Relax"))
+				]
+			]
+		];
+	const TSharedRef<IPropertyHandle> BrushPropHandle = DetailBuilder.GetProperty(GET_MEMBER_NAME_CHECKED(USkinWeightsPaintToolProperties, BrushMode));
+	DetailBuilder.HideProperty(BrushPropHandle);
+
+	// add segmented control toggle for brush falloff modes ("Surface" or "Volume")
+	BrushCategory.AddCustomRow(LOCTEXT("BrushFalloffCategory", "Brush Falloff"), false)
+		.NameContent()
+		[
+			SNew(STextBlock)
+			.Text(LOCTEXT("BrushFalloffLabel", "Falloff Mode"))
+			.Font(FAppStyle::GetFontStyle(TEXT("PropertyWindow.NormalFont")))
+			.ToolTipText(LOCTEXT("BrushFalloffTooltip", "Falloff brush influence along the surface of the mesh, or using straight-line distance (volume)."))
+		]
+		.ValueContent()
+		[
+			SNew(SHorizontalBox)
+			+ SHorizontalBox::Slot()
+			[
+				SNew(SBox)
+				.Padding(2.0f)
+				.HAlign(HAlign_Center)
+				[
+					SNew(SSegmentedControl<EWeightBrushFalloffMode>)
+					.Value_Lambda([this]()
+					{
+						return SkinToolSettings->FalloffMode;
+					})
+					.OnValueChanged_Lambda([this](EWeightBrushFalloffMode Mode)
+					{
+						SkinToolSettings->FalloffMode = Mode;
+						SkinToolSettings->bColorModeChanged = true;
+					})
+					+SSegmentedControl<EWeightBrushFalloffMode>::Slot(EWeightBrushFalloffMode::Surface)
+					.Text(LOCTEXT("SurfaceMode", "Surface"))
+					+ SSegmentedControl<EWeightBrushFalloffMode>::Slot(EWeightBrushFalloffMode::Volume)
+					.Text(LOCTEXT("VolumeMode", "Volume"))
+				]
+			]
+		];
+	const TSharedRef<IPropertyHandle> FalloffPropHandle = DetailBuilder.GetProperty(GET_MEMBER_NAME_CHECKED(USkinWeightsPaintToolProperties, FalloffMode));
+	DetailBuilder.HideProperty(FalloffPropHandle);
+
+	// custom display of falloff mode as segmented toggle buttons
+	IDetailCategoryBuilder& WeightColorsCategory = DetailBuilder.EditCategory("WeightColors", FText::GetEmpty(), ECategoryPriority::Important);
+	WeightColorsCategory.InitiallyCollapsed(true);
+
+	// add segmented control toggle for brush behavior modes ("Add", "Replace", etc..)
+	WeightColorsCategory.AddCustomRow(LOCTEXT("ColorModeCategory", "Color Mode"), false)
+		.NameContent()
+		[
+			SNew(STextBlock)
+			.Text(LOCTEXT("ColorModeLabel", "Color Mode"))
+			.Font(FAppStyle::GetFontStyle(TEXT("PropertyWindow.NormalFont")))
+			.ToolTipText(LOCTEXT("ColorModeTooltip", "Determines how the weight colors are displayed."))
+		]
+		.ValueContent()
+		[
+			SNew(SHorizontalBox)
+			+ SHorizontalBox::Slot()
+			[
+				SNew(SBox)
+				.Padding(2.0f)
+				.HAlign(HAlign_Center)
+				[
+					SNew(SSegmentedControl<EWeightColorMode>)
+					.Value_Lambda([this]()
+					{
+						return SkinToolSettings->ColorMode;
+					})
+					.OnValueChanged_Lambda([this](EWeightColorMode Mode)
+					{
+						SkinToolSettings->ColorMode = Mode;
+						SkinToolSettings->bColorModeChanged = true;
+					})
+					+ SSegmentedControl<EWeightColorMode>::Slot(EWeightColorMode::Greyscale)
+					.Text(LOCTEXT("GreyscaleMode", "Greyscale"))
+					+SSegmentedControl<EWeightColorMode>::Slot(EWeightColorMode::ColorRamp)
+					.Text(LOCTEXT("RampMode", "Color Ramp"))
+				]
+			]
+		];
+	const TSharedRef<IPropertyHandle> ColorModePropHandle = DetailBuilder.GetProperty(GET_MEMBER_NAME_CHECKED(USkinWeightsPaintToolProperties, ColorMode));
+	DetailBuilder.HideProperty(ColorModePropHandle);
+}
+
+#undef LOCTEXT_NAMESPACE
\ No newline at end of file
diff --git a/Engine/Plugins/Experimental/MeshModelingToolsetExp/Source/ModelingEditorUI/Public/SkinWeightDetailCustomization.h b/Engine/Plugins/Experimental/MeshModelingToolsetExp/Source/ModelingEditorUI/Public/SkinWeightDetailCustomization.h
new file mode 100644
--- /dev/null
+++ b/Engine/Plugins/Experimental/MeshModelingToolsetExp/Source/ModelingEditorUI/Public/SkinWeightDetailCustomization.h
@@ -0,0 +1,24 @@
+﻿// Copyright Epic Games, Inc. All Rights Reserved.
+
+#pragma once
+
+#include "CoreMinimal.h"
+#include "IDetailCustomization.h"
+
+class USkinWeightsPaintToolProperties;
+
+class FSkinWeightDetailCustomization : public IDetailCustomization
+{
+public:
+
+	static TSharedRef<IDetailCustomization> MakeInstance()
+	{
+		return MakeShareable(new FSkinWeightDetailCustomization);
+	}
+
+	/** IDetailCustomization interface */
+	virtual void CustomizeDetails(IDetailLayoutBuilder& DetailBuilder) override;
+
+private:
+	TWeakObjectPtr<USkinWeightsPaintToolProperties> SkinToolSettings;
+};
