Non-unity fixes
#rnx

diff --git a/Engine/Plugins/Experimental/MeshModelingToolsetExp/Source/ModelingEditorUI/Public/SkinWeightDetailCustomization.h b/Engine/Plugins/Experimental/MeshModelingToolsetExp/Source/ModelingEditorUI/Public/SkinWeightDetailCustomization.h
--- a/Engine/Plugins/Experimental/MeshModelingToolsetExp/Source/ModelingEditorUI/Public/SkinWeightDetailCustomization.h
+++ b/Engine/Plugins/Experimental/MeshModelingToolsetExp/Source/ModelingEditorUI/Public/SkinWeightDetailCustomization.h
@@ -4,6 +4,7 @@
 
 #include "CoreMinimal.h"
 #include "IDetailCustomization.h"
+#include "UObject/WeakObjectPtr.h"
 
 class USkinWeightsPaintToolProperties;
 
