Add initial version of new CombineMeshInstances geometry processing Modular Feature and initial implementation. This modular feature is used to construct a 'merged' mesh for a large number of input mesh instances, with the idea that simplifying or abstracting each instanced base mesh (ie unique "part") before appending will produce a better (ie more accurate + lower polycount) combined mesh than trying to combine first and then do standard simplification technqiues on the merged mesh. Currently only per-part simplification and fitting oriented boxes is supported. This is a functional checkpoint of work in progress.

diff --git a/Engine/Plugins/Experimental/MeshModelingToolsetExp/Source/GeometryProcessingAdapters/Private/CombineMeshInstancesImpl.cpp b/Engine/Plugins/Experimental/MeshModelingToolsetExp/Source/GeometryProcessingAdapters/Private/CombineMeshInstancesImpl.cpp
new file mode 100644
--- /dev/null
+++ b/Engine/Plugins/Experimental/MeshModelingToolsetExp/Source/GeometryProcessingAdapters/Private/CombineMeshInstancesImpl.cpp
@@ -0,0 +1,630 @@
+// Copyright Epic Games, Inc. All Rights Reserved.
+
+#include "GeometryProcessing/CombineMeshInstancesImpl.h"
+
+#include "Async/ParallelFor.h"
+
+#include "Engine/World.h"
+#include "Engine/StaticMesh.h"
+#include "Components/PrimitiveComponent.h"
+#include "Components/StaticMeshComponent.h"
+#include "Components/InstancedStaticMeshComponent.h"
+#include "PhysicsEngine/BodySetup.h"
+
+#include "MeshDescriptionToDynamicMesh.h"
+#include "DynamicMeshToMeshDescription.h"
+
+#include "DynamicMeshEditor.h"
+#include "DynamicMesh/DynamicMeshAttributeSet.h"
+#include "DynamicMesh/MeshNormals.h"
+
+#include "ShapeApproximation/ShapeDetection3.h"
+#include "ShapeApproximation/MeshSimpleShapeApproximation.h"
+#include "Generators/GridBoxMeshGenerator.h"
+
+
+#include "MeshSimplification.h"
+#include "DynamicMesh/ColliderMesh.h"
+#include "MeshConstraintsUtil.h"
+#include "DynamicMesh/Operations/MergeCoincidentMeshEdges.h"
+
+
+
+using namespace UE::Geometry;
+
+
+
+enum class EMeshDetailLevel
+{
+	Base = 0,
+	Standard = 1,
+	Small = 2,
+	Decorative = 3
+};
+
+
+struct FMeshInstance
+{
+	UE::Geometry::FTransformSequence3d WorldTransform;
+	TArray<UMaterialInterface*> Materials;
+
+	UPrimitiveComponent* SourceComponent;
+	int32 SourceInstanceIndex;
+
+	EMeshDetailLevel DetailLevel = EMeshDetailLevel::Standard;
+
+	// allow FMeshInstance to maintain link to external representation of instance
+	FIndex3i ExternalInstanceIndex = FIndex3i::Invalid();
+};
+
+
+struct FMeshInstanceSet
+{
+	UStaticMesh* SourceAsset;
+	TArray<FMeshInstance> Instances;
+};
+
+
+struct FSourceGeometry
+{
+	TArray<UE::Geometry::FDynamicMesh3> SourceMeshLODs;
+	//UE::Geometry::FSimpleShapeSet3d CollisionShapes;
+};
+
+struct FOptimizedGeometry
+{
+	TArray<UE::Geometry::FDynamicMesh3> SimplifiedMeshLODs;
+
+	UE::Geometry::FDynamicMesh3 OptimizedMesh;
+	//UE::Geometry::FSimpleShapeSet3d CollisionShapes;
+};
+
+
+class FMeshInstanceAssembly
+{
+public:
+	// this is necessary due to TArray<TUniquePtr> below
+	FMeshInstanceAssembly() = default;
+	FMeshInstanceAssembly(FMeshInstanceAssembly&) = delete;
+	FMeshInstanceAssembly& operator=(const FMeshInstanceAssembly&) = delete;
+
+
+	TArray<TUniquePtr<FMeshInstanceSet>> InstanceSets;
+
+	TArray<UMaterialInterface*> UniqueMaterials;
+	TMap<UMaterialInterface*, int32> MaterialMap;
+
+	TArray<FSourceGeometry> SourceMeshGeometry;
+	TArray<FOptimizedGeometry> OptimizedMeshGeometry;
+
+
+	// allow external code to preprocess dynamic mesh for a specific instance
+	TFunction<void(FDynamicMesh3&, const FMeshInstance&)> PreProcessInstanceMeshFunc;
+
+};
+
+
+
+
+void InitializeMeshInstanceAssembly(
+	const IGeometryProcessing_CombineMeshInstances::FInstanceSet& SourceInstanceSet,
+	FMeshInstanceAssembly& AssemblyOut)
+{
+	TMap<UStaticMesh*, FMeshInstanceSet*> MeshToInstanceMap;
+
+	int32 NumInstances = SourceInstanceSet.StaticMeshInstances.Num();
+	for ( int32 Index = 0; Index < NumInstances; ++Index)
+	{
+		const IGeometryProcessing_CombineMeshInstances::FStaticMeshInstance& SourceMeshInstance = SourceInstanceSet.StaticMeshInstances[Index];
+
+		UStaticMesh* StaticMesh = SourceMeshInstance.SourceMesh;
+		FMeshInstanceSet** FoundInstanceSet = MeshToInstanceMap.Find(StaticMesh);
+		if (FoundInstanceSet == nullptr)
+		{
+			TUniquePtr<FMeshInstanceSet> NewInstanceSet = MakeUnique<FMeshInstanceSet>();
+			NewInstanceSet->SourceAsset = StaticMesh;
+			FMeshInstanceSet* Ptr = NewInstanceSet.Get();
+					
+			AssemblyOut.InstanceSets.Add(MoveTemp(NewInstanceSet));
+			// store source model?
+
+			MeshToInstanceMap.Add(StaticMesh, Ptr);
+			FoundInstanceSet = &Ptr;
+		}
+
+		FMeshInstance NewInstance;
+		NewInstance.ExternalInstanceIndex = FIndex3i(Index, -1,-1);
+
+		if ( SourceMeshInstance.GroupDataIndex >= 0 && SourceMeshInstance.GroupDataIndex < SourceInstanceSet.InstanceGroupDatas.Num() )
+		{
+			const IGeometryProcessing_CombineMeshInstances::FMeshInstanceGroupData& GroupData = 
+				SourceInstanceSet.InstanceGroupDatas[SourceMeshInstance.GroupDataIndex];
+			NewInstance.Materials = GroupData.MaterialSet;
+		}
+
+		NewInstance.SourceComponent = SourceMeshInstance.SourceComponent;
+		NewInstance.SourceInstanceIndex = SourceMeshInstance.SourceInstanceIndex;
+		NewInstance.DetailLevel = static_cast<EMeshDetailLevel>( static_cast<int32>(SourceMeshInstance.DetailLevel) );
+		for ( FTransform3d Transform : SourceMeshInstance.TransformSequence )
+		{
+			NewInstance.WorldTransform.Append( Transform );
+		}
+		(*FoundInstanceSet)->Instances.Add(NewInstance);
+	}
+
+
+	// collect unique materials
+	for (TPair<UStaticMesh*, FMeshInstanceSet*>& Pair : MeshToInstanceMap)
+	{
+		UStaticMesh* StaticMesh = Pair.Key;
+		FMeshInstanceSet& InstanceSet = *(Pair.Value);
+
+		for (FMeshInstance& Instance : InstanceSet.Instances)
+		{
+			for (UMaterialInterface* Material : Instance.Materials)
+			{
+				if ( AssemblyOut.MaterialMap.Contains(Material) == false)
+				{
+					int32 NewIndex = AssemblyOut.UniqueMaterials.Num();
+					AssemblyOut.UniqueMaterials.Add(Material);
+					AssemblyOut.MaterialMap.Add(Material, NewIndex);
+				}
+			}
+		}
+	}
+
+}
+
+
+
+
+void InitializeAssemblySourceMeshesFromLOD(
+	FMeshInstanceAssembly& Assembly,
+	int32 SourceAssetBaseLOD,
+	int32 NumSourceLODs)
+{
+	using namespace UE::Geometry;
+
+	check(NumSourceLODs > 0);
+
+	int32 NumSets = Assembly.InstanceSets.Num();
+	Assembly.SourceMeshGeometry.SetNum(NumSets);
+
+	// collect mesh for each assembly item
+	ParallelFor(NumSets, [&](int32 Index)
+	{
+		TUniquePtr<FMeshInstanceSet>& InstanceSet = Assembly.InstanceSets[Index];
+		FSourceGeometry& Target = Assembly.SourceMeshGeometry[Index];
+		Target.SourceMeshLODs.SetNum(NumSourceLODs);
+
+		UStaticMesh* StaticMesh = InstanceSet->SourceAsset;
+
+		for (int32 k = 0; k < NumSourceLODs; ++k)
+		{
+			int32 LODIndex = SourceAssetBaseLOD + k;
+			if (LODIndex < StaticMesh->GetNumSourceModels())
+			{
+				FMeshDescription* UseMeshDescription = StaticMesh->GetMeshDescription(LODIndex);
+				if (UseMeshDescription != nullptr)
+				{
+					FMeshDescriptionToDynamicMesh Converter;
+					Converter.bEnableOutputGroups = true; 
+					Converter.bTransformVertexColorsLinearToSRGB = true;
+					Converter.Convert(UseMeshDescription, Target.SourceMeshLODs[k]);
+				}
+			}
+		}
+
+		// if first LOD is missing try getting LOD0 again
+		if (Target.SourceMeshLODs[0].TriangleCount() == 0)
+		{
+			if (FMeshDescription* UseMeshDescription = StaticMesh->GetMeshDescription(0))
+			{
+				FMeshDescriptionToDynamicMesh Converter;
+				Converter.bEnableOutputGroups = true; 
+				Converter.bTransformVertexColorsLinearToSRGB = true;
+				Converter.Convert(UseMeshDescription, Target.SourceMeshLODs[0]);
+			}
+		}
+
+		// now if first LOD is missing, just fall back to a box
+		if (Target.SourceMeshLODs[0].TriangleCount() == 0)
+		{
+			FGridBoxMeshGenerator BoxGen;
+			Target.SourceMeshLODs[0].Copy(&BoxGen.Generate());
+		}
+
+		// now make sure every one of our Source LODs has a mesh by copying from N-1
+		for (int32 k = 1; k < NumSourceLODs; ++k)
+		{
+			if (Target.SourceMeshLODs[k].TriangleCount() == 0)
+			{
+				Target.SourceMeshLODs[k] = Target.SourceMeshLODs[k-1];
+			}
+		}
+	});
+
+
+	//ParallelFor(NumSets, [&](int32 Index)
+	//{
+	//	TUniquePtr<FMeshInstanceSet>& InstanceSet = Assembly.InstanceSets[Index];
+	//	UStaticMesh* StaticMesh = InstanceSet->SourceAsset;
+
+	//	FSourceGeometry& Target = Assembly.SourceMeshGeometry[Index];
+
+	//	UBodySetup* BodySetup = StaticMesh->GetBodySetup();
+	//	if (BodySetup)
+	//	{
+	//		UE::Geometry::GetShapeSet(BodySetup->AggGeom, Target.CollisionShapes);
+	//	}
+
+	//	// todo:
+	//	//if (bDetectBoxes)
+	//	//{
+	//	//	DetectBoxesFromConvexes(CollisionShapes);
+	//	//}
+
+	//});
+
+
+}
+
+
+
+
+
+
+static void SimplifyPartMesh(
+	FDynamicMesh3& EditMesh, 
+	double Tolerance, 
+	bool bNoSplitAttributes,
+	double RecomputeNormalsAngleThreshold)
+{
+	// weld edges in case input was unwelded...
+	FMergeCoincidentMeshEdges Welder(&EditMesh);
+	Welder.MergeVertexTolerance = Tolerance * 0.001;
+	Welder.OnlyUniquePairs = false;
+	Welder.Apply();
+
+
+	FVolPresMeshSimplification Simplifier(&EditMesh);
+
+	// clear out attributes so it doesn't affect simplification
+	//EditMesh.DiscardAttributes();
+	EditMesh.Attributes()->SetNumUVLayers(0);
+	EditMesh.Attributes()->DisableTangents();
+	EditMesh.Attributes()->DisablePrimaryColors();
+	FMeshNormals::InitializeOverlayToPerVertexNormals(EditMesh.Attributes()->PrimaryNormals(), false);
+	bNoSplitAttributes = true;
+
+	Simplifier.ProjectionMode = FVolPresMeshSimplification::ETargetProjectionMode::NoProjection;
+
+	FColliderMesh ColliderMesh;
+	ColliderMesh.Initialize(EditMesh);
+	FColliderMeshProjectionTarget ProjectionTarget(&ColliderMesh);
+	Simplifier.SetProjectionTarget(&ProjectionTarget);
+
+	Simplifier.DEBUG_CHECK_LEVEL = 0;
+	Simplifier.bRetainQuadricMemory = false; 
+	if ( bNoSplitAttributes == false )
+	{
+		Simplifier.bAllowSeamCollapse = true;
+		Simplifier.SetEdgeFlipTolerance(1.e-5);
+		if (EditMesh.HasAttributes())
+		{
+			EditMesh.Attributes()->SplitAllBowties();	// eliminate any bowties that might have formed on attribute seams.
+		}
+	}
+
+	// this should preserve part shape better but it completely fails currently =\
+	//Simplifier.CollapseMode = FVolPresMeshSimplification::ESimplificationCollapseModes::MinimalExistingVertexError;
+
+	// do these flags matter here since we are not flipping??
+	EEdgeRefineFlags MeshBoundaryConstraints = EEdgeRefineFlags::NoFlip;
+	EEdgeRefineFlags GroupBorderConstraints = EEdgeRefineFlags::NoConstraint;
+	EEdgeRefineFlags MaterialBorderConstraints = EEdgeRefineFlags::NoConstraint;
+
+	FMeshConstraints Constraints;
+	FMeshConstraintsUtil::ConstrainAllBoundariesAndSeams(Constraints, EditMesh,
+		MeshBoundaryConstraints, GroupBorderConstraints, MaterialBorderConstraints, true, false, true);
+	Simplifier.SetExternalConstraints(MoveTemp(Constraints));
+
+	Simplifier.GeometricErrorConstraint = FVolPresMeshSimplification::EGeometricErrorCriteria::PredictedPointToProjectionTarget;
+	Simplifier.GeometricErrorTolerance = Tolerance;
+
+	Simplifier.SimplifyToTriangleCount( 1 );
+
+	// compact result
+	EditMesh.CompactInPlace();
+
+	// recompute normals
+	FMeshNormals::InitializeOverlayTopologyFromOpeningAngle(&EditMesh, EditMesh.Attributes()->PrimaryNormals(), RecomputeNormalsAngleThreshold);
+	FMeshNormals::QuickRecomputeOverlayNormals(EditMesh);
+}
+
+
+
+
+void ComputeMeshApproximations(
+	IGeometryProcessing_CombineMeshInstances::FOptions CombineOptions,
+	FMeshInstanceAssembly& Assembly)
+{
+	using namespace UE::Geometry;
+	const double AngleThreshold = 45.0;
+
+	int32 NumSets = Assembly.InstanceSets.Num();
+	Assembly.OptimizedMeshGeometry.SetNum(NumSets);
+
+	int32 NumSimplifiedLODs = CombineOptions.NumSimplifiedLODs;
+
+	ParallelFor(NumSets, [&](int32 Index)
+	{
+		TUniquePtr<FMeshInstanceSet>& InstanceSet = Assembly.InstanceSets[Index];
+		const FSourceGeometry& SourceGeo = Assembly.SourceMeshGeometry[Index];
+		const FDynamicMesh3& InitialSourceMesh = (CombineOptions.ApproximationSourceLOD < SourceGeo.SourceMeshLODs.Num()) ?
+			SourceGeo.SourceMeshLODs[CombineOptions.ApproximationSourceLOD] : SourceGeo.SourceMeshLODs.Last();
+		FOptimizedGeometry& ApproxGeo = Assembly.OptimizedMeshGeometry[Index];
+
+		// compute simplified part LODs
+		ApproxGeo.SimplifiedMeshLODs.SetNum(NumSimplifiedLODs);
+		double InitialTolerance = CombineOptions.SimplifyBaseTolerance;
+		for (int32 k = 0; k < NumSimplifiedLODs; ++k)
+		{
+			ApproxGeo.SimplifiedMeshLODs[k] = InitialSourceMesh;
+			SimplifyPartMesh(ApproxGeo.SimplifiedMeshLODs[k], InitialTolerance, false, AngleThreshold);
+			InitialTolerance *= CombineOptions.SimplifyLODLevelToleranceScale;
+		}
+
+		// compute shape approximation
+		FMeshSimpleShapeApproximation ShapeApprox;
+		ShapeApprox.InitializeSourceMeshes( {&InitialSourceMesh} );
+		ShapeApprox.bDetectBoxes = ShapeApprox.bDetectCapsules = ShapeApprox.bDetectConvexes = ShapeApprox.bDetectSpheres = false;
+
+		FSimpleShapeSet3d ResultBoxes;
+		ShapeApprox.Generate_OrientedBoxes(ResultBoxes);
+		FGridBoxMeshGenerator BoxGen;
+		BoxGen.Box = ResultBoxes.Boxes[0].Box;
+		BoxGen.EdgeVertices = {0,0,0};
+		ApproxGeo.OptimizedMesh.Copy(&BoxGen.Generate());
+
+		ApproxGeo.OptimizedMesh.EnableMatchingAttributes(InitialSourceMesh);
+		FMeshNormals::InitializeOverlayTopologyFromOpeningAngle(&ApproxGeo.OptimizedMesh, ApproxGeo.OptimizedMesh.Attributes()->PrimaryNormals(), AngleThreshold);
+		FMeshNormals::QuickRecomputeOverlayNormals(ApproxGeo.OptimizedMesh);
+
+	});
+
+}
+
+
+
+namespace UE::Geometry
+{
+
+struct FCombinedMeshLOD
+{
+	FDynamicMesh3 Mesh;
+	FDynamicMeshEditor Editor;
+	FDynamicMeshMaterialAttribute* MaterialIDs;
+
+	FCombinedMeshLOD()
+		: Editor(&Mesh)
+	{
+		Mesh.EnableAttributes();
+		Mesh.Attributes()->EnableMaterialID();
+		
+		// should we do this? maybe should be done via enable-matching?
+		Mesh.Attributes()->EnablePrimaryColors();
+
+		MaterialIDs = Mesh.Attributes()->GetMaterialID();
+	}
+
+};
+
+}
+
+
+
+// change this to build a single LOD, and separate versions for (eg) source mesh vs approx mesh
+// should we even bother w/ storing approx meshes? just generate them as needed?
+
+void BuildCombinedMesh(
+	const FMeshInstanceAssembly& Assembly,
+	IGeometryProcessing_CombineMeshInstances::FOptions CombineOptions,
+	TArray<FDynamicMesh3>& CombinedMeshLODs)
+{
+	using namespace UE::Geometry;
+
+	int32 NumLODs = CombineOptions.NumLODs;
+	TArray<FCombinedMeshLOD> MeshLODs;
+	MeshLODs.SetNum(NumLODs);
+
+	//CombinedLOD0.Attributes()->SetNumPolygroupLayers(2);
+	//FDynamicMeshPolygroupAttribute* PartIDAttrib = AccumMesh.Attributes()->GetPolygroupLayer(0);
+	//FDynamicMeshPolygroupAttribute* PartInstanceMapAttrib = AccumMesh.Attributes()->GetPolygroupLayer(1);
+
+	int32 NumSets = Assembly.InstanceSets.Num();
+
+	//for ( int32 SetIndex = 0; SetIndex < NumSets; ++SetIndex )
+	//{
+	//	CombinedLOD0.EnableMatchingAttributes( Assembly.SourceMeshGeometry[Index].OriginalMesh, false, false );
+	//}
+
+	for ( int32 SetIndex = 0; SetIndex < NumSets; ++SetIndex )
+	{
+		const TUniquePtr<FMeshInstanceSet>& InstanceSet = Assembly.InstanceSets[SetIndex];
+		const FSourceGeometry& SourceGeometry = Assembly.SourceMeshGeometry[SetIndex];
+		const FOptimizedGeometry& OptimizedGeometry = Assembly.OptimizedMeshGeometry[SetIndex];
+		UStaticMesh* StaticMesh = InstanceSet->SourceAsset;
+
+		FMeshIndexMappings Mappings;
+
+		for (int32 LODLevel = 0; LODLevel < NumLODs; ++LODLevel)
+		{
+			const FDynamicMesh3* SourceAppendMesh = nullptr;
+			const FDynamicMesh3* ApproximateAppendMesh = nullptr;
+			const FDynamicMesh3* UseAppendMesh = nullptr;
+
+			ApproximateAppendMesh = &OptimizedGeometry.OptimizedMesh;
+
+			if (LODLevel < CombineOptions.NumCopiedLODs)
+			{
+				SourceAppendMesh = (LODLevel < SourceGeometry.SourceMeshLODs.Num()) ? 
+					&SourceGeometry.SourceMeshLODs[LODLevel] : &SourceGeometry.SourceMeshLODs.Last();
+				UseAppendMesh = SourceAppendMesh;
+			}
+			else if (LODLevel < CombineOptions.NumCopiedLODs + CombineOptions.NumSimplifiedLODs)
+			{
+				int32 SimplifiedLODIndex = LODLevel - CombineOptions.NumCopiedLODs;
+				SourceAppendMesh = &OptimizedGeometry.SimplifiedMeshLODs[SimplifiedLODIndex];
+				UseAppendMesh = SourceAppendMesh;
+			}
+			else
+			{
+				UseAppendMesh = ApproximateAppendMesh;
+			}
+
+			FCombinedMeshLOD& CombinedMeshLODData = MeshLODs[LODLevel];
+
+			for ( const FMeshInstance& Instance : InstanceSet->Instances )
+			{
+				bool bIsDecorativePart = (Instance.DetailLevel == EMeshDetailLevel::Decorative);
+
+				// filter out detail parts at higher LODs
+				if (bIsDecorativePart && LODLevel > CombineOptions.FilterDecorativePartsLODLevel)
+				{
+					continue;
+				}
+				// at last detail part LOD, switch to approximate mesh
+				if (bIsDecorativePart && LODLevel == CombineOptions.FilterDecorativePartsLODLevel)
+				{
+					UseAppendMesh = ApproximateAppendMesh;
+				}
+
+				FDynamicMesh3 TempAppendMesh(*UseAppendMesh);
+				if (Assembly.PreProcessInstanceMeshFunc)
+				{
+					Assembly.PreProcessInstanceMeshFunc(TempAppendMesh, Instance);
+				}
+
+				Mappings.Reset();
+				CombinedMeshLODData.Editor.AppendMesh(&TempAppendMesh, Mappings,
+					[&](int, const FVector3d& Pos) { return Instance.WorldTransform.TransformPosition(Pos); },
+					[&](int, const FVector3d& Normal) { return Instance.WorldTransform.TransformNormal(Normal); });
+
+				//CollisionAssembly.AppendInstance(InstanceSet.Collision, Instance.WorldTransform);
+
+				// append part ID stuff here
+
+				// could precompute these indexes for each instance?
+				// also for source mesh we could transfer material IDs correctly...
+				UMaterialInterface* UseMaterial = Instance.Materials[0];
+				const int32* FoundMaterialIndex = Assembly.MaterialMap.Find(UseMaterial);
+				int32 AssignMaterialIndex = (FoundMaterialIndex != nullptr) ? *FoundMaterialIndex : 0;
+
+				for (int32 tid : TempAppendMesh.TriangleIndicesItr())
+				{
+					CombinedMeshLODData.MaterialIDs->SetValue( Mappings.GetNewTriangle(tid), AssignMaterialIndex );
+				}
+			}
+
+		}
+	}
+
+
+	for (int32 LODLevel = 0; LODLevel < NumLODs; ++LODLevel)
+	{
+		CombinedMeshLODs.Add(MoveTemp(MeshLODs[LODLevel].Mesh));
+	}
+
+}
+
+
+
+
+
+
+
+
+
+
+
+
+IGeometryProcessing_CombineMeshInstances::FOptions FCombineMeshInstancesImpl::ConstructDefaultOptions()
+{
+	//
+	// Construct options for ApproximateActors operation
+	//
+	FOptions Options;
+
+	Options.NumLODs = 5;
+
+	Options.NumCopiedLODs = 1;
+
+	Options.NumSimplifiedLODs = 3;
+	Options.SimplifyBaseTolerance = 0.25;
+	Options.SimplifyLODLevelToleranceScale = 2.0;
+
+	//// LOD level to filter out detail parts
+	Options.FilterDecorativePartsLODLevel = 2;
+
+
+	return Options;
+}
+
+
+
+static void SetConstantVertexColor(FDynamicMesh3& Mesh, FLinearColor LinearColor)
+{
+	if (Mesh.HasAttributes() == false)
+	{
+		Mesh.EnableAttributes();
+	}
+	if (Mesh.Attributes()->HasPrimaryColors() == false)
+	{
+		Mesh.Attributes()->EnablePrimaryColors();
+	}
+	FDynamicMeshColorOverlay* Colors = Mesh.Attributes()->PrimaryColors();
+	TArray<int32> ElemIDs;
+	ElemIDs.SetNum(Mesh.MaxVertexID());
+	for (int32 VertexID : Mesh.VertexIndicesItr())
+	{
+		ElemIDs[VertexID] = Colors->AppendElement( (FVector4f)LinearColor );
+	}
+	for (int32 TriangleID : Mesh.TriangleIndicesItr())
+	{
+		FIndex3i Triangle = Mesh.GetTriangle(TriangleID);
+		Colors->SetTriangle(TriangleID, FIndex3i(ElemIDs[Triangle.A], ElemIDs[Triangle.B], ElemIDs[Triangle.C]) );
+	}
+}
+
+
+
+void FCombineMeshInstancesImpl::CombineMeshInstances(
+	const FInstanceSet& MeshInstances, const FOptions& Options, FResults& ResultsOut)
+{
+	FMeshInstanceAssembly InstanceAssembly;
+	InitializeMeshInstanceAssembly(MeshInstances, InstanceAssembly);
+	InitializeAssemblySourceMeshesFromLOD(InstanceAssembly, 0, Options.NumCopiedLODs);
+	ComputeMeshApproximations(Options, InstanceAssembly);
+	
+	InstanceAssembly.PreProcessInstanceMeshFunc = [&InstanceAssembly, &MeshInstances](FDynamicMesh3& AppendMesh, const FMeshInstance& Instance)
+	{
+		int32 SourceInstance = Instance.ExternalInstanceIndex[0];
+		int GroupDataIdx = MeshInstances.StaticMeshInstances[SourceInstance].GroupDataIndex;
+		if (MeshInstances.InstanceGroupDatas[GroupDataIdx].bHasConstantOverrideVertexColor)
+		{
+			FColor VertexColorSRGB = MeshInstances.InstanceGroupDatas[GroupDataIdx].OverrideVertexColor;
+			//FLinearColor VertexColorLinear(VertexColorSRGB);
+			FLinearColor VertexColorLinear = VertexColorSRGB.ReinterpretAsLinear();
+			SetConstantVertexColor(AppendMesh, VertexColorLinear);
+		}
+	};
+
+	TArray<UE::Geometry::FDynamicMesh3> CombinedMeshLODs;
+	BuildCombinedMesh(InstanceAssembly, Options, CombinedMeshLODs);
+
+	ResultsOut.CombinedMeshes.SetNum(1);
+	IGeometryProcessing_CombineMeshInstances::FOutputMesh& OutputMesh = ResultsOut.CombinedMeshes[0];
+	OutputMesh.MeshLODs = MoveTemp(CombinedMeshLODs);
+	OutputMesh.MaterialSet = InstanceAssembly.UniqueMaterials;
+}
\ No newline at end of file
diff --git a/Engine/Plugins/Experimental/MeshModelingToolsetExp/Source/GeometryProcessingAdapters/Private/GeometryProcessingAdaptersModule.cpp b/Engine/Plugins/Experimental/MeshModelingToolsetExp/Source/GeometryProcessingAdapters/Private/GeometryProcessingAdaptersModule.cpp
--- a/Engine/Plugins/Experimental/MeshModelingToolsetExp/Source/GeometryProcessingAdapters/Private/GeometryProcessingAdaptersModule.cpp
+++ b/Engine/Plugins/Experimental/MeshModelingToolsetExp/Source/GeometryProcessingAdapters/Private/GeometryProcessingAdaptersModule.cpp
@@ -17,6 +17,11 @@ void FGeometryProcessingAdaptersModule::StartupModule()
 		IModularFeatures::Get().RegisterModularFeature(IGeometryProcessing_ApproximateActors::GetModularFeatureName(), ApproximateActors.Get());
 	}
 	
+	CombineMeshInstances = MakeShared<FCombineMeshInstancesImpl>();
+	if (CombineMeshInstances.IsValid())
+	{
+		IModularFeatures::Get().RegisterModularFeature(IGeometryProcessing_CombineMeshInstances::GetModularFeatureName(), CombineMeshInstances.Get());
+	}
 }
 
 void FGeometryProcessingAdaptersModule::ShutdownModule()
@@ -30,6 +35,11 @@ void FGeometryProcessingAdaptersModule::ShutdownModule()
 		ApproximateActors = nullptr;
 	}
 	
+	if ( CombineMeshInstances.IsValid() )
+	{
+		IModularFeatures::Get().UnregisterModularFeature(IGeometryProcessing_CombineMeshInstances::GetModularFeatureName(), CombineMeshInstances.Get());
+		CombineMeshInstances = nullptr;
+	}
 }
 
 #undef LOCTEXT_NAMESPACE
diff --git a/Engine/Plugins/Experimental/MeshModelingToolsetExp/Source/GeometryProcessingAdapters/Public/GeometryProcessing/CombineMeshInstancesImpl.h b/Engine/Plugins/Experimental/MeshModelingToolsetExp/Source/GeometryProcessingAdapters/Public/GeometryProcessing/CombineMeshInstancesImpl.h
new file mode 100644
--- /dev/null
+++ b/Engine/Plugins/Experimental/MeshModelingToolsetExp/Source/GeometryProcessingAdapters/Public/GeometryProcessing/CombineMeshInstancesImpl.h
@@ -0,0 +1,33 @@
+// Copyright Epic Games, Inc. All Rights Reserved.
+
+#pragma once
+
+#include "CoreMinimal.h"
+#include "GeometryProcessingInterfaces/CombineMeshInstances.h"
+
+class UMaterialInterface;
+
+namespace UE
+{
+namespace Geometry
+{
+
+class FDynamicMesh3;
+
+/**
+ * Implementation of IGeometryProcessing_CombineMeshInstances
+ */
+class GEOMETRYPROCESSINGADAPTERS_API FCombineMeshInstancesImpl : public IGeometryProcessing_CombineMeshInstances
+{
+public:
+	virtual FOptions ConstructDefaultOptions() override;
+
+	virtual void CombineMeshInstances(const FInstanceSet& MeshInstances, const FOptions& Options, FResults& ResultsOut) override;
+
+protected:
+
+};
+
+
+}
+}
diff --git a/Engine/Plugins/Experimental/MeshModelingToolsetExp/Source/GeometryProcessingAdapters/Public/GeometryProcessingAdaptersModule.h b/Engine/Plugins/Experimental/MeshModelingToolsetExp/Source/GeometryProcessingAdapters/Public/GeometryProcessingAdaptersModule.h
--- a/Engine/Plugins/Experimental/MeshModelingToolsetExp/Source/GeometryProcessingAdapters/Public/GeometryProcessingAdaptersModule.h
+++ b/Engine/Plugins/Experimental/MeshModelingToolsetExp/Source/GeometryProcessingAdapters/Public/GeometryProcessingAdaptersModule.h
@@ -6,6 +6,7 @@
 #include "Modules/ModuleManager.h"
 
 #include "GeometryProcessing/ApproximateActorsImpl.h"
+#include "GeometryProcessing/CombineMeshInstancesImpl.h"
 
 class FGeometryProcessingAdaptersModule : public IModuleInterface
 {
@@ -18,4 +19,5 @@ public:
 	
 protected:
 	TSharedPtr<UE::Geometry::FApproximateActorsImpl> ApproximateActors;
+	TSharedPtr<UE::Geometry::FCombineMeshInstancesImpl> CombineMeshInstances;
 };
diff --git a/Engine/Source/Developer/GeometryProcessingInterfaces/GeometryProcessingInterfaces.Build.cs b/Engine/Source/Developer/GeometryProcessingInterfaces/GeometryProcessingInterfaces.Build.cs
--- a/Engine/Source/Developer/GeometryProcessingInterfaces/GeometryProcessingInterfaces.Build.cs
+++ b/Engine/Source/Developer/GeometryProcessingInterfaces/GeometryProcessingInterfaces.Build.cs
@@ -10,5 +10,6 @@ public class GeometryProcessingInterfaces : ModuleRules
 		PrivateDependencyModuleNames.Add("Core");
 		PrivateDependencyModuleNames.Add("CoreUObject");
 		PrivateDependencyModuleNames.Add("Engine");
-    }
+		PrivateDependencyModuleNames.Add("GeometryCore");
+	}
 }
diff --git a/Engine/Source/Developer/GeometryProcessingInterfaces/Private/GeometryProcessingInterfacesModule.cpp b/Engine/Source/Developer/GeometryProcessingInterfaces/Private/GeometryProcessingInterfacesModule.cpp
--- a/Engine/Source/Developer/GeometryProcessingInterfaces/Private/GeometryProcessingInterfacesModule.cpp
+++ b/Engine/Source/Developer/GeometryProcessingInterfaces/Private/GeometryProcessingInterfacesModule.cpp
@@ -6,10 +6,9 @@
 #include "Modules/ModuleManager.h"
 
 #include "GeometryProcessingInterfaces/ApproximateActors.h"
+#include "GeometryProcessingInterfaces/CombineMeshInstances.h"
 
 
-//DEFINE_LOG_CATEGORY_STATIC(LogMeshReduction, Verbose, All);
-
 IMPLEMENT_MODULE(FGeometryProcessingInterfacesModule, GeometryProcessingInterfaces);
 
 
@@ -22,6 +21,7 @@ void FGeometryProcessingInterfacesModule::StartupModule()
 void FGeometryProcessingInterfacesModule::ShutdownModule()
 {
 	ApproximateActors = nullptr;
+	CombineMeshInstances = nullptr;
 }
 
 
@@ -37,3 +37,17 @@ IGeometryProcessing_ApproximateActors* FGeometryProcessingInterfacesModule::GetA
 
 	return ApproximateActors;
 }
+IGeometryProcessing_CombineMeshInstances* FGeometryProcessingInterfacesModule::GetCombineMeshInstancesImplementation()
+{
+	if (ApproximateActors == nullptr)
+	{
+		TArray<IGeometryProcessing_CombineMeshInstances*> CombineMeshInstancesOptions =
+			IModularFeatures::Get().GetModularFeatureImplementations<IGeometryProcessing_CombineMeshInstances>(IGeometryProcessing_CombineMeshInstances::GetModularFeatureName());
+
+		CombineMeshInstances = (CombineMeshInstancesOptions.Num() > 0) ? CombineMeshInstancesOptions[0] : nullptr;
+	}
+
+	return CombineMeshInstances;
+}
+
+
diff --git a/Engine/Source/Developer/GeometryProcessingInterfaces/Private/GeometryProcessingInterfacesModule.h b/Engine/Source/Developer/GeometryProcessingInterfaces/Private/GeometryProcessingInterfacesModule.h
--- a/Engine/Source/Developer/GeometryProcessingInterfaces/Private/GeometryProcessingInterfacesModule.h
+++ b/Engine/Source/Developer/GeometryProcessingInterfaces/Private/GeometryProcessingInterfacesModule.h
@@ -6,6 +6,8 @@
 
 
 class IGeometryProcessing_ApproximateActors;
+class IGeometryProcessing_CombineMeshInstances;
+
 
 /**
  * Implementation of IGeometryProcessingInterfacesModule (which extends the standard IModuleInterface)
@@ -24,6 +26,13 @@ public:
 	 */
 	virtual IGeometryProcessing_ApproximateActors* GetApproximateActorsImplementation() override;
 
+	/**
+	 * @return implementation of IGeometryProcessing_CombineMeshInstances, if available, or nullptr (result is cached internally)
+	 */
+	virtual IGeometryProcessing_CombineMeshInstances* GetCombineMeshInstancesImplementation() override;
+
+
 private:
 	IGeometryProcessing_ApproximateActors* ApproximateActors = nullptr;
+	IGeometryProcessing_CombineMeshInstances* CombineMeshInstances = nullptr;
 };
diff --git a/Engine/Source/Developer/GeometryProcessingInterfaces/Public/GeometryProcessingInterfaces/CombineMeshInstances.h b/Engine/Source/Developer/GeometryProcessingInterfaces/Public/GeometryProcessingInterfaces/CombineMeshInstances.h
new file mode 100644
--- /dev/null
+++ b/Engine/Source/Developer/GeometryProcessingInterfaces/Public/GeometryProcessingInterfaces/CombineMeshInstances.h
@@ -0,0 +1,127 @@
+// Copyright Epic Games, Inc. All Rights Reserved.
+
+#pragma once
+
+#include "CoreMinimal.h"
+#include "Features/IModularFeature.h"
+#include "DynamicMesh/DynamicMesh3.h"
+
+class UPrimitiveComponent;
+class UStaticMesh;
+class UMaterialInterface;
+
+
+/**
+ * The CombineMeshInstances modular feature is used to provide a mechanism
+ * for merging a set of instances of meshes (ie mesh + transform + materials + ...)
+ * into a smaller set of meshes. Generally this involves creating simpler versions
+ * of the instances and appending them into one or a small number of combined meshes.
+ */
+class IGeometryProcessing_CombineMeshInstances : public IModularFeature
+{
+public:
+	virtual ~IGeometryProcessing_CombineMeshInstances() {}
+
+
+	enum class EMeshDetailLevel
+	{
+		Base = 0,
+		Standard = 1,
+		Small = 2,
+		Decorative = 3
+	};
+
+
+	struct FMeshInstanceGroupData
+	{
+		TArray<UMaterialInterface*> MaterialSet;
+		
+		bool bHasConstantOverrideVertexColor = false;
+		FColor OverrideVertexColor;
+	};
+
+
+	struct FBaseMeshInstance
+	{
+		EMeshDetailLevel DetailLevel = EMeshDetailLevel::Standard;
+		TArray<FTransform3d> TransformSequence;
+		int32 GroupDataIndex = -1;		// index into FInstanceSet::InstanceGroupDatas
+	};
+
+
+	struct FStaticMeshInstance : FBaseMeshInstance
+	{
+		UStaticMesh* SourceMesh = nullptr;
+		UPrimitiveComponent* SourceComponent = nullptr;
+		int32 SourceInstanceIndex = 0;
+	};
+
+
+	struct FInstanceSet
+	{
+		TArray<FStaticMeshInstance> StaticMeshInstances;
+
+		// sets of data shared across multiple instances
+		TArray<FMeshInstanceGroupData> InstanceGroupDatas;
+	};
+
+
+	struct FOptions
+	{
+		// number of requested LODs
+		int32 NumLODs = 5;
+
+		int32 NumCopiedLODs = 1;
+
+		int32 ApproximationSourceLOD = 0;
+
+		int32 NumSimplifiedLODs = 3;
+		double SimplifyBaseTolerance = 1.0;
+		double SimplifyLODLevelToleranceScale = 2.0;
+
+		// LOD level to filter out detail parts
+		int32 FilterDecorativePartsLODLevel = 2;
+	};
+
+
+	struct FOutputMesh
+	{
+		TArray<UE::Geometry::FDynamicMesh3> MeshLODs;
+		TArray<UMaterialInterface*> MaterialSet;
+	};
+
+
+	struct FResults
+	{
+		//EResultCode ResultCode = EResultCode::UnknownError;
+
+		TArray<FOutputMesh> CombinedMeshes;
+	};
+
+
+
+
+	virtual FOptions ConstructDefaultOptions()
+	{
+		check(false);		// not implemented in base class
+		return FOptions();
+	}
+
+
+	virtual void CombineMeshInstances(
+		const FInstanceSet& MeshInstances, 
+		const FOptions& Options, 
+		FResults& ResultsOut) 
+	{
+		check(false);		// not implemented in base class
+	}
+
+
+
+	// Modular feature name to register for retrieval during runtime
+	static const FName GetModularFeatureName()
+	{
+		return TEXT("GeometryProcessing_CombineMeshInstances");
+	}
+
+};
\ No newline at end of file
diff --git a/Engine/Source/Developer/GeometryProcessingInterfaces/Public/IGeometryProcessingInterfacesModule.h b/Engine/Source/Developer/GeometryProcessingInterfaces/Public/IGeometryProcessingInterfacesModule.h
--- a/Engine/Source/Developer/GeometryProcessingInterfaces/Public/IGeometryProcessingInterfacesModule.h
+++ b/Engine/Source/Developer/GeometryProcessingInterfaces/Public/IGeometryProcessingInterfacesModule.h
@@ -5,6 +5,7 @@
 #include "Modules/ModuleInterface.h"
 
 class IGeometryProcessing_ApproximateActors;
+class IGeometryProcessing_CombineMeshInstances;
 
 /**
  * Abstract interface to a Module that provides functions to access
@@ -28,4 +29,9 @@ public:
 	 */
 	virtual IGeometryProcessing_ApproximateActors* GetApproximateActorsImplementation() = 0;
 
+	/**
+	 * @return implementation of IGeometryProcessing_CombineMeshInstances, if available, or nullptr (result is cached internally)
+	 */
+	virtual IGeometryProcessing_CombineMeshInstances* GetCombineMeshInstancesImplementation() = 0;
+
 };
