add missing header to resolve nonunity build error

diff --git a/Engine/Plugins/Experimental/MeshModelingToolsetExp/Source/GeometryProcessingAdapters/Private/CombineMeshInstancesImpl.cpp b/Engine/Plugins/Experimental/MeshModelingToolsetExp/Source/GeometryProcessingAdapters/Private/CombineMeshInstancesImpl.cpp
--- a/Engine/Plugins/Experimental/MeshModelingToolsetExp/Source/GeometryProcessingAdapters/Private/CombineMeshInstancesImpl.cpp
+++ b/Engine/Plugins/Experimental/MeshModelingToolsetExp/Source/GeometryProcessingAdapters/Private/CombineMeshInstancesImpl.cpp
@@ -22,12 +22,12 @@
 #include "ShapeApproximation/MeshSimpleShapeApproximation.h"
 #include "Generators/GridBoxMeshGenerator.h"
 
-
 #include "MeshSimplification.h"
 #include "DynamicMesh/ColliderMesh.h"
 #include "MeshConstraintsUtil.h"
 #include "DynamicMesh/Operations/MergeCoincidentMeshEdges.h"
 
+#include "TransformSequence.h"
 
 
 using namespace UE::Geometry;
