GeometryProcessing: incremental improvements to CombineMeshInstances implementation. Now generates multiple LODs of approximated meshes, can fit swept-hull and convex-hull, choose between approximations based on a geometric-error-per-triangle cost factor, and optionally remove hidden triangles (expensive).

diff --git a/Engine/Plugins/Experimental/MeshModelingToolsetExp/Source/GeometryProcessingAdapters/Private/CombineMeshInstancesImpl.cpp b/Engine/Plugins/Experimental/MeshModelingToolsetExp/Source/GeometryProcessingAdapters/Private/CombineMeshInstancesImpl.cpp
--- a/Engine/Plugins/Experimental/MeshModelingToolsetExp/Source/GeometryProcessingAdapters/Private/CombineMeshInstancesImpl.cpp
+++ b/Engine/Plugins/Experimental/MeshModelingToolsetExp/Source/GeometryProcessingAdapters/Private/CombineMeshInstancesImpl.cpp
@@ -17,6 +17,7 @@
 #include "DynamicMeshEditor.h"
 #include "DynamicMesh/DynamicMeshAttributeSet.h"
 #include "DynamicMesh/MeshNormals.h"
+#include "DynamicMesh/DynamicMeshAABBTree3.h"
 
 #include "ShapeApproximation/ShapeDetection3.h"
 #include "ShapeApproximation/MeshSimpleShapeApproximation.h"
@@ -26,6 +27,7 @@
 #include "DynamicMesh/ColliderMesh.h"
 #include "MeshConstraintsUtil.h"
 #include "DynamicMesh/Operations/MergeCoincidentMeshEdges.h"
+#include "Operations/RemoveOccludedTriangles.h"
 
 #include "TransformSequence.h"
 
@@ -34,6 +36,12 @@ using namespace UE::Geometry;
 
 
 
+static TAutoConsoleVariable<int> CVarGeometryCombineMeshInstancesRemoveHidden(
+	TEXT("geometry.CombineInstances.DebugRemoveHiddenStrategy"),
+	0,
+	TEXT("Configure hidden-removal strategy via (temporary debug)"));
+
+
 enum class EMeshDetailLevel
 {
 	Base = 0,
@@ -75,7 +83,8 @@ struct FOptimizedGeometry
 {
 	TArray<UE::Geometry::FDynamicMesh3> SimplifiedMeshLODs;
 
-	UE::Geometry::FDynamicMesh3 OptimizedMesh;
+	TArray<UE::Geometry::FDynamicMesh3> ApproximateMeshLODs;
+
 	//UE::Geometry::FSimpleShapeSet3d CollisionShapes;
 };
 
@@ -97,15 +106,17 @@ public:
 	TArray<FSourceGeometry> SourceMeshGeometry;
 	TArray<FOptimizedGeometry> OptimizedMeshGeometry;
 
+	TArray<FDynamicMeshAABBTree3> SourceMeshSpatials;
 
 	// allow external code to preprocess dynamic mesh for a specific instance
 	TFunction<void(FDynamicMesh3&, const FMeshInstance&)> PreProcessInstanceMeshFunc;
-
 };
 
 
 
 
+
+
 void InitializeMeshInstanceAssembly(
 	const IGeometryProcessing_CombineMeshInstances::FInstanceSet& SourceInstanceSet,
 	FMeshInstanceAssembly& AssemblyOut)
@@ -176,8 +187,6 @@ void InitializeMeshInstanceAssembly(
 }
 
 
-
-
 void InitializeAssemblySourceMeshesFromLOD(
 	FMeshInstanceAssembly& Assembly,
 	int32 SourceAssetBaseLOD,
@@ -265,19 +274,190 @@ void InitializeAssemblySourceMeshesFromLOD(
 	//	//}
 
 	//});
+}
+
+
+
+
+
+
+
+
+
+
+/**
+ * @return ( Sqrt(Sum-of-squared-distances) / NumPoints , Max(distance)  )
+ * 
+ */
+static FVector2d DeviationMetric(const FDynamicMesh3& MeasureMesh, const FDynamicMeshAABBTree3& SourceBVH)
+{
+	// todo: could consider normal deviation?
+	int PointCount = 0;
+	double SumDistanceSqr = 0;
+	double MaxDistanceSqr = 0;
+	auto TestPointFunc = [&SumDistanceSqr, &MaxDistanceSqr, &PointCount, &SourceBVH](FVector3d Point)
+	{
+		double NearDistSqr;
+		SourceBVH.FindNearestTriangle(Point, NearDistSqr);
+		if (NearDistSqr > MaxDistanceSqr)
+		{
+			MaxDistanceSqr = NearDistSqr;
+		}
+		SumDistanceSqr += NearDistSqr;
+		PointCount++;
+	};
+
+	for (int32 vid : MeasureMesh.VertexIndicesItr())
+	{
+		TestPointFunc(MeasureMesh.GetVertex(vid));
+	}
+
+	for (int32 tid : MeasureMesh.TriangleIndicesItr())
+	{
+		TestPointFunc(MeasureMesh.GetTriCentroid(tid));
+	}
 
+	for (int32 eid : MeasureMesh.EdgeIndicesItr())
+	{
+		TestPointFunc(MeasureMesh.GetEdgePoint(eid, 0.5));
+	}
 
+	return FVector2d(
+		FMathd::Sqrt(SumDistanceSqr) / (double)PointCount,
+		FMathd::Sqrt(MaxDistanceSqr) );
 }
 
 
 
+class FPartApproxSelector
+{
+public:
+	double TriangleCost = 0.7;
+
+	struct FResultOption
+	{
+		FVector2d DeviationMetric;
+		double CostMetric;
+		TSharedPtr<FDynamicMesh3> Mesh;
+		int32 MethodID;
+	};
+	TArray<FResultOption> Options;
+
+	const FDynamicMesh3* SourceMesh;
+	const FDynamicMeshAABBTree3* Spatial;
+
+	void Initialize(const FDynamicMesh3* SourceMeshIn, const FDynamicMeshAABBTree3* SpatialIn)
+	{
+		SourceMesh = SourceMeshIn;
+		Spatial = SpatialIn;
+	}
+
+	void AddGeneratedMesh(
+		const FDynamicMesh3& ExternalMesh,
+		int32 MethodID )
+	{
+		FResultOption Option;
+		Option.MethodID = MethodID;
+		Option.Mesh = MakeShared<FDynamicMesh3>(ExternalMesh);
+		ComputeMetric(Option);
+		Options.Add(Option);
+	}
+
+	void AddGeneratedMesh(
+		TFunctionRef<void(FDynamicMesh3&)> GeneratorFunc,
+		int32 MethodID )
+	{
+		FResultOption Option;
+		Option.MethodID = MethodID;
+		Option.Mesh = MakeShared<FDynamicMesh3>(*SourceMesh);
+		GeneratorFunc(*Option.Mesh);
+		ComputeMetric(Option);
+		Options.Add(Option);
+	}
+
+	void ComputeMetric(FResultOption& Option)
+	{
+		Option.DeviationMetric = DeviationMetric(*Option.Mesh, *Spatial);
+		int32 TriCount = Option.Mesh->TriangleCount();
+		int32 BaseTriCount = 12;		// 2 tris for each face of box
+		Option.CostMetric = 
+			Option.DeviationMetric[0] * FMathd::Pow( (double)TriCount / (double)BaseTriCount, TriangleCost );
+	}
+
+	void SelectBestOption(
+		FDynamicMesh3& ResultMesh)
+	{
+		Options.StableSort( [&](const FResultOption& A, const FResultOption& B) { return A.CostMetric < B.CostMetric; } );
+		ResultMesh = MoveTemp(*Options[0].Mesh);
+	}
+};
+
+
+
+
+
+
+
+
+
+
+void InitializeInstanceAssemblySpatials(FMeshInstanceAssembly& Assembly)
+{
+	int32 NumSets = Assembly.InstanceSets.Num();
+	Assembly.SourceMeshSpatials.SetNum(NumSets);
+	
+	ParallelFor(NumSets, [&](int32 Index)
+	{
+		TUniquePtr<FMeshInstanceSet>& InstanceSet = Assembly.InstanceSets[Index];
+		FSourceGeometry& Target = Assembly.SourceMeshGeometry[Index];
+		FDynamicMeshAABBTree3& Spatial = Assembly.SourceMeshSpatials[Index];
+		Spatial.SetMesh(&Target.SourceMeshLODs[0], true);
+	});
+}
+
+
+
+/**
+ * Simplification can make a mess on low-poly shapes and sometimes just using a simple
+ * approximation would be better, use our metric to make this decision.
+ * (todo: this could maybe be folded into simplified-mesh computations...)
+ */
+void ReplaceBadSimplifiedLODs(FMeshInstanceAssembly& Assembly)
+{
+	int32 NumSets = Assembly.InstanceSets.Num();
+	
+	ParallelFor(NumSets, [&](int32 Index)
+	{
+		TUniquePtr<FMeshInstanceSet>& InstanceSet = Assembly.InstanceSets[Index];
+		FSourceGeometry& Target = Assembly.SourceMeshGeometry[Index];
+		FDynamicMeshAABBTree3& Spatial = Assembly.SourceMeshSpatials[Index];
+		FOptimizedGeometry& OptimizedTargets = Assembly.OptimizedMeshGeometry[Index];
+
+		for ( int32 k = OptimizedTargets.SimplifiedMeshLODs.Num()-1; k >= 0; --k )
+		{
+			FPartApproxSelector Selector;
+			Selector.Initialize(Spatial.GetMesh(), &Spatial);
+			if ( k == OptimizedTargets.SimplifiedMeshLODs.Num()-1 )
+			{
+				Selector.AddGeneratedMesh(OptimizedTargets.ApproximateMeshLODs[0], 2);
+			}
+			else
+			{
+				Selector.AddGeneratedMesh(OptimizedTargets.SimplifiedMeshLODs[k+1], 1);
+			}
+			Selector.AddGeneratedMesh(OptimizedTargets.SimplifiedMeshLODs[k], 0);
+
+			// either keep current mesh or replace w/ simplified version
+			Selector.SelectBestOption(OptimizedTargets.SimplifiedMeshLODs[k]);
+		}
+	});
+}
 
 
 
 static void SimplifyPartMesh(
 	FDynamicMesh3& EditMesh, 
 	double Tolerance, 
-	bool bNoSplitAttributes,
 	double RecomputeNormalsAngleThreshold)
 {
 	// weld edges in case input was unwelded...
@@ -295,7 +475,6 @@ static void SimplifyPartMesh(
 	EditMesh.Attributes()->DisableTangents();
 	EditMesh.Attributes()->DisablePrimaryColors();
 	FMeshNormals::InitializeOverlayToPerVertexNormals(EditMesh.Attributes()->PrimaryNormals(), false);
-	bNoSplitAttributes = true;
 
 	Simplifier.ProjectionMode = FVolPresMeshSimplification::ETargetProjectionMode::NoProjection;
 
@@ -306,15 +485,16 @@ static void SimplifyPartMesh(
 
 	Simplifier.DEBUG_CHECK_LEVEL = 0;
 	Simplifier.bRetainQuadricMemory = false; 
-	if ( bNoSplitAttributes == false )
-	{
-		Simplifier.bAllowSeamCollapse = true;
-		Simplifier.SetEdgeFlipTolerance(1.e-5);
-		if (EditMesh.HasAttributes())
-		{
-			EditMesh.Attributes()->SplitAllBowties();	// eliminate any bowties that might have formed on attribute seams.
-		}
-	}
+	// currenly no need for this path, may need to resurrect it in the future
+	//if ( bNoSplitAttributes == false )
+	//{
+	//	Simplifier.bAllowSeamCollapse = true;
+	//	Simplifier.SetEdgeFlipTolerance(1.e-5);
+	//	if (EditMesh.HasAttributes())
+	//	{
+	//		EditMesh.Attributes()->SplitAllBowties();	// eliminate any bowties that might have formed on attribute seams.
+	//	}
+	//}
 
 	// this should preserve part shape better but it completely fails currently =\
 	//Simplifier.CollapseMode = FVolPresMeshSimplification::ESimplificationCollapseModes::MinimalExistingVertexError;
@@ -345,6 +525,140 @@ static void SimplifyPartMesh(
 
 
 
+static void ComputeBoxApproximation(
+	const FDynamicMesh3& SourceMesh,
+	FDynamicMesh3& OutputMesh )
+{
+	FMeshSimpleShapeApproximation ShapeApprox;
+	ShapeApprox.InitializeSourceMeshes( {&SourceMesh} );
+	ShapeApprox.bDetectBoxes = ShapeApprox.bDetectCapsules = ShapeApprox.bDetectConvexes = ShapeApprox.bDetectSpheres = false;
+
+	FSimpleShapeSet3d ResultBoxes;
+	ShapeApprox.Generate_OrientedBoxes(ResultBoxes);
+	UE::Geometry::FOrientedBox3d OrientedBox = ResultBoxes.Boxes[0].Box;
+
+	// oriented box fitting is under-determined, in cases where the AABB and the OBB have the nearly the
+	// same volume, generally we prefer an AABB
+	FAxisAlignedBox3d AlignedBox = SourceMesh.GetBounds(false);
+	if (AlignedBox.Volume() < 1.05*OrientedBox.Volume())
+	{
+		OrientedBox = UE::Geometry::FOrientedBox3d(AlignedBox);
+	}
+
+	FGridBoxMeshGenerator BoxGen;
+	BoxGen.Box = ResultBoxes.Boxes[0].Box;
+	BoxGen.EdgeVertices = {0,0,0};
+	OutputMesh.Copy(&BoxGen.Generate());
+}
+
+
+
+enum class EApproximatePartMethod : uint8
+{
+	OrientedBox = 0,
+	MinVolumeSweptHull = 1,
+	ConvexHull = 3,
+	MinTriCountHull = 4,
+	FlattendExtrusion = 5,
+
+	AutoBestFit = 10,
+
+	Original = 100
+
+};
+
+
+
+static void ComputeSimplePartApproximation(
+	const FDynamicMesh3& SourcePartMesh, 
+	FDynamicMesh3& DestMesh,
+	EApproximatePartMethod ApproxMethod)
+{
+
+	if (ApproxMethod == EApproximatePartMethod::OrientedBox)
+	{
+		ComputeBoxApproximation(SourcePartMesh, DestMesh);
+	}
+
+	FMeshSimpleShapeApproximation ShapeApprox;
+	ShapeApprox.InitializeSourceMeshes( {&SourcePartMesh} );
+	ShapeApprox.bDetectBoxes = ShapeApprox.bDetectCapsules = ShapeApprox.bDetectConvexes = ShapeApprox.bDetectSpheres = false;
+
+	FDynamicMesh3 ResultMesh;
+
+	FDynamicMesh3 ConvexMesh;
+	if ( ApproxMethod == EApproximatePartMethod::ConvexHull || ApproxMethod == EApproximatePartMethod::MinTriCountHull )
+	{
+		FSimpleShapeSet3d ResultConvex;
+		ShapeApprox.Generate_ConvexHulls(ResultConvex);
+		ConvexMesh = (ResultConvex.Convexes.Num() > 0) ? MoveTemp(ResultConvex.Convexes[0].Mesh) : FDynamicMesh3();
+	}
+
+	FDynamicMesh3 MinVolumeHull;
+	if ( ApproxMethod != EApproximatePartMethod::ConvexHull )
+	{
+		FSimpleShapeSet3d ResultX, ResultY, ResultZ;
+		ShapeApprox.Generate_ProjectedHulls(ResultX, FMeshSimpleShapeApproximation::EProjectedHullAxisMode::X);
+		ShapeApprox.Generate_ProjectedHulls(ResultY, FMeshSimpleShapeApproximation::EProjectedHullAxisMode::Y);
+		ShapeApprox.Generate_ProjectedHulls(ResultZ, FMeshSimpleShapeApproximation::EProjectedHullAxisMode::Z);
+		FDynamicMesh3 SweptHullX = (ResultX.Convexes.Num() > 0) ? MoveTemp(ResultX.Convexes[0].Mesh) : FDynamicMesh3();
+		double VolumeX = (SweptHullX.TriangleCount() > 0) ? TMeshQueries<FDynamicMesh3>::GetVolumeArea(SweptHullX)[0] : TNumericLimits<double>::Max();
+		FDynamicMesh3 SweptHullY = (ResultY.Convexes.Num() > 0) ? MoveTemp(ResultY.Convexes[0].Mesh) : FDynamicMesh3();
+		double VolumeY = (SweptHullY.TriangleCount() > 0) ? TMeshQueries<FDynamicMesh3>::GetVolumeArea(SweptHullY)[0] : TNumericLimits<double>::Max();
+		FDynamicMesh3 SweptHullZ = (ResultZ.Convexes.Num() > 0) ? MoveTemp(ResultZ.Convexes[0].Mesh) : FDynamicMesh3();
+		double VolumeZ = (SweptHullZ.TriangleCount() > 0) ? TMeshQueries<FDynamicMesh3>::GetVolumeArea(SweptHullZ)[0] : TNumericLimits<double>::Max();
+
+		int Idx = MinElementIndex(FVector(VolumeX, VolumeY, VolumeZ));
+		MinVolumeHull = (Idx == 0) ? SweptHullX : (Idx == 1) ? SweptHullY : SweptHullZ;
+	}
+
+	if (ApproxMethod == EApproximatePartMethod::ConvexHull)
+	{
+		ResultMesh = (ConvexMesh.TriangleCount() > 0) ?  MoveTemp(ConvexMesh) : SourcePartMesh;
+	}
+	else if (ApproxMethod == EApproximatePartMethod::MinVolumeSweptHull)
+	{
+		ResultMesh = (MinVolumeHull.TriangleCount() > 0) ?  MoveTemp(MinVolumeHull) : SourcePartMesh;
+	}
+	else if (ApproxMethod == EApproximatePartMethod::MinTriCountHull)
+	{
+		ResultMesh = (MinVolumeHull.TriangleCount() < ConvexMesh.TriangleCount()) ? 
+			MoveTemp(MinVolumeHull) : MoveTemp(ConvexMesh);
+	}
+
+	DestMesh = (ResultMesh.TriangleCount() > 0) ? MoveTemp(ResultMesh) : SourcePartMesh;
+}
+
+
+static void SelectBestFittingMeshApproximation(
+	const FDynamicMesh3& OriginalMesh, 
+	const FDynamicMeshAABBTree3& OriginalMeshSpatial,
+	FDynamicMesh3& ResultMesh,
+	double AcceptableDeviationTol,
+	double TriangleCost)
+{
+	FPartApproxSelector ApproxSelector;
+	ApproxSelector.Initialize(&OriginalMesh, &OriginalMeshSpatial);
+	ApproxSelector.TriangleCost = TriangleCost;
+
+	ApproxSelector.AddGeneratedMesh( [&](FDynamicMesh3& PartMeshInOut) {
+		ComputeSimplePartApproximation(PartMeshInOut, PartMeshInOut, EApproximatePartMethod::OrientedBox);
+	}, (int32)EApproximatePartMethod::OrientedBox );
+
+	ApproxSelector.AddGeneratedMesh( [&](FDynamicMesh3& PartMeshInOut) {
+		ComputeSimplePartApproximation(PartMeshInOut, PartMeshInOut, EApproximatePartMethod::MinVolumeSweptHull);
+	}, (int32)EApproximatePartMethod::MinVolumeSweptHull );
+
+	ApproxSelector.AddGeneratedMesh( [&](FDynamicMesh3& PartMeshInOut) {
+		ComputeSimplePartApproximation(PartMeshInOut, PartMeshInOut, EApproximatePartMethod::ConvexHull);
+	}, (int32)EApproximatePartMethod::ConvexHull );
+
+	ApproxSelector.SelectBestOption(ResultMesh);
+}
+
+
+
+
 void ComputeMeshApproximations(
 	IGeometryProcessing_CombineMeshInstances::FOptions CombineOptions,
 	FMeshInstanceAssembly& Assembly)
@@ -356,43 +670,91 @@ void ComputeMeshApproximations(
 	Assembly.OptimizedMeshGeometry.SetNum(NumSets);
 
 	int32 NumSimplifiedLODs = CombineOptions.NumSimplifiedLODs;
+	int32 NumApproxLODs = FMath::Max(1, 
+		CombineOptions.NumLODs - CombineOptions.NumCopiedLODs - CombineOptions.NumSimplifiedLODs);
 
 	ParallelFor(NumSets, [&](int32 Index)
 	{
 		TUniquePtr<FMeshInstanceSet>& InstanceSet = Assembly.InstanceSets[Index];
 		const FSourceGeometry& SourceGeo = Assembly.SourceMeshGeometry[Index];
-		const FDynamicMesh3& InitialSourceMesh = (CombineOptions.ApproximationSourceLOD < SourceGeo.SourceMeshLODs.Num()) ?
+		const FDynamicMesh3& OptimizationSourceMesh = (CombineOptions.ApproximationSourceLOD < SourceGeo.SourceMeshLODs.Num()) ?
 			SourceGeo.SourceMeshLODs[CombineOptions.ApproximationSourceLOD] : SourceGeo.SourceMeshLODs.Last();
 		FOptimizedGeometry& ApproxGeo = Assembly.OptimizedMeshGeometry[Index];
 
+		FDynamicMeshAABBTree3 OptimizationSourceMeshSpatial(&OptimizationSourceMesh, true);
+
 		// compute simplified part LODs
 		ApproxGeo.SimplifiedMeshLODs.SetNum(NumSimplifiedLODs);
 		double InitialTolerance = CombineOptions.SimplifyBaseTolerance;
 		for (int32 k = 0; k < NumSimplifiedLODs; ++k)
 		{
-			ApproxGeo.SimplifiedMeshLODs[k] = InitialSourceMesh;
-			SimplifyPartMesh(ApproxGeo.SimplifiedMeshLODs[k], InitialTolerance, false, AngleThreshold);
+			ApproxGeo.SimplifiedMeshLODs[k] = OptimizationSourceMesh;
+			SimplifyPartMesh(ApproxGeo.SimplifiedMeshLODs[k], InitialTolerance, AngleThreshold);
 			InitialTolerance *= CombineOptions.SimplifyLODLevelToleranceScale;
 		}
 
-		// compute shape approximation
-		FMeshSimpleShapeApproximation ShapeApprox;
-		ShapeApprox.InitializeSourceMeshes( {&InitialSourceMesh} );
-		ShapeApprox.bDetectBoxes = ShapeApprox.bDetectCapsules = ShapeApprox.bDetectConvexes = ShapeApprox.bDetectSpheres = false;
+		// compute shape approximation LODs
+		ApproxGeo.ApproximateMeshLODs.SetNum(NumApproxLODs);
+		double InitialTriCost = CombineOptions.OptimizeBaseTriCost;
+		for (int32 k = 0; k < NumApproxLODs; ++k)
+		{
+			SelectBestFittingMeshApproximation(OptimizationSourceMesh, OptimizationSourceMeshSpatial, 
+				ApproxGeo.ApproximateMeshLODs[k], CombineOptions.SimplifyBaseTolerance, InitialTriCost);
+			InitialTriCost *= CombineOptions.OptimizeLODLevelTriCostScale;
 
-		FSimpleShapeSet3d ResultBoxes;
-		ShapeApprox.Generate_OrientedBoxes(ResultBoxes);
-		FGridBoxMeshGenerator BoxGen;
-		BoxGen.Box = ResultBoxes.Boxes[0].Box;
-		BoxGen.EdgeVertices = {0,0,0};
-		ApproxGeo.OptimizedMesh.Copy(&BoxGen.Generate());
+			// update enabled attribs (is this good?)
+			ApproxGeo.ApproximateMeshLODs[k].EnableMatchingAttributes(OptimizationSourceMesh);
 
-		ApproxGeo.OptimizedMesh.EnableMatchingAttributes(InitialSourceMesh);
-		FMeshNormals::InitializeOverlayTopologyFromOpeningAngle(&ApproxGeo.OptimizedMesh, ApproxGeo.OptimizedMesh.Attributes()->PrimaryNormals(), AngleThreshold);
-		FMeshNormals::QuickRecomputeOverlayNormals(ApproxGeo.OptimizedMesh);
+			// recompute normals
+			FMeshNormals::InitializeOverlayTopologyFromOpeningAngle(&ApproxGeo.ApproximateMeshLODs[k], ApproxGeo.ApproximateMeshLODs[k].Attributes()->PrimaryNormals(), AngleThreshold);
+			FMeshNormals::QuickRecomputeOverlayNormals(ApproxGeo.ApproximateMeshLODs[k]);
+		}
 
 	});
 
+
+	// try to filter out simplifications that did bad things
+	// argh crashing!
+	ReplaceBadSimplifiedLODs(Assembly);
+}
+
+
+
+static void RemoveHiddenFaces(FDynamicMesh3& EditMesh, double MaxDistance = 200)
+{
+	TRemoveOccludedTriangles<FDynamicMesh3> Jacket(&EditMesh);
+
+	Jacket.InsideMode = UE::Geometry::EOcclusionCalculationMode::SimpleOcclusionTest;
+	Jacket.TriangleSamplingMethod = UE::Geometry::EOcclusionTriangleSampling::Centroids;
+	Jacket.WindingIsoValue = 0.5;
+	Jacket.NormalOffset = FMathd::ZeroTolerance;
+	Jacket.AddRandomRays = 25;
+	Jacket.AddTriangleSamples = 100;
+	//if (MaxDistance > 0)
+	//{
+	//	Jacket.MaxDistance = MaxDistance;
+	//}
+
+	TArray<FTransformSRT3d> NoTransforms;
+	NoTransforms.Add(FTransformSRT3d::Identity());
+
+	//  set up AABBTree and FWNTree lists
+	FDynamicMeshAABBTree3 Spatial(&EditMesh);
+	TArray<FDynamicMeshAABBTree3*> OccluderTrees; 
+	OccluderTrees.Add(&Spatial);
+		
+	TFastWindingTree<FDynamicMesh3> FastWinding(&Spatial, false);
+	TArray<TFastWindingTree<FDynamicMesh3>*> OccluderWindings; 
+	OccluderWindings.Add(&FastWinding);
+
+	Jacket.Select(NoTransforms, OccluderTrees, OccluderWindings, NoTransforms);
+	
+	if (Jacket.RemovedT.Num() > 0)
+	{
+		Jacket.RemoveSelected();
+	}
+
+	EditMesh.CompactInPlace();
 }
 
 
@@ -400,6 +762,10 @@ void ComputeMeshApproximations(
 namespace UE::Geometry
 {
 
+
+
+
+
 struct FCombinedMeshLOD
 {
 	FDynamicMesh3 Mesh;
@@ -464,7 +830,9 @@ void BuildCombinedMesh(
 			const FDynamicMesh3* ApproximateAppendMesh = nullptr;
 			const FDynamicMesh3* UseAppendMesh = nullptr;
 
-			ApproximateAppendMesh = &OptimizedGeometry.OptimizedMesh;
+			// default approximate mesh to lowest-quality approximation (box), need to do this
+			// so that we always have something to swap to for Decorative parts
+			ApproximateAppendMesh = &OptimizedGeometry.ApproximateMeshLODs.Last();
 
 			if (LODLevel < CombineOptions.NumCopiedLODs)
 			{
@@ -480,6 +848,8 @@ void BuildCombinedMesh(
 			}
 			else
 			{
+				int32 ApproxLODIndex = LODLevel - CombineOptions.NumCopiedLODs - CombineOptions.NumSimplifiedLODs;
+				ApproximateAppendMesh = &OptimizedGeometry.ApproximateMeshLODs[ApproxLODIndex];
 				UseAppendMesh = ApproximateAppendMesh;
 			}
 
@@ -500,6 +870,7 @@ void BuildCombinedMesh(
 					UseAppendMesh = ApproximateAppendMesh;
 				}
 
+				// need to make a copy to run pre-process func
 				FDynamicMesh3 TempAppendMesh(*UseAppendMesh);
 				if (Assembly.PreProcessInstanceMeshFunc)
 				{
@@ -530,6 +901,14 @@ void BuildCombinedMesh(
 		}
 	}
 
+	if (CVarGeometryCombineMeshInstancesRemoveHidden.GetValueOnGameThread() > 0)
+	{
+		ParallelFor(MeshLODs.Num(), [&](int32 k)
+		{
+			RemoveHiddenFaces(MeshLODs[k].Mesh, 200.0);
+		});
+	}
+
 
 	for (int32 LODLevel = 0; LODLevel < NumLODs; ++LODLevel)
 	{
@@ -564,6 +943,9 @@ IGeometryProcessing_CombineMeshInstances::FOptions FCombineMeshInstancesImpl::Co
 	Options.SimplifyBaseTolerance = 0.25;
 	Options.SimplifyLODLevelToleranceScale = 2.0;
 
+	Options.OptimizeBaseTriCost = 0.7;
+	Options.OptimizeLODLevelTriCostScale = 2.5;
+
 	//// LOD level to filter out detail parts
 	Options.FilterDecorativePartsLODLevel = 2;
 
@@ -605,7 +987,10 @@ void FCombineMeshInstancesImpl::CombineMeshInstances(
 	FMeshInstanceAssembly InstanceAssembly;
 	InitializeMeshInstanceAssembly(MeshInstances, InstanceAssembly);
 	InitializeAssemblySourceMeshesFromLOD(InstanceAssembly, 0, Options.NumCopiedLODs);
+	InitializeInstanceAssemblySpatials(InstanceAssembly);
 	ComputeMeshApproximations(Options, InstanceAssembly);
+
+
 	
 	InstanceAssembly.PreProcessInstanceMeshFunc = [&InstanceAssembly, &MeshInstances](FDynamicMesh3& AppendMesh, const FMeshInstance& Instance)
 	{
diff --git a/Engine/Source/Developer/GeometryProcessingInterfaces/Public/GeometryProcessingInterfaces/CombineMeshInstances.h b/Engine/Source/Developer/GeometryProcessingInterfaces/Public/GeometryProcessingInterfaces/CombineMeshInstances.h
--- a/Engine/Source/Developer/GeometryProcessingInterfaces/Public/GeometryProcessingInterfaces/CombineMeshInstances.h
+++ b/Engine/Source/Developer/GeometryProcessingInterfaces/Public/GeometryProcessingInterfaces/CombineMeshInstances.h
@@ -79,6 +79,9 @@ public:
 		double SimplifyBaseTolerance = 1.0;
 		double SimplifyLODLevelToleranceScale = 2.0;
 
+		double OptimizeBaseTriCost = 0.7;
+		double OptimizeLODLevelTriCostScale = 1.5;
+
 		// LOD level to filter out detail parts
 		int32 FilterDecorativePartsLODLevel = 2;
 	};
