fixed crash: EditedMesh must be created before calling UpdateBonePositionInfos as it's used to compute the radius when there's only 1 bone

diff --git a/Engine/Plugins/Experimental/MeshModelingToolsetExp/Source/MeshModelingToolsExp/Private/SkinWeightsPaintTool.cpp b/Engine/Plugins/Experimental/MeshModelingToolsetExp/Source/MeshModelingToolsExp/Private/SkinWeightsPaintTool.cpp
--- a/Engine/Plugins/Experimental/MeshModelingToolsetExp/Source/MeshModelingToolsExp/Private/SkinWeightsPaintTool.cpp
+++ b/Engine/Plugins/Experimental/MeshModelingToolsetExp/Source/MeshModelingToolsExp/Private/SkinWeightsPaintTool.cpp
@@ -375,6 +375,11 @@ void USkinWeightsPaintTool::Setup()
 	// TODO replace this and synchronize selection & rendering with skeletal mesh editor
 	USkeletalMesh& SkeletalMesh = *Component->GetSkeletalMeshAsset();
 	BoneContainer.InitializeTo(Component->RequiredBones, Component->GetCurveFilterSettings(0)/* Always use the highest LOD */, SkeletalMesh);
+
+	// create a mesh description for editing (this must be done before calling UpdateBonePositionInfos) 
+	EditedMesh = MakeUnique<FMeshDescription>();
+	*EditedMesh = *UE::ToolTarget::GetMeshDescription(Target);
+
 	// Update the skeleton drawing information from the original bind pose
 	MaxDrawRadius = Component->Bounds.SphereRadius * 0.0025f;
 	UpdateBonePositionInfos(MaxDrawRadius);
@@ -389,10 +394,6 @@ void USkinWeightsPaintTool::Setup()
 	// attach callback to be informed when tool properties are modified
 	ToolProps->GetOnModified().AddUObject(this, &USkinWeightsPaintTool::OnToolPropertiesModified);
 
-	// create a mesh description for editing
-	EditedMesh = MakeUnique<FMeshDescription>();
-	*EditedMesh = *UE::ToolTarget::GetMeshDescription(Target);
-
 	// default to the root bone as current bone
 	PendingCurrentBone = CurrentBone = ToolProps->SkeletalMesh->GetRefSkeleton().GetBoneName(0);
 	ToolProps->CurrentBone.BoneName = PendingCurrentBone.GetValue();
