Add PolyGroupID option to the Bake Textures and Bake Vertex Colors tools

diff --git a/Engine/Plugins/Editor/ModelingToolsEditorMode/Source/ModelingToolsEditorMode/Private/DetailsCustomizations/BakeMeshAttributeToolCustomizations.cpp b/Engine/Plugins/Editor/ModelingToolsEditorMode/Source/ModelingToolsEditorMode/Private/DetailsCustomizations/BakeMeshAttributeToolCustomizations.cpp
--- a/Engine/Plugins/Editor/ModelingToolsEditorMode/Source/ModelingToolsEditorMode/Private/DetailsCustomizations/BakeMeshAttributeToolCustomizations.cpp
+++ b/Engine/Plugins/Editor/ModelingToolsEditorMode/Source/ModelingToolsEditorMode/Private/DetailsCustomizations/BakeMeshAttributeToolCustomizations.cpp
@@ -63,7 +63,8 @@ namespace BakeCustomizationLocals
 			// Mask
 			{
 				EBakeMapType::VertexColor,
-				EBakeMapType::MaterialID
+				EBakeMapType::MaterialID,
+				EBakeMapType::PolyGroupID
 			}
 		};
 
diff --git a/Engine/Plugins/Experimental/MeshModelingToolsetExp/Source/MeshModelingToolsExp/Private/BakeMeshAttributeMapsTool.cpp b/Engine/Plugins/Experimental/MeshModelingToolsetExp/Source/MeshModelingToolsExp/Private/BakeMeshAttributeMapsTool.cpp
--- a/Engine/Plugins/Experimental/MeshModelingToolsetExp/Source/MeshModelingToolsExp/Private/BakeMeshAttributeMapsTool.cpp
+++ b/Engine/Plugins/Experimental/MeshModelingToolsetExp/Source/MeshModelingToolsExp/Private/BakeMeshAttributeMapsTool.cpp
@@ -238,6 +238,13 @@ public:
 				Baker->AddEvaluator(PropertyEval);
 				break;
 			}
+			case EBakeMapType::PolyGroupID:
+			{
+				TSharedPtr<FMeshPropertyMapEvaluator, ESPMode::ThreadSafe> PropertyEval = MakeShared<FMeshPropertyMapEvaluator, ESPMode::ThreadSafe>();
+				PropertyEval->Property = EMeshPropertyMapType::PolyGroupID;
+				Baker->AddEvaluator(PropertyEval);
+				break;
+			}
 			case EBakeMapType::VertexColor:
 			{
 				TSharedPtr<FMeshPropertyMapEvaluator, ESPMode::ThreadSafe> PropertyEval = MakeShared<FMeshPropertyMapEvaluator, ESPMode::ThreadSafe>();
@@ -654,6 +661,7 @@ void UBakeMeshAttributeMapsTool::UpdateResult()
 		(bool)(CachedBakeSettings.BakeMapTypes & EBakeMapType::FaceNormal) ||
 		(bool)(CachedBakeSettings.BakeMapTypes & EBakeMapType::Position) ||
 		(bool)(CachedBakeSettings.BakeMapTypes & EBakeMapType::MaterialID) ||
+		(bool)(CachedBakeSettings.BakeMapTypes & EBakeMapType::PolyGroupID) ||
 		(bool)(CachedBakeSettings.BakeMapTypes & EBakeMapType::VertexColor))
 	{
 		OpState |= UpdateResult_MeshProperty(CachedBakeSettings.Dimensions);
@@ -784,6 +792,7 @@ void UBakeMeshAttributeMapsTool::UpdateOnModeChange()
 		case EBakeMapType::FaceNormal:
 		case EBakeMapType::Position:
 		case EBakeMapType::MaterialID:
+		case EBakeMapType::PolyGroupID:
 		case EBakeMapType::VertexColor:
 			break;
 		case EBakeMapType::Texture:
diff --git a/Engine/Plugins/Experimental/MeshModelingToolsetExp/Source/MeshModelingToolsExp/Private/BakeMeshAttributeMapsToolBase.cpp b/Engine/Plugins/Experimental/MeshModelingToolsetExp/Source/MeshModelingToolsExp/Private/BakeMeshAttributeMapsToolBase.cpp
--- a/Engine/Plugins/Experimental/MeshModelingToolsetExp/Source/MeshModelingToolsExp/Private/BakeMeshAttributeMapsToolBase.cpp
+++ b/Engine/Plugins/Experimental/MeshModelingToolsetExp/Source/MeshModelingToolsExp/Private/BakeMeshAttributeMapsToolBase.cpp
@@ -265,6 +265,7 @@ void UBakeMeshAttributeMapsToolBase::UpdatePreview(const EBakeMapType PreviewMap
 		case EBakeMapType::FaceNormal:
 		case EBakeMapType::Position:
 		case EBakeMapType::MaterialID:
+		case EBakeMapType::PolyGroupID:
 		case EBakeMapType::Texture:
 		case EBakeMapType::MultiTexture:
 		case EBakeMapType::VertexColor:
@@ -444,6 +445,9 @@ void UBakeMeshAttributeMapsToolBase::OnMapsUpdated(const TUniquePtr<UE::Geometry
 			case EMeshPropertyMapType::MaterialID:
 				MapType = EBakeMapType::MaterialID;
 				break;
+			case EMeshPropertyMapType::PolyGroupID:
+				MapType = EBakeMapType::PolyGroupID;
+				break;
 			case EMeshPropertyMapType::VertexColor:
 				MapType = EBakeMapType::VertexColor;
 				break;
@@ -514,6 +518,7 @@ FTexture2DBuilder::ETextureType UBakeMeshAttributeMapsToolBase::GetTextureType(c
 		TexType = FTexture2DBuilder::ETextureType::ColorLinear;
 		break;
 	case EBakeMapType::MaterialID:
+	case EBakeMapType::PolyGroupID:
 	case EBakeMapType::VertexColor:
 		break;
 	case EBakeMapType::Texture:
@@ -557,6 +562,9 @@ void UBakeMeshAttributeMapsToolBase::GetTextureName(const EBakeMapType MapType,
 	case EBakeMapType::MaterialID:
 		TexName = FString::Printf(TEXT("%s_MaterialID"), *BaseName);
 		break;
+	case EBakeMapType::PolyGroupID:
+		TexName = FString::Printf(TEXT("%s_PolyGroupID"), *BaseName);
+		break;
 	case EBakeMapType::VertexColor:
 		TexName = FString::Printf(TEXT("%s_VertexColor"), *BaseName);
 		break;
@@ -719,6 +727,9 @@ void UBakeMeshAttributeMapsToolBase::RecordAnalytics(const FBakeAnalytics& Data,
 
 	const bool bMaterialID = static_cast<bool>(Data.BakeSettings.SourceBakeMapTypes & EBakeMapType::MaterialID);
 	Attributes.Add(FAnalyticsEventAttribute(TEXT("Settings.MaterialID.Enabled"), bMaterialID));
+
+	const bool bPolyGroupID = static_cast<bool>(Data.BakeSettings.SourceBakeMapTypes & EBakeMapType::PolyGroupID);
+	Attributes.Add(FAnalyticsEventAttribute(TEXT("Settings.PolyGroupID.Enabled"), bPolyGroupID));
 	
 	const bool bTexture = static_cast<bool>(Data.BakeSettings.SourceBakeMapTypes & EBakeMapType::Texture);
 	Attributes.Add(FAnalyticsEventAttribute(TEXT("Settings.Texture.Enabled"), bTexture));
diff --git a/Engine/Plugins/Experimental/MeshModelingToolsetExp/Source/MeshModelingToolsExp/Private/BakeMeshAttributeVertexTool.cpp b/Engine/Plugins/Experimental/MeshModelingToolsetExp/Source/MeshModelingToolsExp/Private/BakeMeshAttributeVertexTool.cpp
--- a/Engine/Plugins/Experimental/MeshModelingToolsetExp/Source/MeshModelingToolsExp/Private/BakeMeshAttributeVertexTool.cpp
+++ b/Engine/Plugins/Experimental/MeshModelingToolsetExp/Source/MeshModelingToolsExp/Private/BakeMeshAttributeVertexTool.cpp
@@ -201,6 +201,13 @@ public:
 				Baker->ColorEvaluator = PropertyEval;
 				break;
 			}
+			case EBakeMapType::PolyGroupID:
+			{
+				TSharedPtr<FMeshPropertyMapEvaluator, ESPMode::ThreadSafe> PropertyEval = MakeShared<FMeshPropertyMapEvaluator, ESPMode::ThreadSafe>();
+				PropertyEval->Property = EMeshPropertyMapType::PolyGroupID;
+				Baker->ColorEvaluator = PropertyEval;
+				break;
+			}
 			case EBakeMapType::Texture:
 			{
 				TSharedPtr<FMeshResampleImageEvaluator, ESPMode::ThreadSafe> TextureEval = MakeShared<FMeshResampleImageEvaluator, ESPMode::ThreadSafe>();
@@ -700,6 +707,7 @@ void UBakeMeshAttributeVertexTool::UpdateResult()
 		case EBakeMapType::FaceNormal:
 		case EBakeMapType::Position:
 		case EBakeMapType::MaterialID:
+		case EBakeMapType::PolyGroupID:
 			OpState |= UpdateResult_MeshProperty(Dimensions);
 			break;
 		case EBakeMapType::VertexColor:
diff --git a/Engine/Plugins/Experimental/MeshModelingToolsetExp/Source/MeshModelingToolsExp/Private/BakeMultiMeshAttributeMapsTool.cpp b/Engine/Plugins/Experimental/MeshModelingToolsetExp/Source/MeshModelingToolsExp/Private/BakeMultiMeshAttributeMapsTool.cpp
--- a/Engine/Plugins/Experimental/MeshModelingToolsetExp/Source/MeshModelingToolsExp/Private/BakeMultiMeshAttributeMapsTool.cpp
+++ b/Engine/Plugins/Experimental/MeshModelingToolsetExp/Source/MeshModelingToolsExp/Private/BakeMultiMeshAttributeMapsTool.cpp
@@ -241,6 +241,13 @@ public:
 		return IndexConstants::InvalidID;
 	}
 
+	virtual int32 GetPolyGroupID(const void* Mesh, const int TriId) const override
+	{
+		// TODO
+		checkSlow(false);
+		return 0;
+	}
+
 	virtual bool HasNormals(const void* Mesh) const override
 	{
 		const FActorChildMesh* ChildMesh = static_cast<const FActorChildMesh*>(Mesh);
diff --git a/Engine/Plugins/Experimental/MeshModelingToolsetExp/Source/MeshModelingToolsExp/Public/BakeMeshAttributeMapsTool.h b/Engine/Plugins/Experimental/MeshModelingToolsetExp/Source/MeshModelingToolsExp/Public/BakeMeshAttributeMapsTool.h
--- a/Engine/Plugins/Experimental/MeshModelingToolsetExp/Source/MeshModelingToolsExp/Public/BakeMeshAttributeMapsTool.h
+++ b/Engine/Plugins/Experimental/MeshModelingToolsetExp/Source/MeshModelingToolsExp/Public/BakeMeshAttributeMapsTool.h
@@ -54,7 +54,7 @@ class MESHMODELINGTOOLSEXP_API UBakeMeshAttributeMapsToolProperties : public UIn
 public:
 	/** The bake output types to generate */
 	UPROPERTY(EditAnywhere, Category = BakeOutput, meta = (DisplayName = "Output Types", Bitmask, BitmaskEnum = "/Script/MeshModelingToolsExp.EBakeMapType",
-		ValidEnumValues="TangentSpaceNormal, AmbientOcclusion, BentNormal, Curvature, Texture, ObjectSpaceNormal, FaceNormal, Position, MaterialID, MultiTexture, VertexColor"))
+		ValidEnumValues="TangentSpaceNormal, AmbientOcclusion, BentNormal, Curvature, Texture, ObjectSpaceNormal, FaceNormal, Position, MaterialID, PolyGroupID, MultiTexture, VertexColor"))
 	int32 MapTypes = static_cast<int32>(EBakeMapType::None);
 
 	/** The baked output type used for preview in the viewport */
diff --git a/Engine/Plugins/Experimental/MeshModelingToolsetExp/Source/MeshModelingToolsExp/Public/BakeMeshAttributeTool.h b/Engine/Plugins/Experimental/MeshModelingToolsetExp/Source/MeshModelingToolsExp/Public/BakeMeshAttributeTool.h
--- a/Engine/Plugins/Experimental/MeshModelingToolsetExp/Source/MeshModelingToolsExp/Public/BakeMeshAttributeTool.h
+++ b/Engine/Plugins/Experimental/MeshModelingToolsetExp/Source/MeshModelingToolsExp/Public/BakeMeshAttributeTool.h
@@ -47,6 +47,8 @@ enum class EBakeMapType
 	VertexColor            = 1 << 9,
 	/* Material IDs as unique colors */
 	MaterialID             = 1 << 10 UMETA(DisplayName = "Material ID"),
+	/* PolyGroup IDs as unique colors */
+	PolyGroupID            = 1 << 11 UMETA(DisplayName = "PolyGroup ID"),
 
 	All                    = 0x7FF UMETA(Hidden)
 };
@@ -69,7 +71,8 @@ static constexpr EBakeMapType ENUM_EBAKEMAPTYPE_ALL[] =
 	EBakeMapType::Texture,
 	EBakeMapType::MultiTexture,
 	EBakeMapType::VertexColor,
-	EBakeMapType::MaterialID
+	EBakeMapType::MaterialID,
+	EBakeMapType::PolyGroupID
 };
 
 
diff --git a/Engine/Plugins/Experimental/MeshModelingToolsetExp/Source/MeshModelingToolsExp/Public/BakeMeshAttributeVertexTool.h b/Engine/Plugins/Experimental/MeshModelingToolsetExp/Source/MeshModelingToolsExp/Public/BakeMeshAttributeVertexTool.h
--- a/Engine/Plugins/Experimental/MeshModelingToolsetExp/Source/MeshModelingToolsExp/Public/BakeMeshAttributeVertexTool.h
+++ b/Engine/Plugins/Experimental/MeshModelingToolsetExp/Source/MeshModelingToolsExp/Public/BakeMeshAttributeVertexTool.h
@@ -65,7 +65,7 @@ public:
 
 	/** The bake output type to generate */
 	UPROPERTY(EditAnywhere, Category = BakeOutput, meta=(Bitmask, BitmaskEnum = "/Script/MeshModelingToolsExp.EBakeMapType",
-		ValidEnumValues="TangentSpaceNormal, AmbientOcclusion, BentNormal, Curvature, Texture, ObjectSpaceNormal, FaceNormal, Position, MaterialID, MultiTexture, VertexColor",
+		ValidEnumValues="TangentSpaceNormal, AmbientOcclusion, BentNormal, Curvature, Texture, ObjectSpaceNormal, FaceNormal, Position, MaterialID, PolyGroupID, MultiTexture, VertexColor",
 		EditCondition="OutputMode == EBakeVertexOutput::RGBA", EditConditionHides))
 	int32 OutputType = static_cast<int32>(EBakeMapType::TangentSpaceNormal);
 
diff --git a/Engine/Plugins/Runtime/GeometryProcessing/Source/DynamicMesh/Private/Sampling/MeshPropertyMapEvaluator.cpp b/Engine/Plugins/Runtime/GeometryProcessing/Source/DynamicMesh/Private/Sampling/MeshPropertyMapEvaluator.cpp
--- a/Engine/Plugins/Runtime/GeometryProcessing/Source/DynamicMesh/Private/Sampling/MeshPropertyMapEvaluator.cpp
+++ b/Engine/Plugins/Runtime/GeometryProcessing/Source/DynamicMesh/Private/Sampling/MeshPropertyMapEvaluator.cpp
@@ -63,6 +63,7 @@ void FMeshPropertyMapEvaluator::Setup(const FMeshBaseBaker& Baker, FEvaluationCo
 		DefaultValue = UVToColor(FVector2f::Zero());
 		break;
 	case EMeshPropertyMapType::MaterialID:
+	case EMeshPropertyMapType::PolyGroupID:
 		DefaultValue = FVector3f(LinearColors::LightPink3f());
 		Context.AccumulateMode = EAccumulateMode::Overwrite;
 		break;
@@ -202,6 +203,12 @@ FVector3f FMeshPropertyMapEvaluator::SampleFunction(const FCorrespondenceSample&
 		}
 	}
 	break;
+	case EMeshPropertyMapType::PolyGroupID:
+	{
+		const int32 GroupID = DetailSampler->GetPolyGroupID(DetailMesh, DetailTriID);
+		Color = LinearColors::SelectColor<FVector3f>(GroupID);
+	}
+	break;
 	}
 	return Color;
 }
diff --git a/Engine/Plugins/Runtime/GeometryProcessing/Source/DynamicMesh/Public/Sampling/MeshBakerCommon.h b/Engine/Plugins/Runtime/GeometryProcessing/Source/DynamicMesh/Public/Sampling/MeshBakerCommon.h
--- a/Engine/Plugins/Runtime/GeometryProcessing/Source/DynamicMesh/Public/Sampling/MeshBakerCommon.h
+++ b/Engine/Plugins/Runtime/GeometryProcessing/Source/DynamicMesh/Public/Sampling/MeshBakerCommon.h
@@ -245,6 +245,13 @@ public:
 	 */
 	virtual int32 GetMaterialID(const void* Mesh, const int TriId) const = 0;
 
+	/**
+	 * @param Mesh pointer to mesh to query
+	 * @param TriId the triangle index to test
+	 * @return the group ID for the given triangle.
+	 */
+	virtual int32 GetPolyGroupID(const void* Mesh, const int TriId) const = 0;
+
 	/**
 	 * @param Mesh pointer to mesh to query
 	 * @return true if this mesh has normals
@@ -489,6 +496,12 @@ public:
 		return MaterialIDOverlay ? MaterialIDOverlay->GetValue(TriId) : IndexConstants::InvalidID;
 	}
 
+	virtual int32 GetPolyGroupID(const void* Mesh, const int TriID) const override
+	{
+		const FDynamicMesh3* DynamicMesh = static_cast<const FDynamicMesh3*>(Mesh);
+		return DynamicMesh->GetTriangleGroup(TriID);
+	}
+
 	virtual bool HasNormals(const void* Mesh) const override
 	{
 		const FDynamicMesh3* DynamicMesh = static_cast<const FDynamicMesh3*>(Mesh);
diff --git a/Engine/Plugins/Runtime/GeometryProcessing/Source/DynamicMesh/Public/Sampling/MeshPropertyMapEvaluator.h b/Engine/Plugins/Runtime/GeometryProcessing/Source/DynamicMesh/Public/Sampling/MeshPropertyMapEvaluator.h
--- a/Engine/Plugins/Runtime/GeometryProcessing/Source/DynamicMesh/Public/Sampling/MeshPropertyMapEvaluator.h
+++ b/Engine/Plugins/Runtime/GeometryProcessing/Source/DynamicMesh/Public/Sampling/MeshPropertyMapEvaluator.h
@@ -20,7 +20,8 @@ enum class EMeshPropertyMapType
 	FacetNormal = 3,
 	UVPosition = 4,
 	MaterialID = 5,
-	VertexColor = 6
+	VertexColor = 6,
+	PolyGroupID = 7
 };
 
 /**
