CombineMeshInstances: add alternate hidden-removal method, enable by default, expose options for hidden removal and normals

diff --git a/Engine/Plugins/Experimental/MeshModelingToolsetExp/Source/GeometryProcessingAdapters/Private/CombineMeshInstancesImpl.cpp b/Engine/Plugins/Experimental/MeshModelingToolsetExp/Source/GeometryProcessingAdapters/Private/CombineMeshInstancesImpl.cpp
--- a/Engine/Plugins/Experimental/MeshModelingToolsetExp/Source/GeometryProcessingAdapters/Private/CombineMeshInstancesImpl.cpp
+++ b/Engine/Plugins/Experimental/MeshModelingToolsetExp/Source/GeometryProcessingAdapters/Private/CombineMeshInstancesImpl.cpp
@@ -33,6 +33,8 @@
 #include "Physics/PhysicsDataCollection.h"
 
 #include "TransformSequence.h"
+#include "Sampling/SphericalFibonacci.h"
+#include "Util/IteratorUtil.h"
 
 
 using namespace UE::Geometry;
@@ -41,7 +43,7 @@ using namespace UE::Geometry;
 
 static TAutoConsoleVariable<int> CVarGeometryCombineMeshInstancesRemoveHidden(
 	TEXT("geometry.CombineInstances.DebugRemoveHiddenStrategy"),
-	0,
+	1,
 	TEXT("Configure hidden-removal strategy via (temporary debug)"));
 
 
@@ -685,7 +687,7 @@ void ComputeMeshApproximations(
 	FMeshInstanceAssembly& Assembly)
 {
 	using namespace UE::Geometry;
-	const double AngleThreshold = 45.0;
+	const double AngleThresholdDeg = CombineOptions.HardNormalAngleDeg;
 
 	int32 NumSets = Assembly.InstanceSets.Num();
 	Assembly.OptimizedMeshGeometry.SetNum(NumSets);
@@ -710,7 +712,7 @@ void ComputeMeshApproximations(
 		for (int32 k = 0; k < NumSimplifiedLODs; ++k)
 		{
 			ApproxGeo.SimplifiedMeshLODs[k] = OptimizationSourceMesh;
-			SimplifyPartMesh(ApproxGeo.SimplifiedMeshLODs[k], InitialTolerance, AngleThreshold);
+			SimplifyPartMesh(ApproxGeo.SimplifiedMeshLODs[k], InitialTolerance, AngleThresholdDeg);
 			InitialTolerance *= CombineOptions.SimplifyLODLevelToleranceScale;
 		}
 
@@ -727,7 +729,7 @@ void ComputeMeshApproximations(
 			ApproxGeo.ApproximateMeshLODs[k].EnableMatchingAttributes(OptimizationSourceMesh);
 
 			// recompute normals
-			FMeshNormals::InitializeOverlayTopologyFromOpeningAngle(&ApproxGeo.ApproximateMeshLODs[k], ApproxGeo.ApproximateMeshLODs[k].Attributes()->PrimaryNormals(), AngleThreshold);
+			FMeshNormals::InitializeOverlayTopologyFromOpeningAngle(&ApproxGeo.ApproximateMeshLODs[k], ApproxGeo.ApproximateMeshLODs[k].Attributes()->PrimaryNormals(), AngleThresholdDeg);
 			FMeshNormals::QuickRecomputeOverlayNormals(ApproxGeo.ApproximateMeshLODs[k]);
 		}
 
@@ -740,8 +742,8 @@ void ComputeMeshApproximations(
 }
 
 
-
-static void RemoveHiddenFaces(FDynamicMesh3& EditMesh, double MaxDistance = 200)
+// Remove hidden faces by (approximately) computing Ambient Occlusion, fully occluded faces are hidden
+static void RemoveHiddenFaces_Occlusion(FDynamicMesh3& EditMesh, double MaxDistance = 200)
 {
 	TRemoveOccludedTriangles<FDynamicMesh3> Jacket(&EditMesh);
 
@@ -780,6 +782,205 @@ static void RemoveHiddenFaces(FDynamicMesh3& EditMesh, double MaxDistance = 200)
 
 
 
+
+// Remove hidden faces by casting rays from exterior at sample points on triangles
+// (This method works quite well and should eventually be extracted out to a general algorithm...)
+static void RemoveHiddenFaces_ExteriorVisibility(FDynamicMesh3& TargetMesh, double SampleRadius)
+{
+	FDynamicMeshAABBTree3 Spatial(&TargetMesh, true);
+	FAxisAlignedBox3d Bounds = Spatial.GetBoundingBox();
+	double Radius = Bounds.DiagonalLength() * 0.5;
+
+
+	auto FindHitTriangleTest = [&](FVector3d TargetPosition, FVector3d FarPosition) -> int
+	{
+		FVector3d RayDir(TargetPosition - FarPosition);
+		double Distance = Normalize(RayDir);
+		FRay3d Ray(FarPosition, RayDir, true);
+		return Spatial.FindNearestHitTriangle(Ray, IMeshSpatial::FQueryOptions(Distance + 1.0));	// 1.0 is random fudge factor here...
+	};
+
+
+	// final triangle visibility, atomics can be updated on any thread
+	TArray<std::atomic<bool>> ThreadSafeTriVisible;
+	ThreadSafeTriVisible.SetNum(TargetMesh.MaxTriangleID());
+	for (int32 tid : TargetMesh.TriangleIndicesItr())
+	{
+		ThreadSafeTriVisible[tid] = false;
+	}
+
+	// array of (+/-)X/Y/Z directions
+	TArray<FVector3d> CardinalDirections;
+	for (int32 k = 0; k < 3; ++k)
+	{
+		FVector3d Direction(0,0,0);
+		Direction[k] = 1.0;
+		CardinalDirections.Add(Direction);
+		CardinalDirections.Add(-Direction);
+	}
+
+	//
+	// First pass. For each triangle, cast a ray at it's centroid from 
+	// outside the model, along the X/Y/Z directions and tri normal.
+	// If tri is hit we mark it as having 'known' status, allowing it
+	// to be skipped in the more expensive pass below
+	//
+	TArray<bool> TriStatusKnown;
+	TriStatusKnown.Init(false, TargetMesh.MaxTriangleID());
+	ParallelFor(TargetMesh.MaxTriangleID(), [&](int32 tid)
+	{
+		FVector3d Normal, Centroid; double Area;
+		TargetMesh.GetTriInfo(tid, Normal, Area, Centroid);
+		if (Normal.SquaredLength() < 0.1 || Area < FMathf::ZeroTolerance)
+		{
+			TriStatusKnown[tid] = true;
+			return;
+		}	
+
+		for (FVector3d Direction : CardinalDirections)
+		{
+			if (FindHitTriangleTest(Centroid, Centroid + Radius*Direction) == tid)
+			{
+				ThreadSafeTriVisible[tid] = true;
+				TriStatusKnown[tid] = true;
+				return;
+			}
+		}
+		if (FindHitTriangleTest(Centroid, Centroid + Radius*Normal) == tid)
+		{
+				ThreadSafeTriVisible[tid] = true;
+				TriStatusKnown[tid] = true;
+				return;
+		}
+
+		// triangle is not definitely visible or hidden
+	});
+
+
+	//
+	// Construct set of exterior sample points, for each triangle sample point
+	// below we will check if it is visible from any of these sample points.
+	// Order is shuffled in hopes that for visible tris we don't waste a bunch
+	// of time on the 'far' side
+	//
+	int32 NumExteriorSamplePoints = 128;
+	TSphericalFibonacci<double> SphereSampler(NumExteriorSamplePoints);
+	TArray<FVector3d> ExteriorSamplePoints;
+	FModuloIteration ModuloIter(NumExteriorSamplePoints);
+	uint32 SampleIndex = 0;
+	while (ModuloIter.GetNextIndex(SampleIndex))
+	{
+		ExteriorSamplePoints.Add( Bounds.Center() + Radius * SphereSampler[SampleIndex] );
+	}
+	// add axis directions?
+
+
+	//
+	// For each triangle, generate a set of sample points on the triangle surface,
+	// and then check if that point is visible from any of the exterior sample points.
+	// This is the expensive part!
+	// 
+	// Does using a fixed set of exterior sample points make sense? Could also 
+	// treat it as a set of sample directions. Seems more likely to hit tri
+	// based on sample directions...
+	//
+	ParallelFor(TargetMesh.MaxTriangleID(), [&](int32 tid)
+	{
+		// if we already found out this triangle is visible or hidden, we can skip it
+		if ( TriStatusKnown[tid] || ThreadSafeTriVisible[tid] ) return;
+
+		FVector3d A,B,C;
+		TargetMesh.GetTriVertices(tid, A,B,C);
+		FVector3d Centroid = (A + B + C) / 3.0;
+		double TriArea;
+		FVector3d TriNormal = VectorUtil::NormalArea(A, B, C, TriArea);		// TriStatusKnown should skip degen tris, do not need to check here
+
+		FFrame3d TriFrame(Centroid, TriNormal);
+		FTriangle2d UVTriangle(TriFrame.ToPlaneUV(A), TriFrame.ToPlaneUV(B), TriFrame.ToPlaneUV(C));
+		double DiscArea = (FMathd::Pi * SampleRadius * SampleRadius);
+		int NumSamples = FMath::Max( (int)(TriArea / DiscArea), 2 );  // a bit arbitrary...
+		FVector2d V1 = UVTriangle.V[1] - UVTriangle.V[0];
+		FVector2d V2 = UVTriangle.V[2] - UVTriangle.V[0];
+
+		TArray<int32> HitTris;		// re-use this array in inner loop to avoid hitting atomics so often
+
+		int NumTested = 0;
+		FRandomStream RandomStream(tid);
+		while (NumTested < NumSamples)
+		{
+			double a1 = RandomStream.GetFraction();
+			double a2 = RandomStream.GetFraction();
+			FVector2d PointUV = UVTriangle.V[0] + a1 * V1 + a2 * V2;
+			if (UVTriangle.IsInside(PointUV))
+			{
+				NumTested++;
+				FVector3d Position = TriFrame.FromPlaneUV(PointUV, 2);
+
+				// cast ray from all exterior sample locations for this triangle sample point
+				HitTris.Reset();
+				for (int32 k = 0; k < NumExteriorSamplePoints; ++k)
+				{
+					int32 HitTriID = FindHitTriangleTest(Position, ExteriorSamplePoints[k]);
+					if ( HitTriID != IndexConstants::InvalidID && TriStatusKnown[HitTriID] == false )
+					{
+						HitTris.AddUnique(HitTriID);		// we hit some triangle, whether or not it is the one we are testing...
+						if (HitTriID == tid)
+						{
+							break;
+						}
+					}
+				}
+
+				// mark any hit tris
+				for ( int32 HitTriID : HitTris )
+				{
+					ThreadSafeTriVisible[HitTriID] = true;
+				}
+
+				// if our triangle has become visible (in this thread or another) we can terminate now
+				if (ThreadSafeTriVisible[tid])
+				{
+					return;
+				}
+			}
+		}
+
+		// should we at any point lock and update TriStatusKnown?
+	});
+
+	// delete hidden tris
+	TArray<int32> TrisToDelete;
+	for (int32 tid : TargetMesh.TriangleIndicesItr())
+	{
+		if (ThreadSafeTriVisible[tid] == false)
+		{
+			TrisToDelete.Add(tid);
+		}
+	}
+	FDynamicMeshEditor Editor(&TargetMesh);
+	Editor.RemoveTriangles(TrisToDelete, true);
+
+	TargetMesh.CompactInPlace();
+}
+
+
+
+
+static void PostProcessHiddenFaceRemovedMesh(FDynamicMesh3& TargetMesh, double Tolerance)
+{
+	// weld edges in case input was unwelded...
+	FMergeCoincidentMeshEdges Welder(&TargetMesh);
+	Welder.MergeVertexTolerance = Tolerance * 0.001;
+	Welder.OnlyUniquePairs = false;
+	Welder.Apply();
+
+	// todo: try to simplify? need to be able to constrain by things like vertex color...
+
+	TargetMesh.CompactInPlace();
+}
+
+
+
 namespace UE::Geometry
 {
 
@@ -881,12 +1082,12 @@ void BuildCombinedMesh(
 				bool bIsDecorativePart = (Instance.DetailLevel == EMeshDetailLevel::Decorative);
 
 				// filter out detail parts at higher LODs
-				if (bIsDecorativePart && LODLevel > CombineOptions.FilterDecorativePartsLODLevel)
+				if (bIsDecorativePart && LODLevel >= CombineOptions.FilterDecorativePartsLODLevel)
 				{
 					continue;
 				}
 				// at last detail part LOD, switch to approximate mesh
-				if (bIsDecorativePart && LODLevel == CombineOptions.FilterDecorativePartsLODLevel)
+				if (bIsDecorativePart && LODLevel >= (CombineOptions.FilterDecorativePartsLODLevel - CombineOptions.ApproximateDecorativePartLODs) )
 				{
 					UseAppendMesh = ApproximateAppendMesh;
 				}
@@ -922,11 +1123,36 @@ void BuildCombinedMesh(
 		}
 	}
 
-	if (CVarGeometryCombineMeshInstancesRemoveHidden.GetValueOnGameThread() > 0)
+
+	//
+	// Hidden-face removal pass
+	//
+	if (CombineOptions.RemoveHiddenFacesMethod != IGeometryProcessing_CombineMeshInstances::ERemoveHiddenFacesMode::None 
+		&& CVarGeometryCombineMeshInstancesRemoveHidden.GetValueOnGameThread() > 0)
 	{
 		ParallelFor(MeshLODs.Num(), [&](int32 k)
 		{
-			RemoveHiddenFaces(MeshLODs[k].Mesh, 200.0);
+			bool bModified = false;
+			if ( k >= CombineOptions.RemoveHiddenStartLOD )
+			{ 
+				switch (CombineOptions.RemoveHiddenFacesMethod)
+				{
+					case IGeometryProcessing_CombineMeshInstances::ERemoveHiddenFacesMode::OcclusionBased:
+						RemoveHiddenFaces_Occlusion(MeshLODs[k].Mesh, 200);		// 200 is arbitrary here! should improve once max-distance is actually available (currently ignored)
+						bModified = true;
+						break;
+					case IGeometryProcessing_CombineMeshInstances::ERemoveHiddenFacesMode::ExteriorVisibility:
+					case IGeometryProcessing_CombineMeshInstances::ERemoveHiddenFacesMode::Fastest:
+						RemoveHiddenFaces_ExteriorVisibility(MeshLODs[k].Mesh, CombineOptions.RemoveHiddenSamplingDensity);
+						bModified = true;
+						break;
+				}
+			}
+
+			if ( bModified )
+			{
+				PostProcessHiddenFaceRemovedMesh(MeshLODs[k].Mesh, CombineOptions.SimplifyBaseTolerance);
+			}
 		});
 	}
 
@@ -1125,6 +1351,8 @@ IGeometryProcessing_CombineMeshInstances::FOptions FCombineMeshInstancesImpl::Co
 	Options.FilterDecorativePartsLODLevel = 2;
 
 
+	Options.RemoveHiddenFacesMethod = ERemoveHiddenFacesMode::Fastest;
+
 	return Options;
 }
 
diff --git a/Engine/Source/Developer/GeometryProcessingInterfaces/Public/GeometryProcessingInterfaces/CombineMeshInstances.h b/Engine/Source/Developer/GeometryProcessingInterfaces/Public/GeometryProcessingInterfaces/CombineMeshInstances.h
--- a/Engine/Source/Developer/GeometryProcessingInterfaces/Public/GeometryProcessingInterfaces/CombineMeshInstances.h
+++ b/Engine/Source/Developer/GeometryProcessingInterfaces/Public/GeometryProcessingInterfaces/CombineMeshInstances.h
@@ -67,6 +67,15 @@ public:
 	};
 
 
+	enum class ERemoveHiddenFacesMode
+	{
+		None = 0,
+		Fastest = 1,
+
+		ExteriorVisibility = 5,
+		OcclusionBased = 6
+	};
+
 	struct FOptions
 	{
 		// number of requested LODs
@@ -83,8 +92,25 @@ public:
 		double OptimizeBaseTriCost = 0.7;
 		double OptimizeLODLevelTriCostScale = 1.5;
 
+		//
+		// Hidden Faces removal options
+		// 
+		
+		// overall strategy to use for removing hidden faces
+		ERemoveHiddenFacesMode RemoveHiddenFacesMethod = ERemoveHiddenFacesMode::None;
+		// start removing hidden faces at this LOD level 
+		int32 RemoveHiddenStartLOD = 0;
+		// (approximately) spacing between samples on triangle faces used for determining exterior visibility
+		double RemoveHiddenSamplingDensity = 1.0;
+
+
 		// LOD level to filter out detail parts
 		int32 FilterDecorativePartsLODLevel = 2;
+		// Decorative part will be approximated by simple shape for this many LOD levels before Filter level
+		int ApproximateDecorativePartLODs = 1;
+
+		// opening angle used to detect/assign sharp edges
+		double HardNormalAngleDeg = 15.0;
 	};
 
 
