CombineMeshInstances: add support for voxelization-based LODs with material reprojection. Currently hardcoded to also reproject per-triangle colors.

diff --git a/Engine/Plugins/Experimental/MeshModelingToolsetExp/Source/GeometryProcessingAdapters/Private/CombineMeshInstancesImpl.cpp b/Engine/Plugins/Experimental/MeshModelingToolsetExp/Source/GeometryProcessingAdapters/Private/CombineMeshInstancesImpl.cpp
--- a/Engine/Plugins/Experimental/MeshModelingToolsetExp/Source/GeometryProcessingAdapters/Private/CombineMeshInstancesImpl.cpp
+++ b/Engine/Plugins/Experimental/MeshModelingToolsetExp/Source/GeometryProcessingAdapters/Private/CombineMeshInstancesImpl.cpp
@@ -36,6 +36,8 @@
 #include "Sampling/SphericalFibonacci.h"
 #include "Util/IteratorUtil.h"
 
+#include "Implicit/Morphology.h"
+#include "ProjectionTargets.h"
 
 using namespace UE::Geometry;
 
@@ -682,6 +684,10 @@ static void SelectBestFittingMeshApproximation(
 
 
 
+
+
+
+
 void ComputeMeshApproximations(
 	IGeometryProcessing_CombineMeshInstances::FOptions CombineOptions,
 	FMeshInstanceAssembly& Assembly)
@@ -987,6 +993,115 @@ namespace UE::Geometry
 
 
 
+static void ComputeVoxWrapMesh(
+	const FDynamicMesh3& CombinedMesh, 
+	FDynamicMeshAABBTree3& CombinedMeshSpatial,
+	FDynamicMesh3& ResultMesh,
+	double ClosureDistance,
+	double TargetCellSize)
+{
+	TImplicitMorphology<FDynamicMesh3> Morphology;
+	Morphology.Source = &CombinedMesh;
+	Morphology.SourceSpatial = &CombinedMeshSpatial;
+	Morphology.MorphologyOp = TImplicitMorphology<FDynamicMesh3>::EMorphologyOp::Close;
+	Morphology.Distance = FMath::Max(ClosureDistance, 0.001);
+
+	FAxisAlignedBox3d Bounds = CombinedMeshSpatial.GetBoundingBox();
+	double UseCellSize = FMath::Max(0.001, TargetCellSize);
+	int MaxGridDimEstimate = (int32)(Bounds.MaxDim() / UseCellSize);
+	if (MaxGridDimEstimate > 128)
+	{
+		UseCellSize = (float)Bounds.MaxDim() / 128.0;
+	}
+	Morphology.GridCellSize = UseCellSize;
+
+	ResultMesh.Copy(&Morphology.Generate());
+	ResultMesh.DiscardAttributes();
+}
+
+static void ComputeSimplifiedVoxWrapMesh(
+	const FDynamicMesh3& CombinedMesh, 
+	FDynamicMeshAABBTree3& CombinedMeshSpatial,
+	FDynamicMesh3& VoxWrapMesh,
+	double SimplifyTolerance,
+	double MaxTriCount)
+{
+	FVolPresMeshSimplification Simplifier(&VoxWrapMesh);
+
+	Simplifier.ProjectionMode = FVolPresMeshSimplification::ETargetProjectionMode::NoProjection;
+
+	//FMeshProjectionTarget ProjectionTarget(&CombinedMesh, &CombinedMeshSpatial);
+	//Simplifier.SetProjectionTarget(&ProjectionTarget);
+
+	Simplifier.DEBUG_CHECK_LEVEL = 0;
+	Simplifier.bRetainQuadricMemory = false; 
+
+	//Simplifier.GeometricErrorConstraint = FVolPresMeshSimplification::EGeometricErrorCriteria::PredictedPointToProjectionTarget;
+	//Simplifier.GeometricErrorTolerance = SimplifyTolerance;
+
+	//Simplifier.SimplifyToTriangleCount( 1 );
+
+	if (VoxWrapMesh.TriangleCount() > MaxTriCount)
+	{
+		//Simplifier.SetProjectionTarget(nullptr);
+		//Simplifier.GeometricErrorConstraint = FVolPresMeshSimplification::EGeometricErrorCriteria::None;
+		Simplifier.SimplifyToTriangleCount( MaxTriCount );
+	}
+
+	// compact result
+	VoxWrapMesh.CompactInPlace();
+
+	VoxWrapMesh.EnableTriangleGroups();
+	VoxWrapMesh.EnableAttributes();
+	FMeshNormals::InitializeOverlayToPerVertexNormals(VoxWrapMesh.Attributes()->PrimaryNormals(), false);
+
+	const FDynamicMeshColorOverlay* SourceColors = nullptr;
+	FDynamicMeshColorOverlay* TargetColors = nullptr;
+	if ( CombinedMesh.HasAttributes() && CombinedMesh.Attributes()->HasPrimaryColors() )
+	{
+		SourceColors = CombinedMesh.Attributes()->PrimaryColors();
+		VoxWrapMesh.Attributes()->EnablePrimaryColors();
+		TargetColors = VoxWrapMesh.Attributes()->PrimaryColors();
+	}
+
+	const FDynamicMeshMaterialAttribute* SourceMaterialID = nullptr;
+	FDynamicMeshMaterialAttribute* TargetMaterialID = nullptr;
+	if ( CombinedMesh.HasAttributes() && CombinedMesh.Attributes()->HasMaterialID() )
+	{
+		SourceMaterialID = CombinedMesh.Attributes()->GetMaterialID();
+		VoxWrapMesh.Attributes()->EnableMaterialID();
+		TargetMaterialID = VoxWrapMesh.Attributes()->GetMaterialID();
+	}
+
+	// compute projected group and MaterialID and vertex colors
+	for (int32 tid : VoxWrapMesh.TriangleIndicesItr())
+	{
+		FVector3d Centroid = VoxWrapMesh.GetTriCentroid(tid);
+
+		double NearDistSqr = 0;
+		int32 NearestTID = CombinedMeshSpatial.FindNearestTriangle(Centroid, NearDistSqr);
+
+		if (SourceMaterialID != nullptr)
+		{
+			int32 MaterialID = SourceMaterialID->GetValue(NearestTID);
+			TargetMaterialID->SetValue(tid, MaterialID);
+		}
+
+		if (SourceColors != nullptr)
+		{
+			FIndex3i SourceTriElems = SourceColors->GetTriangle(NearestTID);
+			// TODO be smarter here...
+			FVector4f Color = SourceColors->GetElement(SourceTriElems.A);
+			int A = TargetColors->AppendElement(Color);
+			int B = TargetColors->AppendElement(Color);
+			int C = TargetColors->AppendElement(Color);
+			TargetColors->SetTriangle(tid, FIndex3i(A,B,C));
+		}
+	}
+}
+
+
+
 
 struct FCombinedMeshLOD
 {
@@ -1011,6 +1126,14 @@ struct FCombinedMeshLOD
 }
 
 
+enum class ECombinedLODType
+{
+	Copied = 0,
+	Simplified = 1,
+	Approximated = 2,
+	VoxWrapped = 3
+};
+
 
 // change this to build a single LOD, and separate versions for (eg) source mesh vs approx mesh
 // should we even bother w/ storing approx meshes? just generate them as needed?
@@ -1026,6 +1149,26 @@ void BuildCombinedMesh(
 	TArray<FCombinedMeshLOD> MeshLODs;
 	MeshLODs.SetNum(NumLODs);
 
+	int FirstVoxWrappedIndex = 9999;
+	TArray<ECombinedLODType> LODTypes;
+	LODTypes.Init(ECombinedLODType::Approximated, NumLODs);
+	for (int32 LODLevel = 0; LODLevel < NumLODs; ++LODLevel)
+	{
+		if (LODLevel < CombineOptions.NumCopiedLODs)
+		{
+			LODTypes[LODLevel] = ECombinedLODType::Copied;
+		}
+		else if (LODLevel < CombineOptions.NumCopiedLODs + CombineOptions.NumSimplifiedLODs)
+		{
+			LODTypes[LODLevel] = ECombinedLODType::Simplified;
+		}
+		else if (LODLevel >= NumLODs - CombineOptions.NumVoxWrapLODs)
+		{
+			LODTypes[LODLevel] = ECombinedLODType::VoxWrapped;
+			FirstVoxWrappedIndex = FMath::Min(LODLevel, FirstVoxWrappedIndex);
+		}
+	}
+
 	//CombinedLOD0.Attributes()->SetNumPolygroupLayers(2);
 	//FDynamicMeshPolygroupAttribute* PartIDAttrib = AccumMesh.Attributes()->GetPolygroupLayer(0);
 	//FDynamicMeshPolygroupAttribute* PartInstanceMapAttrib = AccumMesh.Attributes()->GetPolygroupLayer(1);
@@ -1056,19 +1199,25 @@ void BuildCombinedMesh(
 			// so that we always have something to swap to for Decorative parts
 			ApproximateAppendMesh = &OptimizedGeometry.ApproximateMeshLODs.Last();
 
-			if (LODLevel < CombineOptions.NumCopiedLODs)
+			ECombinedLODType LevelLODType = LODTypes[LODLevel];
+			if (LevelLODType == ECombinedLODType::Copied)
 			{
 				SourceAppendMesh = (LODLevel < SourceGeometry.SourceMeshLODs.Num()) ? 
 					&SourceGeometry.SourceMeshLODs[LODLevel] : &SourceGeometry.SourceMeshLODs.Last();
 				UseAppendMesh = SourceAppendMesh;
 			}
-			else if (LODLevel < CombineOptions.NumCopiedLODs + CombineOptions.NumSimplifiedLODs)
+			else if (LevelLODType == ECombinedLODType::Simplified)
 			{
 				int32 SimplifiedLODIndex = LODLevel - CombineOptions.NumCopiedLODs;
 				SourceAppendMesh = &OptimizedGeometry.SimplifiedMeshLODs[SimplifiedLODIndex];
 				UseAppendMesh = SourceAppendMesh;
 			}
-			else
+			else if (LevelLODType == ECombinedLODType::VoxWrapped)
+			{
+				SourceAppendMesh = &SourceGeometry.SourceMeshLODs.Last();
+				UseAppendMesh = SourceAppendMesh;
+			}
+			else // ECombinedLODType::Approximated
 			{
 				int32 ApproxLODIndex = LODLevel - CombineOptions.NumCopiedLODs - CombineOptions.NumSimplifiedLODs;
 				ApproximateAppendMesh = &OptimizedGeometry.ApproximateMeshLODs[ApproxLODIndex];
@@ -1081,15 +1230,18 @@ void BuildCombinedMesh(
 			{
 				bool bIsDecorativePart = (Instance.DetailLevel == EMeshDetailLevel::Decorative);
 
-				// filter out detail parts at higher LODs
-				if (bIsDecorativePart && LODLevel >= CombineOptions.FilterDecorativePartsLODLevel)
-				{
-					continue;
-				}
-				// at last detail part LOD, switch to approximate mesh
-				if (bIsDecorativePart && LODLevel >= (CombineOptions.FilterDecorativePartsLODLevel - CombineOptions.ApproximateDecorativePartLODs) )
+				if (bIsDecorativePart)
 				{
-					UseAppendMesh = ApproximateAppendMesh;
+					// filter out detail parts at higher LODs, or if we are doing VoxWrap LOD
+					if ( LODLevel >= CombineOptions.FilterDecorativePartsLODLevel || LevelLODType == ECombinedLODType::VoxWrapped )
+					{
+						continue;
+					}
+					// at last detail part LOD, switch to approximate mesh
+					if (LODLevel >= (CombineOptions.FilterDecorativePartsLODLevel - CombineOptions.ApproximateDecorativePartLODs) )
+					{
+						UseAppendMesh = ApproximateAppendMesh;
+					}
 				}
 
 				// need to make a copy to run pre-process func
@@ -1104,8 +1256,6 @@ void BuildCombinedMesh(
 					[&](int, const FVector3d& Pos) { return Instance.WorldTransform.TransformPosition(Pos); },
 					[&](int, const FVector3d& Normal) { return Instance.WorldTransform.TransformNormal(Normal); });
 
-				//CollisionAssembly.AppendInstance(InstanceSet.Collision, Instance.WorldTransform);
-
 				// append part ID stuff here
 
 				// could precompute these indexes for each instance?
@@ -1119,8 +1269,32 @@ void BuildCombinedMesh(
 					CombinedMeshLODData.MaterialIDs->SetValue( Mappings.GetNewTriangle(tid), AssignMaterialIndex );
 				}
 			}
+		}
+	}
+
+	//
+	// Process VoxWrapped LODs 
+	//
+	if ( FirstVoxWrappedIndex < 9999 )
+	{
+		FDynamicMesh3 SourceVoxWrapMesh = MoveTemp(MeshLODs[FirstVoxWrappedIndex].Mesh);
+		FDynamicMeshAABBTree3 Spatial(&SourceVoxWrapMesh, true);
 
+		FDynamicMesh3 TempBaseVoxWrapMesh;
+		ComputeVoxWrapMesh(SourceVoxWrapMesh, Spatial, TempBaseVoxWrapMesh, 10.0, 1.0);
+
+		const FDynamicMesh3& LastApproxLOD = MeshLODs[FirstVoxWrappedIndex-1].Mesh;
+
+		int32 MaxTriCount = CombineOptions.VoxWrapMaxTriCountBase;
+		double SimplifyTolerance = CombineOptions.VoxWrapBaseTolerance;
+		for (int32 LODIndex = FirstVoxWrappedIndex; LODIndex < NumLODs; ++LODIndex)
+		{
+			MeshLODs[LODIndex].Mesh = TempBaseVoxWrapMesh;
+			ComputeSimplifiedVoxWrapMesh(SourceVoxWrapMesh, Spatial, MeshLODs[LODIndex].Mesh, SimplifyTolerance, MaxTriCount);
+			SimplifyTolerance *= 1.5;
+			MaxTriCount /= 2;
 		}
+
 	}
 
 
@@ -1159,8 +1333,20 @@ void BuildCombinedMesh(
 
 	for (int32 LODLevel = 0; LODLevel < NumLODs; ++LODLevel)
 	{
-		CombinedMeshLODs.Add(MoveTemp(MeshLODs[LODLevel].Mesh));
+		FDynamicMesh3 LODMesh = MoveTemp(MeshLODs[LODLevel].Mesh);
+
+		// If we ended up larger than the mesh in the previous LOD, we should use that instead!
+		// This can happen particular with VoxWrap LODs
+		if (LODLevel > 0)
+		{
+			if (LODMesh.TriangleCount() > CombinedMeshLODs.Last().TriangleCount())
+			{
+				LODMesh = CombinedMeshLODs.Last();
+			}
+		}
+		CombinedMeshLODs.Add(MoveTemp(LODMesh));
 	}
+
 }
 
 
diff --git a/Engine/Source/Developer/GeometryProcessingInterfaces/Public/GeometryProcessingInterfaces/CombineMeshInstances.h b/Engine/Source/Developer/GeometryProcessingInterfaces/Public/GeometryProcessingInterfaces/CombineMeshInstances.h
--- a/Engine/Source/Developer/GeometryProcessingInterfaces/Public/GeometryProcessingInterfaces/CombineMeshInstances.h
+++ b/Engine/Source/Developer/GeometryProcessingInterfaces/Public/GeometryProcessingInterfaces/CombineMeshInstances.h
@@ -92,6 +92,10 @@ public:
 		double OptimizeBaseTriCost = 0.7;
 		double OptimizeLODLevelTriCostScale = 1.5;
 
+		int32 NumVoxWrapLODs = 1;
+		double VoxWrapBaseTolerance = 1.0;
+		int32 VoxWrapMaxTriCountBase = 500;
+
 		//
 		// Hidden Faces removal options
 		// 
