CombineMeshInstances updates. Fixed bug where TImplicitMorphology MeshCellSize was not being initialized. Added mesh sorting function, this is necessary as voxel-based methods, when done in parallel, output permutations of the same mesh, breaking downstream determinism. Modified simplification of voxelization output to do a fast-collapse pass down to 50k tris, and then incrementally simplify, instead of (serially) simplifying from original mesh for each voxel LOD. Hidden-removal for non-voxel LODs now runs in background task at the same time as voxel LODs are generating, this is a significant improvement (although probably some CPU contention as both methods are trivially parallel)

diff --git a/Engine/Plugins/Experimental/MeshModelingToolsetExp/Source/GeometryProcessingAdapters/Private/CombineMeshInstancesImpl.cpp b/Engine/Plugins/Experimental/MeshModelingToolsetExp/Source/GeometryProcessingAdapters/Private/CombineMeshInstancesImpl.cpp
--- a/Engine/Plugins/Experimental/MeshModelingToolsetExp/Source/GeometryProcessingAdapters/Private/CombineMeshInstancesImpl.cpp
+++ b/Engine/Plugins/Experimental/MeshModelingToolsetExp/Source/GeometryProcessingAdapters/Private/CombineMeshInstancesImpl.cpp
@@ -3,6 +3,7 @@
 #include "GeometryProcessing/CombineMeshInstancesImpl.h"
 
 #include "Async/ParallelFor.h"
+#include "Tasks/Task.h"
 
 #include "Engine/World.h"
 #include "Engine/StaticMesh.h"
@@ -39,6 +40,7 @@
 #include "Implicit/Morphology.h"
 #include "ProjectionTargets.h"
 
+
 using namespace UE::Geometry;
 
 
@@ -491,6 +493,11 @@ static void SimplifyPartMesh(
 	Welder.OnlyUniquePairs = false;
 	Welder.Apply();
 
+	// Skip out for very low-poly parts, they are unlikely to simplify very nicely.
+	if (EditMesh.VertexCount() < 16)
+	{
+		return;
+	}
 
 	FVolPresMeshSimplification Simplifier(&EditMesh);
 
@@ -564,6 +571,7 @@ static void ComputeBoxApproximation(
 
 	// oriented box fitting is under-determined, in cases where the AABB and the OBB have the nearly the
 	// same volume, generally we prefer an AABB
+	// (note: this rarely works due to tessellation of (eg) circles/spheres, and should be replaced w/ a better heuristic)
 	FAxisAlignedBox3d AlignedBox = SourceMesh.GetBounds(false);
 	if (AlignedBox.Volume() < 1.05*OrientedBox.Volume())
 	{
@@ -571,7 +579,7 @@ static void ComputeBoxApproximation(
 	}
 
 	FGridBoxMeshGenerator BoxGen;
-	BoxGen.Box = ResultBoxes.Boxes[0].Box;
+	BoxGen.Box = OrientedBox;
 	BoxGen.EdgeVertices = {0,0,0};
 	OutputMesh.Copy(&BoxGen.Generate());
 }
@@ -998,7 +1006,7 @@ static void ComputeVoxWrapMesh(
 	FDynamicMeshAABBTree3& CombinedMeshSpatial,
 	FDynamicMesh3& ResultMesh,
 	double ClosureDistance,
-	double TargetCellSize)
+	double& TargetCellSizeInOut)
 {
 	TImplicitMorphology<FDynamicMesh3> Morphology;
 	Morphology.Source = &CombinedMesh;
@@ -1007,22 +1015,24 @@ static void ComputeVoxWrapMesh(
 	Morphology.Distance = FMath::Max(ClosureDistance, 0.001);
 
 	FAxisAlignedBox3d Bounds = CombinedMeshSpatial.GetBoundingBox();
-	double UseCellSize = FMath::Max(0.001, TargetCellSize);
+	double UseCellSize = FMath::Max(0.001, TargetCellSizeInOut);
 	int MaxGridDimEstimate = (int32)(Bounds.MaxDim() / UseCellSize);
-	if (MaxGridDimEstimate > 128)
+	if (MaxGridDimEstimate > 256)
 	{
-		UseCellSize = (float)Bounds.MaxDim() / 128.0;
+		UseCellSize = (float)Bounds.MaxDim() / 256;
 	}
 	Morphology.GridCellSize = UseCellSize;
+	Morphology.MeshCellSize = UseCellSize;
+	TargetCellSizeInOut = UseCellSize;
 
 	ResultMesh.Copy(&Morphology.Generate());
 	ResultMesh.DiscardAttributes();
 }
 
 static void ComputeSimplifiedVoxWrapMesh(
-	const FDynamicMesh3& CombinedMesh, 
-	FDynamicMeshAABBTree3& CombinedMeshSpatial,
 	FDynamicMesh3& VoxWrapMesh,
+	const FDynamicMesh3* CombinedMesh, 
+	FDynamicMeshAABBTree3* CombinedMeshSpatial,
 	double SimplifyTolerance,
 	double MaxTriCount)
 {
@@ -1048,38 +1058,54 @@ static void ComputeSimplifiedVoxWrapMesh(
 		Simplifier.SimplifyToTriangleCount( MaxTriCount );
 	}
 
-	// compact result
 	VoxWrapMesh.CompactInPlace();
+}
+
+
+
+static void InitializeAttributes(
+	FDynamicMesh3& TargetMesh,
+	double NormalAngleThreshDeg,
+	bool bProjectAttributes,
+	const FDynamicMesh3* SourceMesh,
+	FDynamicMeshAABBTree3* SourceMeshSpatial)
+{
+	TargetMesh.EnableTriangleGroups();
+	TargetMesh.EnableAttributes();
+	// recompute normals
+	FMeshNormals::InitializeOverlayTopologyFromOpeningAngle(&TargetMesh, TargetMesh.Attributes()->PrimaryNormals(), NormalAngleThreshDeg);
+	FMeshNormals::QuickRecomputeOverlayNormals(TargetMesh);
 
-	VoxWrapMesh.EnableTriangleGroups();
-	VoxWrapMesh.EnableAttributes();
-	FMeshNormals::InitializeOverlayToPerVertexNormals(VoxWrapMesh.Attributes()->PrimaryNormals(), false);
+	if (bProjectAttributes == false || SourceMesh == nullptr || SourceMeshSpatial == nullptr)
+	{
+		return;
+	}
 
 	const FDynamicMeshColorOverlay* SourceColors = nullptr;
 	FDynamicMeshColorOverlay* TargetColors = nullptr;
-	if ( CombinedMesh.HasAttributes() && CombinedMesh.Attributes()->HasPrimaryColors() )
+	if ( SourceMesh->HasAttributes() && SourceMesh->Attributes()->HasPrimaryColors() )
 	{
-		SourceColors = CombinedMesh.Attributes()->PrimaryColors();
-		VoxWrapMesh.Attributes()->EnablePrimaryColors();
-		TargetColors = VoxWrapMesh.Attributes()->PrimaryColors();
+		SourceColors = SourceMesh->Attributes()->PrimaryColors();
+		TargetMesh.Attributes()->EnablePrimaryColors();
+		TargetColors = TargetMesh.Attributes()->PrimaryColors();
 	}
 
 	const FDynamicMeshMaterialAttribute* SourceMaterialID = nullptr;
 	FDynamicMeshMaterialAttribute* TargetMaterialID = nullptr;
-	if ( CombinedMesh.HasAttributes() && CombinedMesh.Attributes()->HasMaterialID() )
+	if ( SourceMesh->HasAttributes() && SourceMesh->Attributes()->HasMaterialID() )
 	{
-		SourceMaterialID = CombinedMesh.Attributes()->GetMaterialID();
-		VoxWrapMesh.Attributes()->EnableMaterialID();
-		TargetMaterialID = VoxWrapMesh.Attributes()->GetMaterialID();
+		SourceMaterialID = SourceMesh->Attributes()->GetMaterialID();
+		TargetMesh.Attributes()->EnableMaterialID();
+		TargetMaterialID = TargetMesh.Attributes()->GetMaterialID();
 	}
 
 	// compute projected group and MaterialID and vertex colors
-	for (int32 tid : VoxWrapMesh.TriangleIndicesItr())
+	for (int32 tid : TargetMesh.TriangleIndicesItr())
 	{
-		FVector3d Centroid = VoxWrapMesh.GetTriCentroid(tid);
+		FVector3d Centroid = TargetMesh.GetTriCentroid(tid);
 
 		double NearDistSqr = 0;
-		int32 NearestTID = CombinedMeshSpatial.FindNearestTriangle(Centroid, NearDistSqr);
+		int32 NearestTID = SourceMeshSpatial->FindNearestTriangle(Centroid, NearDistSqr);
 
 		if (SourceMaterialID != nullptr)
 		{
@@ -1135,6 +1161,103 @@ enum class ECombinedLODType
 };
 
 
+
+static void SortMesh(FDynamicMesh3& Mesh)
+{
+	if ( ! ensure(Mesh.HasAttributes() == false) ) return;
+
+	TRACE_CPUPROFILER_EVENT_SCOPE(SortMesh);
+
+	struct FVert
+	{
+		FVector3d Position;
+		int32 VertexID;
+		bool operator<(const FVert& V2) const
+		{
+			if ( Position.X != V2.Position.X ) { return Position.X < V2.Position.X; }
+			if ( Position.Y != V2.Position.Y ) { return Position.Y < V2.Position.Y; }
+			if ( Position.Z != V2.Position.Z ) { return Position.Z < V2.Position.Z; }
+			return VertexID < V2.VertexID;
+		}
+	};
+	struct FTri
+	{
+		FIndex3i Triangle;
+		bool operator<(const FTri& Tri2) const
+		{
+			if ( Triangle.A != Tri2.Triangle.A ) { return Triangle.A < Tri2.Triangle.A; }
+			if ( Triangle.B != Tri2.Triangle.B ) { return Triangle.B < Tri2.Triangle.B; }
+			return Triangle.C < Tri2.Triangle.C; 
+		}
+	};
+
+	TArray<FVert> Vertices;
+	for ( int32 vid : Mesh.VertexIndicesItr() )
+	{
+		Vertices.Add( FVert{Mesh.GetVertex(vid), vid} );
+	}
+	Vertices.Sort();
+
+	TArray<int32> VertMap;
+	VertMap.SetNum(Mesh.MaxVertexID());
+	for (int32 k = 0; k < Vertices.Num(); ++k)
+	{
+		const FVert& Vert = Vertices[k];
+		VertMap[Vert.VertexID] = k;
+	}
+
+	TArray<FTri> Triangles;
+	for ( int32 tid : Mesh.TriangleIndicesItr() )
+	{
+		FIndex3i Tri = Mesh.GetTriangle(tid);
+		Tri.A = VertMap[Tri.A];
+		Tri.B = VertMap[Tri.B];
+		Tri.C = VertMap[Tri.C];
+		Triangles.Add( FTri{Tri} );
+	}
+	Triangles.Sort();
+
+	FDynamicMesh3 SortedMesh;
+	for (const FVert& Vert : Vertices)
+	{
+		SortedMesh.AppendVertex(Mesh, Vert.VertexID);
+	}
+	for (const FTri& Tri : Triangles)
+	{
+		SortedMesh.AppendTriangle(Tri.Triangle.A, Tri.Triangle.B, Tri.Triangle.C);
+	}
+
+	Mesh = MoveTemp(SortedMesh);
+}
+
+
+
+void ComputeHiddenRemovalForLOD(
+	FDynamicMesh3& MeshLOD,
+	IGeometryProcessing_CombineMeshInstances::FOptions CombineOptions)
+{
+	TRACE_CPUPROFILER_EVENT_SCOPE(RemoveHidden_LOD);
+	bool bModified = false;
+	switch (CombineOptions.RemoveHiddenFacesMethod)
+	{
+		case IGeometryProcessing_CombineMeshInstances::ERemoveHiddenFacesMode::OcclusionBased:
+			RemoveHiddenFaces_Occlusion(MeshLOD, 200);		// 200 is arbitrary here! should improve once max-distance is actually available (currently ignored)
+			bModified = true;
+			break;
+		case IGeometryProcessing_CombineMeshInstances::ERemoveHiddenFacesMode::ExteriorVisibility:
+		case IGeometryProcessing_CombineMeshInstances::ERemoveHiddenFacesMode::Fastest:
+			RemoveHiddenFaces_ExteriorVisibility(MeshLOD, CombineOptions.RemoveHiddenSamplingDensity);
+			bModified = true;
+			break;
+	}
+
+	if ( bModified )
+	{
+		PostProcessHiddenFaceRemovedMesh(MeshLOD, CombineOptions.SimplifyBaseTolerance);
+	}
+}
+
+
 // change this to build a single LOD, and separate versions for (eg) source mesh vs approx mesh
 // should we even bother w/ storing approx meshes? just generate them as needed?
 
@@ -1272,6 +1395,27 @@ void BuildCombinedMesh(
 		}
 	}
 
+
+	//
+	// start hidden-removal passes on all meshes up to voxel LODs here, because we can compute voxel LOD at the same time
+	//
+	TArray<UE::Tasks::FTask> PendingRemoveHiddenTasks;
+	bool bRemoveHiddenFaces = 
+		(CombineOptions.RemoveHiddenFacesMethod != IGeometryProcessing_CombineMeshInstances::ERemoveHiddenFacesMode::None 
+		&& CVarGeometryCombineMeshInstancesRemoveHidden.GetValueOnGameThread() > 0);
+	if (bRemoveHiddenFaces)
+	{
+		for (int32 k = CombineOptions.RemoveHiddenStartLOD; k < MeshLODs.Num() && k < FirstVoxWrappedIndex; ++k)
+		{
+			UE::Tasks::FTask RemoveHiddenTask = UE::Tasks::Launch(UE_SOURCE_LOCATION, [&MeshLODs, &CombineOptions, k]()
+			{ 
+				ComputeHiddenRemovalForLOD(MeshLODs[k].Mesh, CombineOptions);
+			});
+			PendingRemoveHiddenTasks.Add(RemoveHiddenTask);
+		}
+	}
+
+
 	//
 	// Process VoxWrapped LODs 
 	//
@@ -1281,7 +1425,23 @@ void BuildCombinedMesh(
 		FDynamicMeshAABBTree3 Spatial(&SourceVoxWrapMesh, true);
 
 		FDynamicMesh3 TempBaseVoxWrapMesh;
-		ComputeVoxWrapMesh(SourceVoxWrapMesh, Spatial, TempBaseVoxWrapMesh, 10.0, 1.0);
+		double VoxelDimension = 2.0;	// may be modified by ComputeVoxWrapMesh call
+		{
+			TRACE_CPUPROFILER_EVENT_SCOPE(ComputeVoxWrap);
+			ComputeVoxWrapMesh(SourceVoxWrapMesh, Spatial, TempBaseVoxWrapMesh, 10.0, VoxelDimension);
+			// currently need to re-sort output to remove non-determinism...
+			SortMesh(TempBaseVoxWrapMesh);
+
+			//UE_LOG(LogGeometry, Warning, TEXT("VoxWrapMesh has %d triangles %d vertices"), TempBaseVoxWrapMesh.TriangleCount(), TempBaseVoxWrapMesh.VertexCount());
+		}
+
+		{
+			TRACE_CPUPROFILER_EVENT_SCOPE(FastCollapsePrePass);
+			TempBaseVoxWrapMesh.DiscardAttributes();
+			FVolPresMeshSimplification Simplifier(&TempBaseVoxWrapMesh);
+			Simplifier.bAllowSeamCollapse = false;
+			Simplifier.FastCollapsePass(VoxelDimension*0.5, 10, false, 50000);
+		}
 
 		const FDynamicMesh3& LastApproxLOD = MeshLODs[FirstVoxWrappedIndex-1].Mesh;
 
@@ -1289,43 +1449,38 @@ void BuildCombinedMesh(
 		double SimplifyTolerance = CombineOptions.VoxWrapBaseTolerance;
 		for (int32 LODIndex = FirstVoxWrappedIndex; LODIndex < NumLODs; ++LODIndex)
 		{
-			MeshLODs[LODIndex].Mesh = TempBaseVoxWrapMesh;
-			ComputeSimplifiedVoxWrapMesh(SourceVoxWrapMesh, Spatial, MeshLODs[LODIndex].Mesh, SimplifyTolerance, MaxTriCount);
+			TRACE_CPUPROFILER_EVENT_SCOPE(SimplifyVoxWrap);
+
+			// using previous-simplified mesh for next level may not be ideal...
+			MeshLODs[LODIndex].Mesh = (LODIndex == FirstVoxWrappedIndex) ?
+				MoveTemp(TempBaseVoxWrapMesh) : MeshLODs[LODIndex-1].Mesh;
+			// need to do this because we projected attributes in previous loop
+			MeshLODs[LODIndex].Mesh.DiscardAttributes();
+
+			ComputeSimplifiedVoxWrapMesh(MeshLODs[LODIndex].Mesh, &SourceVoxWrapMesh, &Spatial, 
+				SimplifyTolerance, MaxTriCount);
+
+			InitializeAttributes(MeshLODs[LODIndex].Mesh, CombineOptions.HardNormalAngleDeg,
+				/*bProjectAttributes*/true, &SourceVoxWrapMesh, &Spatial);
+
 			SimplifyTolerance *= 1.5;
 			MaxTriCount /= 2;
 		}
-
 	}
 
 
-	//
-	// Hidden-face removal pass
-	//
-	if (CombineOptions.RemoveHiddenFacesMethod != IGeometryProcessing_CombineMeshInstances::ERemoveHiddenFacesMode::None 
-		&& CVarGeometryCombineMeshInstancesRemoveHidden.GetValueOnGameThread() > 0)
+	// wait...
+	UE::Tasks::Wait(PendingRemoveHiddenTasks);
+
+	// remove hiddel faces on voxel LODs (todo: can do this via shape sorting, much faster)
+	if (bRemoveHiddenFaces)
 	{
+		TRACE_CPUPROFILER_EVENT_SCOPE(RemoveHidden);
 		ParallelFor(MeshLODs.Num(), [&](int32 k)
 		{
-			bool bModified = false;
-			if ( k >= CombineOptions.RemoveHiddenStartLOD )
+			if ( k >= FirstVoxWrappedIndex )
 			{ 
-				switch (CombineOptions.RemoveHiddenFacesMethod)
-				{
-					case IGeometryProcessing_CombineMeshInstances::ERemoveHiddenFacesMode::OcclusionBased:
-						RemoveHiddenFaces_Occlusion(MeshLODs[k].Mesh, 200);		// 200 is arbitrary here! should improve once max-distance is actually available (currently ignored)
-						bModified = true;
-						break;
-					case IGeometryProcessing_CombineMeshInstances::ERemoveHiddenFacesMode::ExteriorVisibility:
-					case IGeometryProcessing_CombineMeshInstances::ERemoveHiddenFacesMode::Fastest:
-						RemoveHiddenFaces_ExteriorVisibility(MeshLODs[k].Mesh, CombineOptions.RemoveHiddenSamplingDensity);
-						bModified = true;
-						break;
-				}
-			}
-
-			if ( bModified )
-			{
-				PostProcessHiddenFaceRemovedMesh(MeshLODs[k].Mesh, CombineOptions.SimplifyBaseTolerance);
+				ComputeHiddenRemovalForLOD(MeshLODs[k].Mesh, CombineOptions);
 			}
 		});
 	}
@@ -1573,11 +1728,20 @@ static void SetConstantVertexColor(FDynamicMesh3& Mesh, FLinearColor LinearColor
 void FCombineMeshInstancesImpl::CombineMeshInstances(
 	const FInstanceSet& MeshInstances, const FOptions& Options, FResults& ResultsOut)
 {
+	TRACE_CPUPROFILER_EVENT_SCOPE(CombineMeshInstances);
+
 	FMeshInstanceAssembly InstanceAssembly;
-	InitializeMeshInstanceAssembly(MeshInstances, InstanceAssembly);
-	InitializeAssemblySourceMeshesFromLOD(InstanceAssembly, 0, Options.NumCopiedLODs);
-	InitializeInstanceAssemblySpatials(InstanceAssembly);
-	ComputeMeshApproximations(Options, InstanceAssembly);
+
+	{
+		TRACE_CPUPROFILER_EVENT_SCOPE(CombineMeshInst_Setup);
+		InitializeMeshInstanceAssembly(MeshInstances, InstanceAssembly);
+		InitializeAssemblySourceMeshesFromLOD(InstanceAssembly, 0, Options.NumCopiedLODs);
+		InitializeInstanceAssemblySpatials(InstanceAssembly);
+	}
+	{
+		TRACE_CPUPROFILER_EVENT_SCOPE(CombineMeshInst_PartApprox);
+		ComputeMeshApproximations(Options, InstanceAssembly);
+	}
 
 
 	
@@ -1595,10 +1759,16 @@ void FCombineMeshInstancesImpl::CombineMeshInstances(
 	};
 
 	TArray<UE::Geometry::FDynamicMesh3> CombinedMeshLODs;
-	BuildCombinedMesh(InstanceAssembly, Options, CombinedMeshLODs);
+	{
+		TRACE_CPUPROFILER_EVENT_SCOPE(CombineMeshInst_BuildMeshes);
+		BuildCombinedMesh(InstanceAssembly, Options, CombinedMeshLODs);
+	}
 
 	FSimpleShapeSet3d CombinedCollisionShapes;
-	BuildCombinedCollisionShapes(InstanceAssembly, Options, CombinedCollisionShapes);
+	{
+		TRACE_CPUPROFILER_EVENT_SCOPE(CombineMeshInst_BuildCollision);
+		BuildCombinedCollisionShapes(InstanceAssembly, Options, CombinedCollisionShapes);
+	}
 	FPhysicsDataCollection PhysicsData;
 	PhysicsData.Geometry = CombinedCollisionShapes;
 	PhysicsData.CopyGeometryToAggregate();		// need FPhysicsDataCollection to convert to agg geom, should fix this
