Performance improvements to skin painting tool.
- new sparse weight format
- cached vertex positions in bone space
- parallelized deformation update
- override HitTest() to use tri octree instead of mesh AABBTree
- octree updates asynchronously

diff --git a/Engine/Plugins/Experimental/MeshModelingToolsetExp/Source/MeshModelingToolsExp/Private/SkinWeightsPaintTool.cpp b/Engine/Plugins/Experimental/MeshModelingToolsetExp/Source/MeshModelingToolsExp/Private/SkinWeightsPaintTool.cpp
--- a/Engine/Plugins/Experimental/MeshModelingToolsetExp/Source/MeshModelingToolsExp/Private/SkinWeightsPaintTool.cpp
+++ b/Engine/Plugins/Experimental/MeshModelingToolsetExp/Source/MeshModelingToolsExp/Private/SkinWeightsPaintTool.cpp
@@ -3,8 +3,6 @@
 #include "SkinWeightsPaintTool.h"
 #include "Engine/SkeletalMesh.h"
 #include "InteractiveToolManager.h"
-#include "ToolBuilderUtil.h"
-#include "SceneManagement.h"
 #include "SkeletalMeshAttributes.h"
 #include "SkeletalDebugRendering.h"
 #include "Math/UnrealMathUtility.h"
@@ -16,14 +14,20 @@
 #include "MeshDescription.h"
 #include "DynamicMesh/NonManifoldMappingSupport.h"
 #include "Parameterization/MeshLocalParam.h"
+#include "Spatial/FastWinding.h"
+#include "Async/Async.h"
 
 #include UE_INLINE_GENERATED_CPP_BY_NAME(SkinWeightsPaintTool)
 
-
 #define LOCTEXT_NAMESPACE "USkinWeightsPaintTool"
 
-// Any weight below this value is ignored, since it won't be representable in a uint8.
-constexpr float MinimumWeightThreshold = 1.0f / 255.0f;
+using namespace SkinPaintTool;
+
+// thread pool to use for async operations
+static EAsyncExecution SkinPaintToolAsyncExecTarget = EAsyncExecution::ThreadPool;
+
+// any weight below this value is ignored, since it won't be representable in unsigned 16-bit precision
+constexpr float MinimumWeightThreshold = 1.0f / 65535.0f;
 
 USkinWeightsPaintToolProperties::USkinWeightsPaintToolProperties(const FObjectInitializer& ObjectInitializer)
 		: Super(ObjectInitializer)
@@ -37,9 +41,45 @@ USkinWeightsPaintToolProperties::USkinWeightsPaintToolProperties(const FObjectIn
 	MaxColor = FLinearColor::White;
 }
 
-void FSkinToolDeformer::Initialize(
-	const USkeletalMeshComponent* SkeletalMeshComponent,
-	const FMeshDescription* Mesh)
+void FMultiBoneWeightEdits::MergeSingleEdit(
+	const int32 BoneIndex,
+	const int32 VertexID,
+	const float OldWeight,
+	const float NewWeight)
+{
+	FSingleBoneWeightEdits& BoneWeightEdit = PerBoneWeightEdits.FindOrAdd(BoneIndex);
+	BoneWeightEdit.BoneIndex = BoneIndex;
+	BoneWeightEdit.NewWeights.Add(VertexID, NewWeight);
+	BoneWeightEdit.OldWeights.FindOrAdd(VertexID, OldWeight);
+}
+
+void FMultiBoneWeightEdits::MergeEdits(const FSingleBoneWeightEdits& BoneWeightEdits)
+{
+	// make sure bone has an entry in the map of weight edits
+	const int32 BoneIndex = BoneWeightEdits.BoneIndex;
+	PerBoneWeightEdits.FindOrAdd(BoneIndex);
+	PerBoneWeightEdits[BoneIndex].BoneIndex = BoneIndex;
+	
+	for (const TTuple<int32, float>& NewWeight : BoneWeightEdits.NewWeights)
+	{
+		int32 VertexIndex = NewWeight.Key;
+		PerBoneWeightEdits[BoneIndex].NewWeights.Add(VertexIndex, NewWeight.Value);
+		PerBoneWeightEdits[BoneIndex].OldWeights.FindOrAdd(VertexIndex, BoneWeightEdits.OldWeights[VertexIndex]);
+	}
+}
+
+float FMultiBoneWeightEdits::GetVertexDeltaFromEdits(const int32 BoneIndex, const int32 VertexIndex)
+{
+	PerBoneWeightEdits.FindOrAdd(BoneIndex);
+	if (const float* NewVertexWeight = PerBoneWeightEdits[BoneIndex].NewWeights.Find(VertexIndex))
+	{
+		return *NewVertexWeight - PerBoneWeightEdits[BoneIndex].OldWeights[VertexIndex];
+	}
+
+	return 0.0f;
+}
+
+void FSkinToolDeformer::Initialize(const USkeletalMeshComponent* SkeletalMeshComponent, const FMeshDescription* Mesh)
 {
 	// get all bone transforms in the reference pose store a copy in component space
 	Component = SkeletalMeshComponent;
@@ -66,8 +106,9 @@ void FSkinToolDeformer::Initialize(
 		// pre-invert the transforms so we don't have to at runtime
 		InvCSRefPoseTransforms[BoneIndex] = InvCSRefPoseTransforms[BoneIndex].Inverse();
 
-		// store map of bone names to bone indices
+		// store map of bone indices to bone names
 		FName BoneName = RefSkeleton.GetBoneName(BoneIndex);
+		BoneNames.Add(BoneName);
 		BoneNameToIndexMap.Add(BoneName, BoneIndex);
 	}
 
@@ -76,66 +117,74 @@ void FSkinToolDeformer::Initialize(
 	RefPoseVertexPositions = VertexPositions;
 
 	// force all vertices to be updated initially
-	VerticesWithModifiedWeights.Reserve(RefPoseVertexPositions.Num());
+	VerticesWithModifiedWeights.Empty(RefPoseVertexPositions.Num());
 	for (int32 VertexID=0; VertexID<RefPoseVertexPositions.Num(); ++VertexID)
 	{
 		VerticesWithModifiedWeights.Add(VertexID);
 	}
 }
 
-void FSkinToolDeformer::UpdateVertexDeformation(
-	PerBoneWeightMap& Weights,
-	UPreviewMesh* PreviewMesh,
-	UE::Geometry::TDynamicVerticesOctree3<FDynamicMesh3>& OctreeToUpdate)
+void FSkinToolDeformer::UpdateVertexDeformation(USkinWeightsPaintTool* Tool)
 {
+	TRACE_CPUPROFILER_EVENT_SCOPE(SkinTool::UpdateDeformationTotal);
+	
 	if (VerticesWithModifiedWeights.IsEmpty())
 	{
 		return;
 	}
 	
 	// update vertex positions
-	PreviewMesh->EditMesh([this, &Weights](FDynamicMesh3& Mesh)
+	UPreviewMesh* PreviewMesh = Tool->PreviewMesh;
+	const TArray<VertexWeights>& CurrentWeights = Tool->Weights.CurrentWeights;
+	PreviewMesh->DeferredEditMesh([this, &CurrentWeights](FDynamicMesh3& Mesh)
 	{
+		TRACE_CPUPROFILER_EVENT_SCOPE(SkinTool::UpdateDeformation);
 		const TArray<FTransform>& CurrentBoneTransforms = Component->GetComponentSpaceTransforms();
-
-		// accumulate new skinned position for each vertex that has had it's weight's modified
-		for (const int32 VertexID : VerticesWithModifiedWeights)
+		const TArray<int32> VertexIndices = VerticesWithModifiedWeights.Array();
+		
+		ParallelFor( VerticesWithModifiedWeights.Num(), [this, &VertexIndices, &Mesh, &CurrentBoneTransforms, &CurrentWeights](int32 Index)
 		{
+			const int32 VertexID = VertexIndices[Index];
 			FVector VertexNewPosition = FVector::ZeroVector;
-			for (TTuple<FName, TArray<float>> Pair : Weights)
+			const VertexWeights& VertexPerBoneData = CurrentWeights[VertexID];
+			for (const FVertexBoneWeight& VertexData : VertexPerBoneData)
 			{
-				const float VertexWeight = Pair.Value[VertexID];
-				if (VertexWeight <= MinimumWeightThreshold)
-				{
-					continue;
-				}
-				const FName BoneName = Pair.Key;
-				const int32 BoneIndex = BoneNameToIndexMap[BoneName];
-				const FVector& RefPoseVertexPosition = RefPoseVertexPositions[VertexID];
-				const FTransform& InvRefPoseTransform = InvCSRefPoseTransforms[BoneIndex];
-				const FVector& BoneLocalPositionInRefPose = InvRefPoseTransform.TransformPosition(RefPoseVertexPosition);
-				const FTransform& CurrentTransform = CurrentBoneTransforms[BoneIndex];
-				VertexNewPosition += CurrentTransform.TransformPosition(BoneLocalPositionInRefPose) * VertexWeight;
+				const FTransform& CurrentTransform = CurrentBoneTransforms[VertexData.BoneIndex];
+				VertexNewPosition += CurrentTransform.TransformPosition(VertexData.VertexInBoneSpace) * VertexData.Weight;
 			}
-			 
+			
 			Mesh.SetVertex(VertexID, VertexNewPosition, false);
-		}
-	});
-
-	// update acceleration structure
-	OctreeToUpdate.RemoveVertices(VerticesWithModifiedWeights);
-	OctreeToUpdate.InsertVertices(VerticesWithModifiedWeights);
+		});
+	}, false);
+	PreviewMesh->NotifyDeferredEditCompleted(UPreviewMesh::ERenderUpdateMode::FastUpdate, EMeshRenderAttributeFlags::Positions, false);
 
-	// empty queue of vertices to update
-	VerticesWithModifiedWeights.Reset();
-}
+	// update vertex acceleration structure
+	{
+		TRACE_CPUPROFILER_EVENT_SCOPE(SkinTool::UpdateVertexOctree);
+		Tool->VerticesOctree.RemoveVertices(VerticesWithModifiedWeights);
+		Tool->VerticesOctree.InsertVertices(VerticesWithModifiedWeights);
+	}
 
-void FSkinToolDeformer::SetVerticesNeedUpdated(TArray<int32> VertexIndices)
-{
-	for (int32 VertexID : VertexIndices)
+	// update triangle acceleration structure
 	{
-		VerticesWithModifiedWeights.Add(VertexID);
+		TRACE_CPUPROFILER_EVENT_SCOPE(SkinTool::UpdateTriangleOctree);
+
+		// ensure previous async update is finished before queuing the next one...
+		Tool->TriangleOctreeFuture.Wait();
+
+		TArray<int32>& TrianglesToReinsert = Tool->TriangleToReinsert;
+		const UE::Geometry::FAxisAlignedBox3d QueryBox(Tool->StampLocalPos, Tool->CurrentBrushRadius);
+		Tool->TrianglesOctree.RangeQuery(QueryBox, TrianglesToReinsert);
+		UE::Geometry::FDynamicMeshOctree3& OctreeToUpdate = Tool->TrianglesOctree;
+		Tool->TriangleOctreeFuture = Async(SkinPaintToolAsyncExecTarget, [&OctreeToUpdate, &TrianglesToReinsert]()
+		{
+			TRACE_CPUPROFILER_EVENT_SCOPE(SkinTool::TriangleOctreeReinsert);	
+			OctreeToUpdate.ReinsertTriangles(TrianglesToReinsert);
+		});
 	}
+
+	// empty queue of vertices to update
+	VerticesWithModifiedWeights.Reset();
 }
 
 void FSkinToolDeformer::SetVertexNeedsUpdated(int32 VertexIndex)
@@ -143,82 +192,70 @@ void FSkinToolDeformer::SetVertexNeedsUpdated(int32 VertexIndex)
 	VerticesWithModifiedWeights.Add(VertexIndex);
 }
 
-void FSkinToolWeights::InitializeSkinWeights(const FReferenceSkeleton& RefSkeleton, FMeshDescription* EditedMesh)
+void FSkinToolWeights::InitializeSkinWeights(
+	const USkeletalMeshComponent* SkeletalMeshComponent,
+	FMeshDescription* Mesh)
 {
-	const FSkeletalMeshConstAttributes MeshAttribs(*EditedMesh);
-	const FSkinWeightsVertexAttributesConstRef VertexSkinWeights = MeshAttribs.GetVertexSkinWeights();
-	const int32 NumVertices = EditedMesh->Vertices().Num();
-
-	// Create a map of all bones to their per-vertex weights.
-	CurrentWeightsMap.Reset();
-	for (const FMeshBoneInfo& BoneInfo : RefSkeleton.GetRefBoneInfo())
-	{
-		CurrentWeightsMap.Add(BoneInfo.Name, {}).AddZeroed(NumVertices);
-	}
+	// initialize deformer data
+	Deformer.Initialize(SkeletalMeshComponent, Mesh);
 
+	// initialize current weights (using compact format: num_verts * max_influences)
+	const FSkeletalMeshConstAttributes MeshAttribs(*Mesh);
+	const FSkinWeightsVertexAttributesConstRef VertexSkinWeights = MeshAttribs.GetVertexSkinWeights();
+	const int32 NumVertices = Mesh->Vertices().Num();
+	CurrentWeights.SetNum(NumVertices);
 	for (int32 VertexIndex = 0; VertexIndex < NumVertices; VertexIndex++)
 	{
 		const FVertexID VertexID(VertexIndex);
-
+		int32 InfluenceIndex = 0;
 		for (UE::AnimationCore::FBoneWeight BoneWeight: VertexSkinWeights.Get(VertexID))
 		{
-			FName BoneName = RefSkeleton.GetBoneName(static_cast<int32>(BoneWeight.GetBoneIndex()));
+			check(InfluenceIndex < MAX_TOTAL_INFLUENCES);
+			const int32 BoneIndex = BoneWeight.GetBoneIndex();
 			const float Weight = BoneWeight.GetWeight();
-
-			if (Weight >= MinimumWeightThreshold)
-			{
-				// If the source mesh has a bone that we don't recognize, we ignore it. It's
-				// weight will get cleared when the new weights are updated back to the 
-				// source mesh.
-				TArray<float>* PerVertexWeights = CurrentWeightsMap.Find(BoneName);
-
-				if (PerVertexWeights)
-				{
-					PerVertexWeights->GetData()[VertexIndex] = Weight;
-				}
-			}
+			const FVector& RefPoseVertexPosition = Deformer.RefPoseVertexPositions[VertexID];
+			const FTransform& InvRefPoseTransform = Deformer.InvCSRefPoseTransforms[BoneIndex];
+			const FVector& BoneLocalPositionInRefPose = InvRefPoseTransform.TransformPosition(RefPoseVertexPosition);
+			CurrentWeights[VertexIndex].Emplace(BoneIndex, BoneLocalPositionInRefPose, Weight);
+			++InfluenceIndex;
 		}
 	}
 	
 	// maintain duplicate weight map
-	PreStrokeWeightsMap = CurrentWeightsMap;
+	PreStrokeWeights = CurrentWeights;
 
 	// maintain relax-per stroke map
 	MaxFalloffPerVertexThisStroke.SetNumZeroed(NumVertices);
 }
 
-TArray<float>* FSkinToolWeights::GetWeightsForBone(const FName BoneName)
-{
-	return CurrentWeightsMap.Find(BoneName);
-}
-
 void FSkinToolWeights::EditVertexWeightAndNormalize(
 	const FName BoneToHoldConstant,
-	const int32 VertexId,
+	const int32 VertexID,
 	const float NewWeightValue,
 	FMultiBoneWeightEdits& WeightEdits)
 {
+	const int32 BoneToHoldIndex = Deformer.BoneNameToIndexMap[BoneToHoldConstant];
+	
 	// calculate the sum of all the weights on this vertex (not including the one we currently applied)
-	TArray<FName> NamesOfBonesAffectingVertex;
+	TArray<int32> BonesAffectingVertex;
 	TArray<float> ValuesToNormalize;
 	float Total = 0.0f;
-	for (const TTuple<FName, TArray<float>>& Pair : CurrentWeightsMap)
+	const VertexWeights& VertexData = CurrentWeights[VertexID];
+	for (const FVertexBoneWeight& VertexBoneData : VertexData)
 	{
-		const FName BoneName = Pair.Key;
-		if (BoneName == BoneToHoldConstant)
+		if (VertexBoneData.BoneIndex == BoneToHoldIndex)
 		{
 			continue;
 		}
 		
-		const float WeightOnThisBone = CurrentWeightsMap[BoneName][VertexId];
-		if (WeightOnThisBone < MinimumWeightThreshold)
+		if (VertexBoneData.Weight < MinimumWeightThreshold)
 		{
 			continue;
 		}
 		
-		NamesOfBonesAffectingVertex.Add(BoneName);
-		ValuesToNormalize.Add(WeightOnThisBone);
-		Total += WeightOnThisBone;
+		BonesAffectingVertex.Add(VertexBoneData.BoneIndex);
+		ValuesToNormalize.Add(VertexBoneData.Weight);
+		Total += VertexBoneData.Weight;
 	}
 
 	// if user applied FULL weight to this vertex OR there's no other weights of any significance,
@@ -228,17 +265,18 @@ void FSkinToolWeights::EditVertexWeightAndNormalize(
 		// set all other influences to 0.0f
 		for (int32 i=0; i<ValuesToNormalize.Num(); ++i)
 		{
-			const FName BoneName = NamesOfBonesAffectingVertex[i];
+			const int32 BoneIndex = BonesAffectingVertex[i];
 			const float OldWeight = ValuesToNormalize[i];
 			constexpr float NewWeight = 0.0f;
-			WeightEdits.MergeSingleEdit(BoneName, VertexId, OldWeight, NewWeight);
+			WeightEdits.MergeSingleEdit(BoneIndex, VertexID, OldWeight, NewWeight);
 		}
 
 		// set current bone value to 1.0f
+		const float PrevWeight = GetWeightOfBoneOnVertex(BoneToHoldIndex, VertexID, PreStrokeWeights);
 		WeightEdits.MergeSingleEdit(
-			BoneToHoldConstant,
-			VertexId,
-			PreStrokeWeightsMap[BoneToHoldConstant][VertexId],
+			BoneToHoldIndex,
+			VertexID,
+			PrevWeight,
 			1.0f);
 		
 		return;
@@ -255,62 +293,132 @@ void FSkinToolWeights::EditVertexWeightAndNormalize(
 		{
 			NormalizedValue = (ValuesToNormalize[i] / Total) * AvailableTotal;	
 		}
-		const FName BoneName = NamesOfBonesAffectingVertex[i];
+		const int32 BoneIndex = BonesAffectingVertex[i];
 		const float OldWeight = ValuesToNormalize[i];
 		const float NewWeight = NormalizedValue;
-		WeightEdits.MergeSingleEdit(BoneName, VertexId, OldWeight, NewWeight);
+		WeightEdits.MergeSingleEdit(BoneIndex, VertexID, OldWeight, NewWeight);
 	}
 
 	// record current bone edit
+	const float PrevWeight = GetWeightOfBoneOnVertex(BoneToHoldIndex, VertexID, PreStrokeWeights);
 	WeightEdits.MergeSingleEdit(
-		BoneToHoldConstant,
-		VertexId,
-		PreStrokeWeightsMap[BoneToHoldConstant][VertexId],
+		BoneToHoldIndex,
+		VertexID,
+		PrevWeight,
 		NewWeightValue);
 }
 
-void FSkinToolWeights::ApplyCurrentWeightsToMeshDescription(
-	const FReferenceSkeleton& RefSkeleton,
-	FMeshDescription* EditedMesh)
+void FSkinToolWeights::ApplyCurrentWeightsToMeshDescription(FMeshDescription* EditedMesh)
 {
-	using namespace UE::AnimationCore;
-	
 	FSkeletalMeshAttributes MeshAttribs(*EditedMesh);
 	FSkinWeightsVertexAttributesRef VertexSkinWeights = MeshAttribs.GetVertexSkinWeights();
+	
+	UE::AnimationCore::FBoneWeightsSettings Settings;
+	Settings.SetNormalizeType(UE::AnimationCore::EBoneWeightNormalizeType::None);
+
+	TArray<UE::AnimationCore::FBoneWeight> SourceBoneWeights;
+	SourceBoneWeights.Reserve(UE::AnimationCore::MaxInlineBoneWeightCount);
+
+	const int32 NumVertices = EditedMesh->Vertices().Num();
+	for (int32 VertexIndex = 0; VertexIndex < NumVertices; VertexIndex++)
+	{
+		SourceBoneWeights.Reset();
+
+		const VertexWeights& VertexWeights = CurrentWeights[VertexIndex];
+		for (const FVertexBoneWeight& SingleBoneWeight : VertexWeights)
+		{
+			SourceBoneWeights.Add(UE::AnimationCore::FBoneWeight(SingleBoneWeight.BoneIndex, SingleBoneWeight.Weight));
+		}
+
+		VertexSkinWeights.Set(FVertexID(VertexIndex), UE::AnimationCore::FBoneWeights::Create(SourceBoneWeights, Settings));
+	}
+}
 
-	TMap<FBoneIndexType, const TArray<float>*> BoneIndexWeightMap;
-	for (const TTuple<FName, TArray<float>>& BoneNameAndWeights : CurrentWeightsMap)
+float FSkinToolWeights::GetWeightOfBoneOnVertex(
+	const int32 BoneIndex,
+	const int32 VertexID,
+	const TArray<VertexWeights>& InVertexWeights)
+{
+	const VertexWeights& VertexWeights = InVertexWeights[VertexID];
+	for (const FVertexBoneWeight& BoneWeight : VertexWeights)
 	{
-		const int32 BoneIndex = RefSkeleton.FindBoneIndex(BoneNameAndWeights.Key);
-		if (BoneIndex != INDEX_NONE)
+		if (BoneWeight.BoneIndex == BoneIndex)
 		{
-			BoneIndexWeightMap.Add(static_cast<FBoneIndexType>(BoneIndex), &BoneNameAndWeights.Value);
+			return BoneWeight.Weight;
 		}
 	}
 
-	FBoneWeightsSettings Settings;
-	Settings.SetNormalizeType(EBoneWeightNormalizeType::None);
+	return 0.f;
+}
 
-	TArray<FBoneWeight> SourceBoneWeights;
-	SourceBoneWeights.Reserve(MaxInlineBoneWeightCount);
+void FSkinToolWeights::SetWeightOfBoneOnVertex(
+	const int32 BoneIndex,
+	const int32 VertexID,
+	const float Weight,
+	TArray<VertexWeights>& InOutVertexWeights)
+{
+	Deformer.SetVertexNeedsUpdated(VertexID);
+	
+	// incoming weights are assumed to be normalized already, so set it directly
+	VertexWeights& VertexWeights = InOutVertexWeights[VertexID];
+	for (FVertexBoneWeight& BoneWeight : VertexWeights)
+	{
+		if (BoneWeight.BoneIndex == BoneIndex)
+		{
+			BoneWeight.Weight = Weight;
+			return;
+		}
+	}
 
-	const int32 NumVertices = EditedMesh->Vertices().Num();
-	for (int32 VertexIndex = 0; VertexIndex < NumVertices; VertexIndex++)
+	// bone not already an influence on this vertex, so we need to add it..
+
+	// if vertex has room for more influences, then simply add it
+	if (VertexWeights.Num() < UE::AnimationCore::MaxInlineBoneWeightCount)
 	{
-		SourceBoneWeights.Reset();
+		// add a new influence to this vertex
+		const FVector PosLocalToBone = Deformer.InvCSRefPoseTransforms[BoneIndex].TransformPosition(Deformer.RefPoseVertexPositions[VertexID]);
+		VertexWeights.Emplace(BoneIndex, PosLocalToBone, Weight);
+		return;
+	}
 
-		for (const TTuple<unsigned short, const TArray<float>*>& BoneIndexWeight : BoneIndexWeightMap)
+	//
+	// uh oh, we're out of room for more influences on this vertex, so lets kick the smallest influence to make room
+	//
+
+	// find the smallest influence
+	float SmallestInfluence = TNumericLimits<float>::Max();
+	int32 SmallestInfluenceIndex = INDEX_NONE;
+	for (int32 InfluenceIndex=0; InfluenceIndex<VertexWeights.Num(); ++InfluenceIndex)
+	{
+		const FVertexBoneWeight& BoneWeight = VertexWeights[InfluenceIndex];
+		if (BoneWeight.Weight <= SmallestInfluence)
 		{
-			SourceBoneWeights.Add(FBoneWeight(BoneIndexWeight.Key, (*BoneIndexWeight.Value)[VertexIndex]));
+			SmallestInfluence = BoneWeight.Weight;
+			SmallestInfluenceIndex = InfluenceIndex;
 		}
+	}
+
+	// replace smallest influence
+	FVertexBoneWeight& BoneWeightToReplace = VertexWeights[SmallestInfluenceIndex];
+	BoneWeightToReplace.Weight = Weight;
+	BoneWeightToReplace.BoneIndex = BoneIndex;
+	BoneWeightToReplace.VertexInBoneSpace = Deformer.InvCSRefPoseTransforms[BoneIndex].TransformPosition(Deformer.RefPoseVertexPositions[VertexID]);
 
-		VertexSkinWeights.Set(FVertexID(VertexIndex), FBoneWeights::Create(SourceBoneWeights, Settings));
+	// now we need to re-normalize because the stamp does not handle maximum influences
+	float TotalWeight = 0.f;
+	for (const FVertexBoneWeight& BoneWeight : VertexWeights)
+	{
+		TotalWeight += BoneWeight.Weight;
+	}
+	for (FVertexBoneWeight& BoneWeight : VertexWeights)
+	{
+		BoneWeight.Weight /= TotalWeight;
 	}
 }
 
 void FSkinToolWeights::ResetAfterStroke()
 {
-	PreStrokeWeightsMap = CurrentWeightsMap;
+	PreStrokeWeights = CurrentWeights;
 
 	for (int32 i=0; i<MaxFalloffPerVertexThisStroke.Num(); ++i)
 	{
@@ -318,21 +426,38 @@ void FSkinToolWeights::ResetAfterStroke()
 	}
 }
 
-float FSkinToolWeights::SetCurrentFalloffAndGetMaxFalloffThisStroke(int32 VertexID, float CurrentStength)
+float FSkinToolWeights::SetCurrentFalloffAndGetMaxFalloffThisStroke(int32 VertexID, float CurrentStrength)
 {
 	float& MaxFalloffThisStroke = MaxFalloffPerVertexThisStroke[VertexID];
-	if (MaxFalloffThisStroke < CurrentStength)
+	if (MaxFalloffThisStroke < CurrentStrength)
 	{
-		MaxFalloffThisStroke = CurrentStength;
+		MaxFalloffThisStroke = CurrentStrength;
 	}
 	return MaxFalloffThisStroke;
 }
 
+void FSkinToolWeights::ApplyEditsToWeightMap(
+	const FMultiBoneWeightEdits& Edits, 
+	TArray<VertexWeights>& InOutWeights)
+{
+	for (const TTuple<BoneIndex, FSingleBoneWeightEdits>& BoneWeightEdits : Edits.PerBoneWeightEdits)
+	{
+		const FSingleBoneWeightEdits& WeightEdits = BoneWeightEdits.Value;
+		const int32 BoneIndex = WeightEdits.BoneIndex;
+		for (const TTuple<int32, float>& NewWeight : WeightEdits.NewWeights)
+		{
+			const int32 VertexID = NewWeight.Key;
+			const float Weight = NewWeight.Value;
+			SetWeightOfBoneOnVertex(BoneIndex, VertexID, Weight, InOutWeights);
+		}
+	}
+}
+
 void FMeshSkinWeightsChange::Apply(UObject* Object)
 {
 	USkinWeightsPaintTool* Tool = CastChecked<USkinWeightsPaintTool>(Object);
 	
-	for (TTuple<FName, FSingleBoneWeightEdits>& Pair : AllWeightEdits.PerBoneWeightEdits)
+	for (TTuple<int32, FSingleBoneWeightEdits>& Pair : AllWeightEdits.PerBoneWeightEdits)
 	{
 		Tool->ExternalUpdateWeights(Pair.Key, Pair.Value.NewWeights);
 	}
@@ -342,7 +467,7 @@ void FMeshSkinWeightsChange::Revert(UObject* Object)
 {
 	USkinWeightsPaintTool* Tool = CastChecked<USkinWeightsPaintTool>(Object);
 
-	for (TTuple<FName, FSingleBoneWeightEdits>& Pair : AllWeightEdits.PerBoneWeightEdits)
+	for (TTuple<int32, FSingleBoneWeightEdits>& Pair : AllWeightEdits.PerBoneWeightEdits)
 	{
 		Tool->ExternalUpdateWeights(Pair.Key, Pair.Value.OldWeights);
 	}
@@ -364,6 +489,8 @@ UMeshSurfacePointTool* USkinWeightsPaintToolBuilder::CreateNewTool(const FToolBu
 
 void USkinWeightsPaintTool::Setup()
 {
+	TRACE_CPUPROFILER_EVENT_SCOPE(SkinTool::Setup);
+	
 	UDynamicMeshBrushTool::Setup();
 
 	const IPrimitiveComponentBackedTarget* TargetComponent = Cast<IPrimitiveComponentBackedTarget>(Target);
@@ -371,32 +498,21 @@ void USkinWeightsPaintTool::Setup()
 	const USkeletalMeshComponent* Component = Cast<USkeletalMeshComponent>(TargetComponent->GetOwnerComponent());
 	check(Component && Component->GetSkeletalMeshAsset())
 
-	// initialize the bone browser
-	// TODO replace this and synchronize selection & rendering with skeletal mesh editor
-	USkeletalMesh& SkeletalMesh = *Component->GetSkeletalMeshAsset();
-	BoneContainer.InitializeTo(Component->RequiredBones, Component->GetCurveFilterSettings(0)/* Always use the highest LOD */, SkeletalMesh);
-
 	// create a mesh description for editing (this must be done before calling UpdateBonePositionInfos) 
 	EditedMesh = MakeUnique<FMeshDescription>();
 	*EditedMesh = *UE::ToolTarget::GetMeshDescription(Target);
 
-	// Update the skeleton drawing information from the original bind pose
-	MaxDrawRadius = Component->Bounds.SphereRadius * 0.0025f;
-	UpdateBonePositionInfos(MaxDrawRadius);
-
 	// initialize the tool properties
 	BrushProperties->RestoreProperties(this); // hides strength and falloff
 	ToolProps = NewObject<USkinWeightsPaintToolProperties>(this);
 	ToolProps->RestoreProperties(this);
 	ToolProps->SkeletalMesh = Component->GetSkeletalMeshAsset();
-	ToolProps->CurrentBone.Initialize(BoneContainer);
 	AddToolPropertySource(ToolProps);
 	// attach callback to be informed when tool properties are modified
 	ToolProps->GetOnModified().AddUObject(this, &USkinWeightsPaintTool::OnToolPropertiesModified);
 
 	// default to the root bone as current bone
 	PendingCurrentBone = CurrentBone = ToolProps->SkeletalMesh->GetRefSkeleton().GetBoneName(0);
-	ToolProps->CurrentBone.BoneName = PendingCurrentBone.GetValue();
 
 	// configure preview mesh
 	PreviewMesh->SetTangentsMode(EDynamicMeshComponentTangentsMode::AutoCalculated);
@@ -417,15 +533,22 @@ void USkinWeightsPaintTool::Setup()
 		PreviewMesh->SetOverrideRenderMaterial(VtxColorMaterial);
 	}
 
-	// build octree for brush
+	// build octree for vertices
 	VerticesOctree.Initialize(PreviewMesh->GetMesh(), true);
 
-	// record data needed to compute deformations
-	Deformer.Initialize(Component, EditedMesh.Get());
-	
-	// copy weight maps
-	const FReferenceSkeleton& RefSkeleton = ToolProps->SkeletalMesh->GetRefSkeleton();
-	Weights.InitializeSkinWeights(RefSkeleton, EditedMesh.Get());
+	// build octree for triangles
+	{
+		TRACE_CPUPROFILER_EVENT_SCOPE(SkinTool::InitTriangleOctree);
+		
+		TriangleOctreeFuture = Async(SkinPaintToolAsyncExecTarget, [&]()
+		{
+			TRACE_CPUPROFILER_EVENT_SCOPE(SkinTool::InitTriangleOctreeRun);
+			TrianglesOctree.Initialize(PreviewMesh->GetMesh());
+		});
+	}
+
+	// initialize weight maps and deformation data
+	Weights.InitializeSkinWeights(Component, EditedMesh.Get());
 	bVisibleWeightsValid = false;
 
 	RecalculateBrushRadius();
@@ -463,21 +586,58 @@ void USkinWeightsPaintTool::OnTick(float DeltaTime)
 		ToolProps->bColorModeChanged = false;
 	}
 
-	// sparsely updates vertex positions (only when needed)
-	Deformer.UpdateVertexDeformation(Weights.CurrentWeightsMap, PreviewMesh, VerticesOctree);
-}
-
-void USkinWeightsPaintTool::Render(IToolsContextRenderAPI* RenderAPI)
-{
-	UDynamicMeshBrushTool::Render(RenderAPI);
-
-	// FIXME: Make selective.
-	//RenderBonePositions(RenderAPI->GetPrimitiveDrawInterface());
+	// sparsely updates vertex positions (only on vertices with modified weights)
+	Weights.Deformer.UpdateVertexDeformation(this);
 }
 
 bool USkinWeightsPaintTool::HitTest(const FRay& Ray, FHitResult& OutHit)
 {
-	return UDynamicMeshBrushTool::HitTest(Ray, OutHit);
+	// do not query the triangle octree until all async ops are finished
+	TriangleOctreeFuture.Wait();
+	
+	// put ray in local space of skeletal mesh component
+	// currently no way to transform skeletal meshes in the editor,
+	// but at some point in the future we may add the ability to move parts around
+	const IPrimitiveComponentBackedTarget* TargetComponent = Cast<IPrimitiveComponentBackedTarget>(Target);
+	const FTransform3d CurTargetTransform(TargetComponent->GetWorldTransform());
+	FRay3d LocalRay(
+		CurTargetTransform.InverseTransformPosition((FVector3d)Ray.Origin),
+		CurTargetTransform.InverseTransformVector((FVector3d)Ray.Direction));
+	UE::Geometry::Normalize(LocalRay.Direction);
+	
+	const FDynamicMesh3* Mesh = PreviewMesh->GetPreviewDynamicMesh();
+
+	FViewCameraState StateOut;
+	GetToolManager()->GetContextQueriesAPI()->GetCurrentViewState(StateOut);
+	FVector3d LocalEyePosition(CurTargetTransform.InverseTransformPosition((FVector3d)StateOut.Position));
+	const int32 TriID = TrianglesOctree.FindNearestHitObject(
+		LocalRay,
+		[this, Mesh, &LocalEyePosition](int TriangleID)
+	{
+		FVector3d Normal, Centroid;
+		double Area;
+		Mesh->GetTriInfo(TriangleID, Normal, Area, Centroid);
+		return Normal.Dot((Centroid - LocalEyePosition)) < 0;
+	});
+	
+	if (TriID != IndexConstants::InvalidID)
+	{	
+		FastTriWinding::FTriangle3d Triangle;
+		Mesh->GetTriVertices(TriID, Triangle.V[0], Triangle.V[1], Triangle.V[2]);
+		UE::Geometry::FIntrRay3Triangle3d Query(LocalRay, Triangle);
+		Query.Find();
+
+		StampLocalPos = LocalRay.PointAt(Query.RayParameter);
+		TriangleUnderStamp = TriID;
+
+		OutHit.FaceIndex = TriID;
+		OutHit.Distance = Query.RayParameter;
+		OutHit.Normal = CurTargetTransform.TransformVector(Mesh->GetTriNormal(TriID));
+		OutHit.ImpactPoint = CurTargetTransform.TransformPosition(StampLocalPos);
+		return true;
+	}
+	
+	return false;
 }
 
 void USkinWeightsPaintTool::OnBeginDrag(const FRay& WorldRay)
@@ -538,26 +698,31 @@ void USkinWeightsPaintTool::CalculateVertexROI(
 {
 	using namespace UE::Geometry;
 
+	TRACE_CPUPROFILER_EVENT_SCOPE(SkinTool::CalculateVertexROI);
+	
 	if (ToolProps->FalloffMode == EWeightBrushFalloffMode::Volume)
 	{
-		IPrimitiveComponentBackedTarget* TargetComponent = Cast<IPrimitiveComponentBackedTarget>(Target);
-		FTransform3d Transform(TargetComponent->GetWorldTransform());
-		FVector3d StampPosLocal = Transform.InverseTransformPosition((FVector3d)Stamp.WorldPosition);
-
-		float RadiusSqr = CurrentBrushRadius * CurrentBrushRadius;
+		const IPrimitiveComponentBackedTarget* TargetComponent = Cast<IPrimitiveComponentBackedTarget>(Target);
+		const FTransform3d Transform(TargetComponent->GetWorldTransform());
+		const FVector3d StampPosLocal = Transform.InverseTransformPosition((FVector3d)Stamp.WorldPosition);
+		const float RadiusSqr = CurrentBrushRadius * CurrentBrushRadius;
 		const FDynamicMesh3* Mesh = PreviewMesh->GetPreviewDynamicMesh();
-		FAxisAlignedBox3d QueryBox(StampPosLocal, CurrentBrushRadius);
+		const FAxisAlignedBox3d QueryBox(StampPosLocal, CurrentBrushRadius);
 		VerticesOctree.RangeQuery(QueryBox,
 			[&](int32 VertexID) { return FVector3d::DistSquared(Mesh->GetVertex(VertexID), StampPosLocal) < RadiusSqr; },
 			VertexROI);
 
-		for (int32 VertexID : VertexROI)
+		for (const int32 VertexID : VertexROI)
 		{
 			VertexSqDistances.Add(FVector3d::DistSquared(Mesh->GetVertex(VertexID), StampPosLocal));
 		}
+		
+		return;
 	}
-	else
+
+	if (ToolProps->FalloffMode == EWeightBrushFalloffMode::Surface)
 	{
+		// get coordinate frame from stamp
 		auto GetFrameFromStamp = [](const FBrushStampData& InStamp) -> FFrame3d
 		{
 			const FVector3d Origin = InStamp.WorldPosition;
@@ -571,57 +736,51 @@ void USkinWeightsPaintTool::CalculateVertexROI(
 				NonCollinear.Y = FMath::RandRange(-1.0f, 1.0f);
 				NonCollinear.Z = FMath::RandRange(-1.0f, 1.0f);
 				NonCollinear.Normalize();
-				
+					
 			} while (FMath::Abs(NonCollinear.Dot(Normal)) > 0.8f);
 
 			const FVector3d Plane = Normal.Cross(NonCollinear);
 			const FVector3d Cross = Plane.Cross(Normal);
 			return FFrame3d(Origin, Cross, Plane, Normal);
-			
+				
 		};
-
-		// get coordinate frame from stamp
-		FFrame3d SeedFrame = GetFrameFromStamp(Stamp);
-		
-		const FDynamicMesh3* Mesh = PreviewMesh->GetPreviewDynamicMesh();
-
-		// get triangle under the stamp
-		// we pass it's vertices to the ExpMap generator to act as seed nodes to start the search from
-		// this is necessary because the stamp may be in a large face center, from from any vertex on the mesh.
-		FDynamicMeshAABBTree3 Spatial(Mesh, true);
-		double NearDistSqr;
-		int32 SeedTID = Spatial.FindNearestTriangle(SeedFrame.Origin, NearDistSqr);
-		//FDistPoint3Triangle3d Query = TMeshQueries<FDynamicMesh3>::TriangleDistance(*Mesh, SeedTID, SeedFrame.Origin);
-		FIndex3i TriVerts = Mesh->GetTriangle(SeedTID);
-
+		const FFrame3d SeedFrame = GetFrameFromStamp(Stamp);
+			
 		// create the ExpMap generator, computes vertex polar coordinates in a plane tangent to the surface
+		const FDynamicMesh3* Mesh = PreviewMesh->GetPreviewDynamicMesh();
 		TMeshLocalParam<FDynamicMesh3> Param(Mesh);
 		Param.ParamMode = ELocalParamTypes::ExponentialMapUpwindAvg;
+		const FIndex3i TriVerts = Mesh->GetTriangle(TriangleUnderStamp);
 		Param.ComputeToMaxDistance(SeedFrame, TriVerts, Stamp.Radius);
-
-		// store vertices under the brush and their distances from the stamp
-		const float StampRadSq = FMath::Pow(Stamp.Radius, 2);
-		for (int32 VertexID : Mesh->VertexIndicesItr())
 		{
-			if (!Param.HasUV(VertexID))
-			{
-				continue;
-			}
-			
-			FVector2d UV = Param.GetUV(VertexID);
-			const float DistSq = UV.SizeSquared();
-			if (DistSq >= StampRadSq)
+			// store vertices under the brush and their distances from the stamp
+			const float StampRadSq = FMath::Pow(Stamp.Radius, 2);
+			for (int32 VertexID : Mesh->VertexIndicesItr())
 			{
-				continue;
+				if (!Param.HasUV(VertexID))
+				{
+					continue;
+				}
+				
+				FVector2d UV = Param.GetUV(VertexID);
+				const float DistSq = UV.SizeSquared();
+				if (DistSq >= StampRadSq)
+				{
+					continue;
+				}
+				
+				VertexSqDistances.Add(DistSq);
+				VertexROI.Add(VertexID);
 			}
-			
-			VertexSqDistances.Add(DistSq);
-			VertexROI.Add(VertexID);
 		}
+		
+		return;
 	}
+	
+	checkNoEntry();
 }
 
-FVector4f USkinWeightsPaintTool::WeightToColor(float Value)
+FVector4f USkinWeightsPaintTool::WeightToColor(float Value) const
 {
 	// optional greyscale mode
 	if (ToolProps->ColorMode == EWeightColorMode::Greyscale)
@@ -667,32 +826,28 @@ FVector4f USkinWeightsPaintTool::WeightToColor(float Value)
 
 void USkinWeightsPaintTool::UpdateCurrentBoneVertexColors()
 {
-	if (!Weights.CurrentWeightsMap.Contains(CurrentBone))
-	{
-		return;
-	}
+	const int32 CurrentBoneIndex = Weights.Deformer.BoneNameToIndexMap[CurrentBone];
 	
-	TArray<float>& SkinWeightsData = *Weights.CurrentWeightsMap.Find(CurrentBone);
-
 	// update mesh with new value colors
-	PreviewMesh->EditMesh([&](FDynamicMesh3& Mesh)
+	PreviewMesh->DeferredEditMesh([this, &CurrentBoneIndex](FDynamicMesh3& Mesh)
 	{
-		UE::Geometry::FNonManifoldMappingSupport NonManifoldMappingSupport(Mesh);
+		const UE::Geometry::FNonManifoldMappingSupport NonManifoldMappingSupport(Mesh);
 		UE::Geometry::FDynamicMeshColorOverlay* ColorOverlay = Mesh.Attributes()->PrimaryColors();
-		for (int32 ElementId : ColorOverlay->ElementIndicesItr())
+		for (const int32 ElementId : ColorOverlay->ElementIndicesItr())
 		{
-			const int32 VertexId = ColorOverlay->GetParentVertex(ElementId);	
-			const int32 SrcVertexId = NonManifoldMappingSupport.GetOriginalNonManifoldVertexID(VertexId);
-			const float Value = SkinWeightsData.IsValidIndex(SrcVertexId) ? SkinWeightsData[SrcVertexId] : 0.f;
+			const int32 VertexID = ColorOverlay->GetParentVertex(ElementId);	
+			const int32 SrcVertexID = NonManifoldMappingSupport.GetOriginalNonManifoldVertexID(VertexID);
+			const float Value = Weights.GetWeightOfBoneOnVertex(CurrentBoneIndex, SrcVertexID, Weights.CurrentWeights);
 			const FVector4f Color(WeightToColor(Value));
 			ColorOverlay->SetElement(ElementId, Color);
 		}
-	});
+	}, false);
+	PreviewMesh->NotifyDeferredEditCompleted(UPreviewMesh::ERenderUpdateMode::FastUpdate, EMeshRenderAttributeFlags::VertexColors, false);
 }
 
 double USkinWeightsPaintTool::CalculateBrushFalloff(double Distance) const
 {
-	double f = FMathd::Clamp(1.0 - BrushProperties->BrushFalloffAmount, 0.0, 1.0);
+	const double f = FMathd::Clamp(1.0 - BrushProperties->BrushFalloffAmount, 0.0, 1.0);
 	double d = Distance / CurrentBrushRadius;
 	double w = 1;
 	if (d > f)
@@ -706,6 +861,8 @@ double USkinWeightsPaintTool::CalculateBrushFalloff(double Distance) const
 
 void USkinWeightsPaintTool::ApplyStamp(const FBrushStampData& Stamp)
 {
+	TRACE_CPUPROFILER_EVENT_SCOPE(SkinTool::ApplyStamp);
+	
 	// get the vertices under the brush, and their squared distances to the brush center
 	// when using "Volume" brush, distances are straight line
 	// when using "Surface" brush, distances are geodesics
@@ -716,55 +873,64 @@ void USkinWeightsPaintTool::ApplyStamp(const FBrushStampData& Stamp)
 	// gather sparse set of modifications made from this stamp, these edits are merged throughout
 	// the lifetime of a single brush stroke in the "ActiveChange" allowing for undo/redo
 	FMultiBoneWeightEdits WeightEditsFromStamp;
-
-	// generate a weight edit from this stamp (includes modifications caused by normalization)
-	if (bSmoothStroke || ToolProps->BrushMode == EBrushBehaviorMode::Relax)
 	{
-		// use mesh topology to iteratively smooth weights across neighboring vertices
-		RelaxWeightOnVertices(VerticesInStamp, VertexSqDistances, WeightEditsFromStamp);
+		TRACE_CPUPROFILER_EVENT_SCOPE(SkinTool::EditWeightOfVerticesInStamp);
+		// generate a weight edit from this stamp (includes modifications caused by normalization)
+		if (bSmoothStroke || ToolProps->BrushMode == EBrushBehaviorMode::Relax)
+		{
+			// use mesh topology to iteratively smooth weights across neighboring vertices
+			RelaxWeightOnVertices(VerticesInStamp, VertexSqDistances, WeightEditsFromStamp);
+		}
+		else
+		{
+			// edit weight; either by "Add", "Remove", "Replace", "Multiply"
+			EditWeightOfVerticesInStamp(
+				ToolProps->BrushMode,
+				VerticesInStamp,
+				VertexSqDistances,
+				WeightEditsFromStamp);
+		}
 	}
-	else
+	
 	{
-		// edit weight; either by "Add", "Remove", "Replace", "Multiply"
-		EditWeightOfVerticesInStamp(
-			ToolProps->BrushMode,
-			VerticesInStamp,
-			VertexSqDistances,
-			WeightEditsFromStamp);
+		TRACE_CPUPROFILER_EVENT_SCOPE(SkinTool::ApplyWeightEditsToActiveChange);
+		// store weight edits from all stamps made during a single stroke (1 transaction per stroke)
+		for (const TTuple<int32, FSingleBoneWeightEdits>& BoneWeightEdits : WeightEditsFromStamp.PerBoneWeightEdits)
+		{
+			ActiveChange->AddBoneWeightEdit(BoneWeightEdits.Value);
+		}
 	}
-
-	// store weight edits from all stamps made during a single stroke (1 transaction per stroke)
-	for (const TTuple<FName, FSingleBoneWeightEdits>& BoneWeightEdits : WeightEditsFromStamp.PerBoneWeightEdits)
+	
 	{
-		ActiveChange->AddBoneWeightEdit(BoneWeightEdits.Value);
+		TRACE_CPUPROFILER_EVENT_SCOPE(SkinTool::ApplyWeightEditsToCurrentWeights);
+		// apply weights to current weights
+		Weights.ApplyEditsToWeightMap(WeightEditsFromStamp, Weights.CurrentWeights);
 	}
 
-	// apply weights to weight map
-	WeightEditsFromStamp.ApplyEditsToWeightMap(Weights.CurrentWeightsMap);
-
-	// weights have been modified, so update deformations
-	Deformer.SetVerticesNeedUpdated(VerticesInStamp);
-
-	// update vertex colors
-	PreviewMesh->DeferredEditMesh([&](FDynamicMesh3& Mesh)
 	{
-		TArray<int> ElementIds;
-		UE::Geometry::FDynamicMeshColorOverlay* ColorOverlay = Mesh.Attributes()->PrimaryColors();
-		const TArray<float>& VertexWeights = Weights.CurrentWeightsMap[CurrentBone];
-		const int32 NumVerticesInStamp = VerticesInStamp.Num();
-		for (int32 Index = 0; Index < NumVerticesInStamp; ++Index)
+		TRACE_CPUPROFILER_EVENT_SCOPE(SkinTool::UpdateVertexColors);
+		// update vertex colors
+		PreviewMesh->DeferredEditMesh([&](FDynamicMesh3& Mesh)
 		{
-			const int32 VertexId = VerticesInStamp[Index];
-			FVector4f NewColor(WeightToColor(VertexWeights[VertexId]));
-			ColorOverlay->GetVertexElements(VertexId, ElementIds);
-			for (const int32 ElementId : ElementIds)
+			TArray<int> ElementIds;
+			UE::Geometry::FDynamicMeshColorOverlay* ColorOverlay = Mesh.Attributes()->PrimaryColors();
+			const int32 BoneIndex = Weights.Deformer.BoneNameToIndexMap[CurrentBone];
+			const int32 NumVerticesInStamp = VerticesInStamp.Num();
+			for (int32 Index = 0; Index < NumVerticesInStamp; ++Index)
 			{
-				ColorOverlay->SetElement(ElementId, NewColor);
+				const int32 VertexID = VerticesInStamp[Index];
+				const float Weight = Weights.GetWeightOfBoneOnVertex(BoneIndex, VertexID, Weights.CurrentWeights);
+				FVector4f NewColor(WeightToColor(Weight));
+				ColorOverlay->GetVertexElements(VertexID, ElementIds);
+				for (const int32 ElementId : ElementIds)
+				{
+					ColorOverlay->SetElement(ElementId, NewColor);
+				}
+				ElementIds.Reset();
 			}
-			ElementIds.Reset();
-		}
-	}, false);
-	PreviewMesh->NotifyDeferredEditCompleted(UPreviewMesh::ERenderUpdateMode::FastUpdate, EMeshRenderAttributeFlags::VertexColors, false);
+		}, false);
+		PreviewMesh->NotifyDeferredEditCompleted(UPreviewMesh::ERenderUpdateMode::FastUpdate, EMeshRenderAttributeFlags::VertexColors, false);
+	}
 }
 
 void USkinWeightsPaintTool::RelaxWeightOnVertices(
@@ -775,56 +941,55 @@ void USkinWeightsPaintTool::RelaxWeightOnVertices(
 	const FDynamicMesh3* CurrentMesh = PreviewMesh->GetMesh();
 	const UE::Geometry::FNonManifoldMappingSupport NonManifoldMappingSupport(*CurrentMesh);
 
-	auto NormalizeWeights = [](TMap<FName, float>& InOutWeights)
+	auto NormalizeWeights = [](TMap<BoneIndex, float>& InOutWeights)
 	{
 		float TotalWeight = 0.f;
-		for (const TTuple<FName, float>& Weight : InOutWeights)
+		for (const TTuple<BoneIndex, float>& Weight : InOutWeights)
 		{
 			TotalWeight += Weight.Value;
 		}
-
-		for (TTuple<FName, float>& Weight : InOutWeights)
+		for (TTuple<BoneIndex, float>& Weight : InOutWeights)
 		{
 			Weight.Value /= TotalWeight;
 		}
 	};
 
 	// for each vertex in the stamp...
+	constexpr int32 AvgNumNeighbors = 8;
+	using VertexNeighborWeights = TArray<float, TInlineAllocator<AvgNumNeighbors>>;
 	TArray<int32> AllNeighborVertices;
-	TMap<FName, TArray<float>> WeightsOnAllNeighbors;
-	TMap<FName, float> FinalWeights;
+	TMap<BoneIndex, VertexNeighborWeights> WeightsOnAllNeighbors;
+	TMap<BoneIndex, float> FinalWeights;
 	for (int32 Index = 0; Index < VerticesInStamp.Num(); ++Index)
 	{
-		const int32 VertexId = VerticesInStamp[Index];
-		const int32 SrcVertexId = NonManifoldMappingSupport.GetOriginalNonManifoldVertexID(VertexId);
+		const int32 VertexID = VerticesInStamp[Index];
+		const int32 SrcVertexID = NonManifoldMappingSupport.GetOriginalNonManifoldVertexID(VertexID);
 
 		// get list of all neighboring vertices, AND this vertex
 		AllNeighborVertices.Reset();
-		AllNeighborVertices.Add(VertexId);
-		for (const int32 NeighborVertexId : CurrentMesh->VtxVerticesItr(SrcVertexId))
+		AllNeighborVertices.Add(VertexID);
+		for (const int32 NeighborVertexID : CurrentMesh->VtxVerticesItr(SrcVertexID))
 		{
-			AllNeighborVertices.Add(NeighborVertexId);
+			AllNeighborVertices.Add(NeighborVertexID);
 		}
 
 		// get all weights above a given threshold across ALL neighbors (including self)
 		WeightsOnAllNeighbors.Reset();
-		for (const int32 VertexID : AllNeighborVertices)
+		for (const int32 NeighborVertexID : AllNeighborVertices)
 		{
-			for (TTuple<FName, TArray<float>> Pair : Weights.CurrentWeightsMap)
+			for (const FVertexBoneWeight& BoneWeight : Weights.PreStrokeWeights[NeighborVertexID])
 			{
-				const float NeighborWeightValue = Pair.Value[VertexID];
-				if (NeighborWeightValue > MinimumWeightThreshold)
+				if (BoneWeight.Weight > MinimumWeightThreshold)
 				{
-					const FName BoneName = Pair.Key;
-					TArray<float>& BoneWeights = WeightsOnAllNeighbors.FindOrAdd(BoneName);
-					BoneWeights.Add(NeighborWeightValue);
+					VertexNeighborWeights& BoneWeights = WeightsOnAllNeighbors.FindOrAdd(BoneWeight.BoneIndex);
+					BoneWeights.Add(BoneWeight.Weight);
 				}
 			}
 		}
 
 		// calculate single average weight of each bone on all the neighbors
 		FinalWeights.Reset();
-		for (const TTuple<FName, TArray<float>>& NeighborWeights : WeightsOnAllNeighbors)
+		for (const TTuple<BoneIndex, VertexNeighborWeights>& NeighborWeights : WeightsOnAllNeighbors)
 		{
 			float TotalWeightOnThisBone = 0.f;
 			for (const float& Value : NeighborWeights.Value)
@@ -839,13 +1004,13 @@ void USkinWeightsPaintTool::RelaxWeightOnVertices(
 
 		// lerp weights from previous values, to fully relaxed values by brush strength scaled by falloff
 		const float CurrentFalloff = static_cast<float>(CalculateBrushFalloff(FMath::Sqrt(VertexSqDistances[Index])));
-		const float UseFalloff = Weights.SetCurrentFalloffAndGetMaxFalloffThisStroke(VertexId, CurrentFalloff);
+		const float UseFalloff = Weights.SetCurrentFalloffAndGetMaxFalloffThisStroke(VertexID, CurrentFalloff);
 		const float UseStrength = BrushProperties->BrushStrength * UseFalloff;
-		for (TTuple<FName, float>& FinalWeight : FinalWeights)
+		for (TTuple<BoneIndex, float>& FinalWeight : FinalWeights)
 		{
-			FName BoneName = FinalWeight.Key;
+			const int32 BoneIndex = FinalWeight.Key;
 			float NewWeight = FinalWeight.Value;
-			float OldWeight = Weights.PreStrokeWeightsMap[BoneName][VertexId];
+			float OldWeight = Weights.GetWeightOfBoneOnVertex(BoneIndex, VertexID, Weights.PreStrokeWeights);
 			FinalWeight.Value = FMath::Lerp(OldWeight, NewWeight, UseStrength);
 		}
 
@@ -853,13 +1018,13 @@ void USkinWeightsPaintTool::RelaxWeightOnVertices(
 		NormalizeWeights(FinalWeights);
 
 		// apply weight edits
-		for (const TTuple<FName, float>& FinalWeight : FinalWeights)
+		for (const TTuple<BoneIndex, float>& FinalWeight : FinalWeights)
 		{
 			// record an edit for this vertex, for this bone
-			FName BoneName = FinalWeight.Key;
-			float NewWeight = FinalWeight.Value;
-			float OldWeight = Weights.PreStrokeWeightsMap[BoneName][VertexId];
-			AllBoneWeightEditsFromStamp.MergeSingleEdit(BoneName, VertexId, OldWeight, NewWeight);
+			const int32 BoneIndex = FinalWeight.Key;
+			const float NewWeight = FinalWeight.Value;
+			const float OldWeight = Weights.GetWeightOfBoneOnVertex(BoneIndex, VertexID, Weights.PreStrokeWeights);
+			AllBoneWeightEditsFromStamp.MergeSingleEdit(BoneIndex, VertexID, OldWeight, NewWeight);
 		}
 	}
 }
@@ -871,7 +1036,6 @@ void USkinWeightsPaintTool::EditWeightOfVerticesInStamp(
 	FMultiBoneWeightEdits& AllBoneWeightEditsFromStamp)
 {
 	const UE::Geometry::FNonManifoldMappingSupport NonManifoldMappingSupport(*PreviewMesh->GetMesh());
-	TArray<float>& PreWeightData = *Weights.PreStrokeWeightsMap.Find(CurrentBone);
 
 	// invert brush strength differently depending on brush mode
 	float UseStrength = BrushProperties->BrushStrength;
@@ -896,16 +1060,17 @@ void USkinWeightsPaintTool::EditWeightOfVerticesInStamp(
 		checkNoEntry();
 	}
 	
-	// spin through the vertices in the stamp and store new weight values in NewValuesFromStap
-	// afterwards, these values are normalized while taking into consideration the users desired changes
+	// spin through the vertices in the stamp and store new weight values in NewValuesFromStamp
+	// afterwards, these values are normalized while taking into consideration the user's desired changes
+	const int32 CurrentBoneIndex = Weights.Deformer.BoneNameToIndexMap[CurrentBone];
 	const int32 NumVerticesInStamp = VerticesInStamp.Num();
 	for (int32 Index = 0; Index < NumVerticesInStamp; ++Index)
 	{
-		const int32 VertexId = VerticesInStamp[Index];
-		const int32 SrcVertexId = NonManifoldMappingSupport.GetOriginalNonManifoldVertexID(VertexId);
+		const int32 VertexID = VerticesInStamp[Index];
+		const int32 SrcVertexID = NonManifoldMappingSupport.GetOriginalNonManifoldVertexID(VertexID);
 		const float CurrentFalloff = static_cast<float>(CalculateBrushFalloff(FMath::Sqrt(VertexSqDistances[Index])));
-		const float UseFalloff = Weights.SetCurrentFalloffAndGetMaxFalloffThisStroke(SrcVertexId, CurrentFalloff);
-		const float ValueBeforeStroke = PreWeightData[SrcVertexId];
+		const float UseFalloff = Weights.SetCurrentFalloffAndGetMaxFalloffThisStroke(SrcVertexID, CurrentFalloff);
+		const float ValueBeforeStroke = Weights.GetWeightOfBoneOnVertex(CurrentBoneIndex, VertexID, Weights.PreStrokeWeights);
 
 		// calculate new weight value
 		float NewValueAfterStamp = ValueBeforeStroke;
@@ -936,7 +1101,7 @@ void USkinWeightsPaintTool::EditWeightOfVerticesInStamp(
 		NewValueAfterStamp = FMath::Clamp(NewValueAfterStamp, 0.0f, 1.0f);
 		Weights.EditVertexWeightAndNormalize(
 			CurrentBone,
-			VertexId,
+			VertexID,
 			NewValueAfterStamp,
 			AllBoneWeightEditsFromStamp);
 	}
@@ -955,8 +1120,7 @@ void USkinWeightsPaintTool::OnShutdown(EToolShutdownType ShutdownType)
 	if (ShutdownType == EToolShutdownType::Accept)
 	{
 		// apply the weights to the mesh description
-		const FReferenceSkeleton& RefSkeleton = ToolProps->SkeletalMesh->GetRefSkeleton();
-		Weights.ApplyCurrentWeightsToMeshDescription(RefSkeleton, EditedMesh.Get());
+		Weights.ApplyCurrentWeightsToMeshDescription(EditedMesh.Get());
 
 		// this block bakes the modified DynamicMeshComponent back into the StaticMeshComponent inside an undo transaction
 		GetToolManager()->BeginUndoTransaction(LOCTEXT("SkinWeightsPaintTool", "Paint Skin Weights"));
@@ -965,151 +1129,6 @@ void USkinWeightsPaintTool::OnShutdown(EToolShutdownType ShutdownType)
 	}
 }
 
-void USkinWeightsPaintTool::UpdateBonePositionInfos(float MinRadius)
-{
-	const FReferenceSkeleton& RefSkeleton = BoneContainer.GetReferenceSkeleton();
-	const TArray<FMeshBoneInfo>& BoneInfos = RefSkeleton.GetRefBoneInfo();
-	const TArray<FTransform>& BonePoses = RefSkeleton.GetRefBonePose();
-
-	BonePositionInfos.Reset();
-
-	// Exclude virtual bones.
-	for (int BoneIndex = 0; BoneIndex < RefSkeleton.GetRawBoneNum(); BoneIndex++)
-	{
-		FTransform Xform = BonePoses[BoneIndex];
-		int32 ParentBoneIndex = BoneInfos[BoneIndex].ParentIndex;
-
-		while (ParentBoneIndex != INDEX_NONE)
-		{
-			Xform = Xform * BonePoses[ParentBoneIndex];
-			ParentBoneIndex = BoneInfos[ParentBoneIndex].ParentIndex;
-		}
-
-		BonePositionInfos.Add({ BoneInfos[BoneIndex].Name, BoneInfos[BoneIndex].ParentIndex, Xform.GetLocation(), -1.0f });
-	}
-
-	// Populate the children.
-	for (int BoneIndex = 0; BoneIndex < BonePositionInfos.Num(); BoneIndex++)
-	{
-		FBonePositionInfo& BoneInfo = BonePositionInfos[BoneIndex];
-		if (BoneInfo.ParentBoneIndex != INDEX_NONE)
-		{
-			BonePositionInfos[BoneInfo.ParentBoneIndex].ChildBones.Add(BoneInfo.BoneName, BoneIndex);
-		}
-	}
-
-	bool bComputedRadius = true;
-	while (bComputedRadius)
-	{
-		bComputedRadius = false;
-
-		for (int BoneIndex = 0; BoneIndex < BonePositionInfos.Num(); BoneIndex++)
-		{
-			FBonePositionInfo& BoneInfo = BonePositionInfos[BoneIndex];
-			if (BoneInfo.Radius > 0.0f)
-			{
-				continue;
-			}
-
-			if (BoneInfo.ParentBoneIndex == INDEX_NONE)
-			{
-				if (BoneInfo.ChildBones.Num())
-				{
-					int32 Count = 0;
-					float RadiusSum = 0.0f;
-					for (const auto& CB : BoneInfo.ChildBones)
-					{
-						const FBonePositionInfo& ChildBoneInfo = BonePositionInfos[CB.Value];
-						if (ChildBoneInfo.Radius > 0.0f)
-						{
-							RadiusSum += ChildBoneInfo.Radius;
-							Count++;
-						}
-					}
-					if (BoneInfo.ChildBones.Num() == Count)
-					{
-						BoneInfo.Radius = RadiusSum / float(Count);
-						bComputedRadius = true;
-					}
-				}
-				else
-				{
-					// No children either? Take the whole mesh.
-					BoneInfo.Radius = EditedMesh->GetBounds().SphereRadius;
-				}
-			}
-			else 
-			{
-				BoneInfo.Radius = FVector::Dist(BoneInfo.Position, BonePositionInfos[BoneInfo.ParentBoneIndex].Position) / 2.0f;
-				bComputedRadius = true;
-			}
-
-			if (bComputedRadius)
-			{
-				BoneInfo.Radius = FMath::Max(BoneInfo.Radius, MinRadius);
-			}
-		}
-	}
-}
-
-
-void USkinWeightsPaintTool::RenderBonePositions(FPrimitiveDrawInterface* PDI)
-{
-	static const int32 NumSphereSides = 10;
-	static const int32 NumConeSides = 4;
-
-	IPrimitiveComponentBackedTarget* TargetComponent = Cast<IPrimitiveComponentBackedTarget>(Target);
-	FTransform WorldTransform = TargetComponent->GetWorldTransform();
-
-	for (const FBonePositionInfo& BoneInfo : BonePositionInfos)
-	{
-		FLinearColor BoneColor;
-		FVector Start, End;
-
-		End = BoneInfo.Position;
-		End = WorldTransform.TransformPosition(End);
-
-		if (BoneInfo.ParentBoneIndex != INDEX_NONE)
-		{
-			Start = BonePositionInfos[BoneInfo.ParentBoneIndex].Position;
-			Start = WorldTransform.TransformPosition(Start);
-			BoneColor = FLinearColor::White;
-		}
-		else
-		{
-			// Root bone.
-			BoneColor = FLinearColor::Red;
-		}
-		BoneColor.A = 0.10f;
-
-		if (BoneInfo.BoneName == CurrentBone)
-		{
-			BoneColor = FLinearColor(1.0f, 0.34f, 0.0f, 0.75f);
-		}
-
-		const float BoneLength = (End - Start).Size();
-		// clamp by bound, we don't want too long or big
-		const float Radius = FMath::Clamp(BoneLength * 0.05f, 0.1f, MaxDrawRadius);
-
-		// Render Sphere for bone end point and a cone between it and its parent.
-		DrawWireSphere(PDI, End, BoneColor, Radius, NumSphereSides, SDPG_Foreground, 0.0f, 1.0f);
-
-		if (BoneInfo.ParentBoneIndex != INDEX_NONE)
-		{
-			// Calc cone size 
-			const FVector EndToStart = (Start - End);
-			const float ConeLength = EndToStart.Size();
-			const float Angle = FMath::RadiansToDegrees(FMath::Atan(Radius / ConeLength));
-
-			TArray<FVector> Verts;
-			DrawWireCone(PDI, Verts, FRotationMatrix::MakeFromX(EndToStart) * FTranslationMatrix(End), ConeLength, Angle, NumConeSides, BoneColor, SDPG_Foreground, 0.0f, 1.0f);
-		}
-
-		// SkeletalDebugRendering::DrawWireBone(PDI, Start, End, BoneColor, SDPG_Foreground, Radius);
-	}
-}
-
-
 void USkinWeightsPaintTool::BeginChange()
 {
 	ActiveChange = MakeUnique<FMeshSkinWeightsChange>();
@@ -1122,22 +1141,17 @@ TUniquePtr<FMeshSkinWeightsChange> USkinWeightsPaintTool::EndChange()
 }
 
 
-void USkinWeightsPaintTool::ExternalUpdateWeights(const FName& BoneName, const TMap<int32, float>& NewValues)
+void USkinWeightsPaintTool::ExternalUpdateWeights(const int32 BoneIndex, const TMap<int32, float>& NewValues)
 {
-	TArray<float>* PostSkinWeightValues = Weights.CurrentWeightsMap.Find(BoneName);
-	TArray<float>* PreSkinWeightValues = Weights.PreStrokeWeightsMap.Find(BoneName);
-	if (PostSkinWeightValues == nullptr)
-	{
-		return;
-	}
-
-	for (const TTuple<int, float>& IV : NewValues)
+	for (const TTuple<int32, float>& Pair : NewValues)
 	{
-		PostSkinWeightValues->GetData()[IV.Key] = IV.Value;
-		PreSkinWeightValues->GetData()[IV.Key] = IV.Value;
-		Deformer.SetVertexNeedsUpdated(IV.Key);
+		const int32 VertexID = Pair.Key;
+		const float Weight = Pair.Value;
+		Weights.SetWeightOfBoneOnVertex(BoneIndex, VertexID, Weight, Weights.CurrentWeights);
+		Weights.SetWeightOfBoneOnVertex(BoneIndex, VertexID, Weight, Weights.PreStrokeWeights);
 	}
 
+	const FName BoneName = Weights.Deformer.BoneNames[BoneIndex];
 	if (BoneName == CurrentBone)
 	{
 		UpdateCurrentBoneVertexColors();
@@ -1166,18 +1180,7 @@ void USkinWeightsPaintTool::HandleSkeletalMeshModified(const TArray<FName>& InBo
 }
 
 void USkinWeightsPaintTool::OnToolPropertiesModified(UObject* ModifiedObject, FProperty* ModifiedProperty)
-{
-	// bone changed?
-	const bool bCurrentBoneModified = ModifiedProperty->GetNameCPP() == GET_MEMBER_NAME_STRING_CHECKED(FBoneReference, BoneName);
-	if (bCurrentBoneModified)
-	{
-		PendingCurrentBone = ToolProps->CurrentBone.BoneName;
-		if (NeedsNotification())
-		{
-			GetNotifier().Notify({ToolProps->CurrentBone.BoneName}, ESkeletalMeshNotifyType::BonesSelected);
-		}
-	}
-	
+{	
 	// invalidate vertex color cache when weight color properties are modified
 	const bool bColorModeModified = ModifiedProperty->GetNameCPP() == GET_MEMBER_NAME_STRING_CHECKED(USkinWeightsPaintToolProperties, ColorMode);
 	const bool bColorRampModified = ModifiedProperty->GetNameCPP() == GET_MEMBER_NAME_STRING_CHECKED(USkinWeightsPaintToolProperties, ColorRamp);
@@ -1190,12 +1193,4 @@ void USkinWeightsPaintTool::OnToolPropertiesModified(UObject* ModifiedObject, FP
 }
 
 
-USkeleton* USkinWeightsPaintToolProperties::GetSkeleton(bool& bInvalidSkeletonIsError, const IPropertyHandle* PropertyHandle)
-{
-	bInvalidSkeletonIsError = false;
-	return SkeletalMesh ? SkeletalMesh->GetSkeleton() : nullptr;
-}
-
-
-#undef LOCTEXT_NAMESPACE
-
+#undef LOCTEXT_NAMESPACE
\ No newline at end of file
diff --git a/Engine/Plugins/Experimental/MeshModelingToolsetExp/Source/MeshModelingToolsExp/Public/SkinWeightsPaintTool.h b/Engine/Plugins/Experimental/MeshModelingToolsetExp/Source/MeshModelingToolsExp/Public/SkinWeightsPaintTool.h
--- a/Engine/Plugins/Experimental/MeshModelingToolsetExp/Source/MeshModelingToolsExp/Public/SkinWeightsPaintTool.h
+++ b/Engine/Plugins/Experimental/MeshModelingToolsetExp/Source/MeshModelingToolsExp/Public/SkinWeightsPaintTool.h
@@ -4,13 +4,12 @@
 
 #include "DynamicMeshBrushTool.h"
 #include "BaseTools/MeshSurfacePointMeshEditingTool.h"
-#include "Changes/ValueWatcher.h"
 #include "MeshDescription.h"
 #include "DynamicMesh/DynamicVerticesOctree3.h"
-#include "BoneContainer.h"
-#include "Interfaces/Interface_BoneReferenceSkeletonProvider.h"
+#include "BoneWeights.h"
 #include "Misc/Optional.h"
 #include "Containers/Map.h"
+#include "DynamicMesh/DynamicMeshOctree3.h"
 #include "SkeletalMesh/SkeletalMeshEditionInterface.h"
 
 #include "SkinWeightsPaintTool.generated.h"
@@ -19,159 +18,140 @@
 struct FMeshDescription;
 class USkinWeightsPaintTool;
 
-// store a map of bone names to weight arrays, where array length == num vertices.
-// this struct stores ALL weights for ALL bones in the entire mesh.
-using PerBoneWeightMap = TMap<FName, TArray<float>>;
+using BoneIndex = int32;
+using VertexIndex = int32;
 
-// store a sparse set of modifications to a set of vertex weights on a SINGLE bone
-struct FSingleBoneWeightEdits
+namespace SkinPaintTool
 {
-	FName BoneName;
-	TMap<int32, float> OldWeights;
-	TMap<int32, float> NewWeights;
-};
-
-// store a sparse set of modifications to a set of vertex weights for a SET of bones
-// with support for merging edits. these are used for transaction history undo/redo.
-struct FMultiBoneWeightEdits
-{
-	void MergeSingleEdit(const FName BoneName, const int32 VertexID, const float OldWeight, const float NewWeight)
-	{
-		PerBoneWeightEdits.FindOrAdd(BoneName);
-		PerBoneWeightEdits[BoneName].BoneName = BoneName;
-		PerBoneWeightEdits[BoneName].NewWeights.Add(VertexID, NewWeight);
-		PerBoneWeightEdits[BoneName].OldWeights.FindOrAdd(VertexID, OldWeight);
-	}
-	
-	void MergeEdits(const FSingleBoneWeightEdits& BoneWeightEdits)
-	{
-		// make sure bone has an entry in the map of weight edits
-		const FName Bone = BoneWeightEdits.BoneName;
-		PerBoneWeightEdits.FindOrAdd(Bone);
-		PerBoneWeightEdits[Bone].BoneName = Bone;
-	
-		for (const TTuple<int32, float>& NewWeight : BoneWeightEdits.NewWeights)
-		{
-			int32 VertexIndex = NewWeight.Key;
-			PerBoneWeightEdits[Bone].NewWeights.Add(VertexIndex, NewWeight.Value);
-			PerBoneWeightEdits[Bone].OldWeights.FindOrAdd(VertexIndex, BoneWeightEdits.OldWeights[VertexIndex]);
-		}
-	}
+	struct FSkinToolWeights;
 
-	float GetVertexDeltaFromEdits(FName BoneName, int32 VertexIndex)
+	struct FVertexBoneWeight
 	{
-		PerBoneWeightEdits.FindOrAdd(BoneName);
-		if (const float* NewVertexWeight = PerBoneWeightEdits[BoneName].NewWeights.Find(VertexIndex))
-		{
-			return *NewVertexWeight - PerBoneWeightEdits[BoneName].OldWeights[VertexIndex];
-		}
+		FVertexBoneWeight() : BoneIndex(INDEX_NONE), VertexInBoneSpace(FVector::ZeroVector), Weight(0.0f) {}
+		FVertexBoneWeight(int32 InBoneIndex, const FVector& InPosInRefPose, float InWeight) :
+			BoneIndex(InBoneIndex), VertexInBoneSpace(InPosInRefPose), Weight(InWeight){}
+		
+		int32 BoneIndex;
+		FVector VertexInBoneSpace;
+		float Weight;
+	};
 
-		return 0.0f;
-	}
+	using VertexWeights = TArray<FVertexBoneWeight, TFixedAllocator<UE::AnimationCore::MaxInlineBoneWeightCount>>;
 
-	void ApplyEditsToWeightMap(PerBoneWeightMap& InOutWeightMap)
+	// data required to preview the skinning deformations as you paint
+	struct FSkinToolDeformer
 	{
-		for (TTuple<FName, FSingleBoneWeightEdits>& BoneWeightEdits : PerBoneWeightEdits)
-		{
-			const FSingleBoneWeightEdits& WeightEdits = BoneWeightEdits.Value;
-			TArray<float>& VertexWeights = InOutWeightMap[WeightEdits.BoneName];
-			for (const TTuple<int32, float>& NewWeight : WeightEdits.NewWeights)
-			{
-				VertexWeights[NewWeight.Key] = NewWeight.Value;
-			}
-		}
-	}
-	
-	TMap<FName, FSingleBoneWeightEdits> PerBoneWeightEdits;
-};
-
-// data required to preview the skinning deformations as you paint
-struct FSkinToolDeformer
-{
-	void Initialize(
-		const USkeletalMeshComponent* SkeletalMeshComponent,
-		const FMeshDescription* Mesh);
+		void Initialize(const USkeletalMeshComponent* SkeletalMeshComponent, const FMeshDescription* Mesh);
 
-	void UpdateVertexDeformation(
-		PerBoneWeightMap& Weights,
-		UPreviewMesh* PreviewMesh,
-		UE::Geometry::TDynamicVerticesOctree3<FDynamicMesh3>& OctreeToUpdate);
+		void UpdateVertexDeformation(USkinWeightsPaintTool* Tool);
 
-	void SetVerticesNeedUpdated(TArray<int32> VertexIndices);
-
-	void SetVertexNeedsUpdated(int32 VertexIndex);
+		void SetVertexNeedsUpdated(int32 VertexIndex);
 	
-	// which vertices require updating (partially re-calculated skinning deformation while paintin)
-	TSet<int32> VerticesWithModifiedWeights;
-	// position of all vertices in the reference pose
-	TArray<FVector> RefPoseVertexPositions;
-	// inverted, component space ref pose transform of each bone
-	TArray<FTransform> InvCSRefPoseTransforms;
-	// bone name to bone index
-	TMap<FName, int32> BoneNameToIndexMap;
-	// the skeletal mesh to get the current pose from
-	const USkeletalMeshComponent* Component;
-};
-
-class MESHMODELINGTOOLSEXP_API FMeshSkinWeightsChange : public FToolCommandChange
-{
-public:
-	FMeshSkinWeightsChange() : FToolCommandChange()	{ }
+		// which vertices require updating (partially re-calculated skinning deformation while painting)
+		TSet<int32> VerticesWithModifiedWeights;
+		// position of all vertices in the reference pose
+		TArray<FVector> RefPoseVertexPositions;
+		// inverted, component space ref pose transform of each bone
+		TArray<FTransform> InvCSRefPoseTransforms;
+		// bone index to bone name
+		TArray<FName> BoneNames;
+		TMap<FName, BoneIndex> BoneNameToIndexMap;
+		// the skeletal mesh to get the current pose from
+		const USkeletalMeshComponent* Component;
+	};
 
-	virtual FString ToString() const override
+	// store a sparse set of modifications to a set of vertex weights on a SINGLE bone
+	struct FSingleBoneWeightEdits
 	{
-		return FString(TEXT("Paint Skin Weights"));
-	}
+		int32 BoneIndex;
+		TMap<VertexIndex, float> OldWeights;
+		TMap<VertexIndex, float> NewWeights;
+	};
 
-	void Apply(UObject* Object) override;
+	// store a sparse set of modifications to a set of vertex weights for a SET of bones
+	// with support for merging edits. these are used for transaction history undo/redo.
+	struct FMultiBoneWeightEdits
+	{
+		void MergeSingleEdit(const int32 BoneIndex, const int32 VertexID, const float OldWeight, const float NewWeight);
+		void MergeEdits(const FSingleBoneWeightEdits& BoneWeightEdits);
+		float GetVertexDeltaFromEdits(const int32 BoneIndex, const int32 VertexIndex);
 
-	void Revert(UObject* Object) override;
+		// map of bone indices to weight edits made to that bone
+		TMap<BoneIndex, FSingleBoneWeightEdits> PerBoneWeightEdits;
+	};
 
-	void AddBoneWeightEdit(const FSingleBoneWeightEdits& BoneWeightEdit);
+	class FMeshSkinWeightsChange : public FToolCommandChange
+	{
+	public:
+		FMeshSkinWeightsChange() : FToolCommandChange()	{ }
 
-private:
-	FMultiBoneWeightEdits AllWeightEdits;
-};
+		virtual FString ToString() const override
+		{
+			return FString(TEXT("Paint Skin Weights"));
+		}
 
-// intermediate storage of the weight maps for duration of tool
-struct FSkinToolWeights
-{
-	// copy the initial weight values from the skeletal mesh
-	void InitializeSkinWeights(
-		const FReferenceSkeleton& RefSkeleton,
-		FMeshDescription* EditedMesh);
+		void Apply(UObject* Object) override;
 
-	// get an array of weights, per-vertex for the given bone
-	TArray<float>* GetWeightsForBone(const FName BoneName);
+		void Revert(UObject* Object) override;
 
-	// applies an edit to a single vertex weight on a single bone, then normalizes the remaining weights while
-	// keeping the edited weight intact (ie, adapts OTHER influences to achieve normalization)
-	void EditVertexWeightAndNormalize(
-		const FName BoneToHoldConstant,
-		const int32 VertexId,
-		const float NewWeightValue,
-		FMultiBoneWeightEdits& WeightEdits);
+		void AddBoneWeightEdit(const FSingleBoneWeightEdits& BoneWeightEdit);
 
-	void ApplyCurrentWeightsToMeshDescription(
-		const FReferenceSkeleton& RefSkeleton,
-		FMeshDescription* EditedMesh);
+	private:
+		FMultiBoneWeightEdits AllWeightEdits;
+	};
 
-	void ResetAfterStroke();
+	// intermediate storage of the weight maps for duration of tool
+	struct FSkinToolWeights
+	{
+		// copy the initial weight values from the skeletal mesh
+		void InitializeSkinWeights(
+			const USkeletalMeshComponent* SkeletalMeshComponent,
+			FMeshDescription* Mesh);
+
+		// applies an edit to a single vertex weight on a single bone, then normalizes the remaining weights while
+		// keeping the edited weight intact (ie, adapts OTHER influences to achieve normalization)
+		void EditVertexWeightAndNormalize(
+			const FName BoneToHoldConstant,
+			const int32 VertexId,
+			const float NewWeightValue,
+			FMultiBoneWeightEdits& WeightEdits);
+
+		void ApplyCurrentWeightsToMeshDescription(FMeshDescription* EditedMesh);
+
+		static float GetWeightOfBoneOnVertex(
+			const int32 BoneIndex,
+			const int32 VertexID,
+			const TArray<VertexWeights>& InVertexWeights);
+
+		void SetWeightOfBoneOnVertex(
+			const int32 BoneIndex,
+			const int32 VertexID,
+			const float Weight,
+			TArray<VertexWeights>& InOutVertexData);
+
+		void ResetAfterStroke();
+
+		float SetCurrentFalloffAndGetMaxFalloffThisStroke(int32 VertexID, float CurrentStrength);
+
+		void ApplyEditsToWeightMap(const FMultiBoneWeightEdits& Edits, TArray<VertexWeights>& InOutWeights);
+		
+		// double-buffer of the entire weight matrix (stored sparsely for fast deformation)
+		// "Pre" is state of weights at stroke start
+		// "Current" is state of weights during stroke
+		// When stroke is over, Pre weights are synchronized with current.
+		TArray<VertexWeights> PreStrokeWeights;
+		TArray<VertexWeights> CurrentWeights;
+
+		// record the current maximum amount of falloff applied to each vertex during the current stroke
+		// values range from 0-1, this allows brushes to sweep over the same vertex, and apply only the maximum amount
+		// of modification (add/replace/relax etc) that was encountered for the duration of the stroke.
+		TArray<float> MaxFalloffPerVertexThisStroke;
+
+		// update deformation when vertex weights are modified
+		FSkinToolDeformer Deformer;
+	};
 
-	float SetCurrentFalloffAndGetMaxFalloffThisStroke(int32 VertexID, float CurrentStength);
-	
-	// double-buffer of the entire weight matrix
-	// "Pre" is state of weights at stroke start
-	// "Current" is state of weights during stroke
-	// When stroke is over, Pre weights are synchronized with current.
-	PerBoneWeightMap PreStrokeWeightsMap;
-	PerBoneWeightMap CurrentWeightsMap;
-
-	// record the current maximum amount of falloff applied to each vertex during the current stroke
-	// values range from 0-1, this allows brushes to sweep over the same vertex, and apply only the maximum amount
-	// of modification (add/replace/relax etc) that was encountered for the duration of the stroke.
-	TArray<float> MaxFalloffPerVertexThisStroke;
-};
+}
 
 /**
  *
@@ -213,9 +193,7 @@ enum class EBrushBehaviorMode : uint8
 
 
 UCLASS()
-class MESHMODELINGTOOLSEXP_API USkinWeightsPaintToolProperties : 
-	public UInteractiveToolPropertySet,
-	public IBoneReferenceSkeletonProvider
+class MESHMODELINGTOOLSEXP_API USkinWeightsPaintToolProperties : public UInteractiveToolPropertySet
 {
 	GENERATED_BODY()
 
@@ -228,9 +206,6 @@ public:
 	
 	UPROPERTY(EditAnywhere, Category = Brush)
 	EWeightBrushFalloffMode FalloffMode;
-	
-	UPROPERTY(VisibleAnywhere, Category = Skeleton)
-	FBoneReference CurrentBone;
 
 	UPROPERTY(EditAnywhere, Category = Brush)
 	EWeightColorMode ColorMode;
@@ -244,9 +219,6 @@ public:
 	UPROPERTY(EditAnywhere, Category = WeightColors)
 	FLinearColor MaxColor;
 	
-	// IBoneReferenceSkeletonProvider
-	USkeleton* GetSkeleton(bool& bInvalidSkeletonIsError, const IPropertyHandle* PropertyHandle) override;
-
 	TObjectPtr<USkeletalMesh> SkeletalMesh;
 	bool bColorModeChanged = false;
 };
@@ -263,7 +235,6 @@ public:
 	void RegisterActions(FInteractiveToolActionSet& ActionSet) override;
 
 	void Setup() override;
-	void Render(IToolsContextRenderAPI* RenderAPI) override;
 
 	bool HasCancel() const override { return true; }
 	bool HasAccept() const override { return true; }
@@ -276,7 +247,7 @@ public:
 	void OnEndDrag(const FRay& Ray) override;
 	bool OnUpdateHover(const FInputDeviceRay& DevicePos) override;
 
-	void ExternalUpdateWeights(const FName &BoneName, const TMap<int32, float>& IndexValues);
+	void ExternalUpdateWeights(const int32 BoneIndex, const TMap<int32, float>& IndexValues);
 
 protected:
 
@@ -294,17 +265,24 @@ protected:
 		EBrushBehaviorMode EditMode,
 		const TArray<int32>& VerticesInStamp,
 		const TArray<float>& VertexSqDistances,
-		FMultiBoneWeightEdits& AllBoneWeightEditsFromStamp);
+		SkinPaintTool::FMultiBoneWeightEdits& AllBoneWeightEditsFromStamp);
 	void RelaxWeightOnVertices(
 		TArray<int32> VerticesInStamp,
 		TArray<float> VertexSqDistances,
-		FMultiBoneWeightEdits& AllBoneWeightEditsFromStamp);
+		SkinPaintTool::FMultiBoneWeightEdits& AllBoneWeightEditsFromStamp);
 	bool bInvertStroke = false;
 	bool bSmoothStroke = false;
 	FBrushStampData StartStamp;
 	FBrushStampData LastStamp;
 	bool bStampPending;
+	int32 TriangleUnderStamp;
+	FVector StampLocalPos;
+
+	// used to accelerate mesh queries
 	UE::Geometry::TDynamicVerticesOctree3<FDynamicMesh3> VerticesOctree;
+	UE::Geometry::FDynamicMeshOctree3 TrianglesOctree;
+	TFuture<void> TriangleOctreeFuture;
+	TArray<int32> TriangleToReinsert;
 
 	// tool properties
 	UPROPERTY()
@@ -315,41 +293,24 @@ protected:
 	TUniquePtr<FMeshDescription> EditedMesh;
 
 	// storage of vertex weights per bone 
-	FSkinToolWeights Weights;
-	
-	// update deformation when vertex weights are modified
-	FSkinToolDeformer Deformer;
+	SkinPaintTool::FSkinToolWeights Weights;
 
 	// vertex colors updated when switching current bone or editing weights
 	void UpdateCurrentBoneVertexColors();
-	FVector4f WeightToColor(float Value);
+	FVector4f WeightToColor(float Value) const;
 	bool bVisibleWeightsValid = false;
 
 	// weight editing transactions
 	void BeginChange();
-	TUniquePtr<FMeshSkinWeightsChange> EndChange();
-	TUniquePtr<FMeshSkinWeightsChange> ActiveChange;
+	TUniquePtr<SkinPaintTool::FMeshSkinWeightsChange> EndChange();
+	TUniquePtr<SkinPaintTool::FMeshSkinWeightsChange> ActiveChange;
 
 	// which bone are we currently painting?
 	void UpdateCurrentBone(const FName &BoneName);
 	FName CurrentBone = NAME_None;
 	TOptional<FName> PendingCurrentBone;
 
-	// bone rendering
-	// TODO, bring this up to new method, add support for rotating bones here
-	struct FBonePositionInfo
-	{
-		FName BoneName;
-		int32 ParentBoneIndex;
-		FVector Position;
-		float Radius;
-		TMap<FName, int32> ChildBones;
-	};
-	FBoneContainer BoneContainer;
-	TArray<FBonePositionInfo> BonePositionInfos;
-	float MaxDrawRadius = 0.0f;
-	void UpdateBonePositionInfos(float MinRadius);
-	void RenderBonePositions(FPrimitiveDrawInterface *PDI);
-
 	virtual void HandleSkeletalMeshModified(const TArray<FName>& InBoneNames, const ESkeletalMeshNotifyType InNotifyType) override;
+
+	friend SkinPaintTool::FSkinToolDeformer;
 };
