Changed default color ramp colors in Skin Paint Tool

diff --git a/Engine/Plugins/Experimental/MeshModelingToolsetExp/Source/MeshModelingToolsExp/Private/SkinWeightsPaintTool.cpp b/Engine/Plugins/Experimental/MeshModelingToolsetExp/Source/MeshModelingToolsExp/Private/SkinWeightsPaintTool.cpp
--- a/Engine/Plugins/Experimental/MeshModelingToolsetExp/Source/MeshModelingToolsExp/Private/SkinWeightsPaintTool.cpp
+++ b/Engine/Plugins/Experimental/MeshModelingToolsetExp/Source/MeshModelingToolsExp/Private/SkinWeightsPaintTool.cpp
@@ -33,9 +33,7 @@ USkinWeightsPaintToolProperties::USkinWeightsPaintToolProperties(const FObjectIn
 		: Super(ObjectInitializer)
 {
 	ColorRamp.Add(FLinearColor::Blue);
-	ColorRamp.Add(FLinearColor::Green);
-	ColorRamp.Add(FLinearColor(FColor::Orange));
-	ColorRamp.Add(FLinearColor::Red);
+	ColorRamp.Add(FLinearColor::Yellow);
 
 	MinColor = FLinearColor::Black;
 	MaxColor = FLinearColor::White;
