Fixing unity build errors in SkinWeightsPaintTool

diff --git a/Engine/Plugins/Experimental/MeshModelingToolsetExp/Source/MeshModelingToolsExp/Private/SkinWeightsPaintTool.cpp b/Engine/Plugins/Experimental/MeshModelingToolsetExp/Source/MeshModelingToolsExp/Private/SkinWeightsPaintTool.cpp
--- a/Engine/Plugins/Experimental/MeshModelingToolsetExp/Source/MeshModelingToolsExp/Private/SkinWeightsPaintTool.cpp
+++ b/Engine/Plugins/Experimental/MeshModelingToolsetExp/Source/MeshModelingToolsExp/Private/SkinWeightsPaintTool.cpp
@@ -10,6 +10,7 @@
 
 #include "TargetInterfaces/PrimitiveComponentBackedTarget.h"
 #include "ModelingToolTargetUtil.h"
+#include "ToolBuilderUtil.h"
 
 #include "MeshDescription.h"
 #include "DynamicMesh/NonManifoldMappingSupport.h"
@@ -502,6 +503,7 @@ void USkinWeightsPaintTool::Setup()
 
 	// initialize the tool properties
 	BrushProperties->RestoreProperties(this); // hides strength and falloff
+	
 	ToolProps = NewObject<USkinWeightsPaintToolProperties>(this);
 	ToolProps->RestoreProperties(this);
 	ToolProps->SkeletalMesh = Component->GetSkeletalMeshAsset();
diff --git a/Engine/Plugins/Experimental/MeshModelingToolsetExp/Source/MeshModelingToolsExp/Public/SkinWeightsPaintTool.h b/Engine/Plugins/Experimental/MeshModelingToolsetExp/Source/MeshModelingToolsExp/Public/SkinWeightsPaintTool.h
--- a/Engine/Plugins/Experimental/MeshModelingToolsetExp/Source/MeshModelingToolsExp/Public/SkinWeightsPaintTool.h
+++ b/Engine/Plugins/Experimental/MeshModelingToolsetExp/Source/MeshModelingToolsExp/Public/SkinWeightsPaintTool.h
@@ -11,6 +11,7 @@
 #include "Containers/Map.h"
 #include "DynamicMesh/DynamicMeshOctree3.h"
 #include "SkeletalMesh/SkeletalMeshEditionInterface.h"
+#include "Engine/SkeletalMesh.h"
 
 #include "SkinWeightsPaintTool.generated.h"
 
