CombineMeshInstances: Add support for merging faces after removing hidden faces. Provide a mechanism for external calling code to prevent this merging based on external "Grouping ID" (eg color index, etc). Add CVar geometry.CombineInstances.Verbose, when enabled, multi-threaded LOD generation is disabled and various log info is printed during the process.

diff --git a/Engine/Plugins/Experimental/MeshModelingToolsetExp/Source/GeometryProcessingAdapters/Private/CombineMeshInstancesImpl.cpp b/Engine/Plugins/Experimental/MeshModelingToolsetExp/Source/GeometryProcessingAdapters/Private/CombineMeshInstancesImpl.cpp
--- a/Engine/Plugins/Experimental/MeshModelingToolsetExp/Source/GeometryProcessingAdapters/Private/CombineMeshInstancesImpl.cpp
+++ b/Engine/Plugins/Experimental/MeshModelingToolsetExp/Source/GeometryProcessingAdapters/Private/CombineMeshInstancesImpl.cpp
@@ -29,6 +29,7 @@
 #include "MeshConstraintsUtil.h"
 #include "DynamicMesh/Operations/MergeCoincidentMeshEdges.h"
 #include "Operations/RemoveOccludedTriangles.h"
+#include "Operations/MeshResolveTJunctions.h"
 
 #include "Physics/CollisionGeometryConversion.h"
 #include "Physics/PhysicsDataCollection.h"
@@ -51,6 +52,12 @@ static TAutoConsoleVariable<int> CVarGeometryCombineMeshInstancesRemoveHidden(
 	TEXT("Configure hidden-removal strategy via (temporary debug)"));
 
 
+static TAutoConsoleVariable<bool> CVarGeometryCombineMeshInstancesVerbose(
+	TEXT("geometry.CombineInstances.Verbose"),
+	false,
+	TEXT("Enable Verbose logging in Combine Mesh Instances, also disables parallel LOD processing"));
+
+
 enum class EMeshDetailLevel
 {
 	Base = 0,
@@ -279,28 +286,6 @@ void InitializeAssemblySourceMeshesFromLOD(
 		}
 	}
 
-
-
-	//ParallelFor(NumSets, [&](int32 Index)
-	//{
-	//	TUniquePtr<FMeshInstanceSet>& InstanceSet = Assembly.InstanceSets[Index];
-	//	UStaticMesh* StaticMesh = InstanceSet->SourceAsset;
-
-	//	FSourceGeometry& Target = Assembly.SourceMeshGeometry[Index];
-
-	//	UBodySetup* BodySetup = StaticMesh->GetBodySetup();
-	//	if (BodySetup)
-	//	{
-	//		UE::Geometry::GetShapeSet(BodySetup->AggGeom, Target.CollisionShapes);
-	//	}
-
-	//	// todo:
-	//	//if (bDetectBoxes)
-	//	//{
-	//	//	DetectBoxesFromConvexes(CollisionShapes);
-	//	//}
-
-	//});
 }
 
 
@@ -710,6 +695,7 @@ void ComputeMeshApproximations(
 	int32 NumApproxLODs = FMath::Max(1, 
 		CombineOptions.NumLODs - CombineOptions.NumCopiedLODs - CombineOptions.NumSimplifiedLODs);
 
+	bool bVerbose = CVarGeometryCombineMeshInstancesVerbose.GetValueOnGameThread();
 	ParallelFor(NumSets, [&](int32 Index)
 	{
 		TUniquePtr<FMeshInstanceSet>& InstanceSet = Assembly.InstanceSets[Index];
@@ -747,7 +733,7 @@ void ComputeMeshApproximations(
 			FMeshNormals::QuickRecomputeOverlayNormals(ApproxGeo.ApproximateMeshLODs[k]);
 		}
 
-	});
+	}, (bVerbose) ? EParallelForFlags::ForceSingleThread : EParallelForFlags::None );
 
 
 	// try to filter out simplifications that did bad things
@@ -845,7 +831,7 @@ static void RemoveHiddenFaces_ExteriorVisibility(FDynamicMesh3& TargetMesh, doub
 	{
 		FVector3d Normal, Centroid; double Area;
 		TargetMesh.GetTriInfo(tid, Normal, Area, Centroid);
-		if (Normal.SquaredLength() < 0.1 || Area < FMathf::ZeroTolerance)
+		if (Normal.SquaredLength() < 0.1 || Area <= FMathd::ZeroTolerance)
 		{
 			TriStatusKnown[tid] = true;
 			return;
@@ -853,24 +839,30 @@ static void RemoveHiddenFaces_ExteriorVisibility(FDynamicMesh3& TargetMesh, doub
 
 		for (FVector3d Direction : CardinalDirections)
 		{
-			if (FindHitTriangleTest(Centroid, Centroid + Radius*Direction) == tid)
+			// if direction is orthogonal to the triangle, hit-test is unstable, but even
+			// worse, on rectilinear shapes (eg imagine some stacked cubes or adjacent parts)
+			// the ray can get "through" the cracks between adjacent connected triangles
+			// and manage to hit the search triangle
+			if ( FMathd::Abs(Direction.Dot(Normal)) > 0.01 )
 			{
-				ThreadSafeTriVisible[tid] = true;
-				TriStatusKnown[tid] = true;
-				return;
+				if (FindHitTriangleTest(Centroid, Centroid + Radius*Direction) == tid)
+				{
+					ThreadSafeTriVisible[tid] = true;
+					TriStatusKnown[tid] = true;
+					return;
+				}
 			}
 		}
 		if (FindHitTriangleTest(Centroid, Centroid + Radius*Normal) == tid)
 		{
-				ThreadSafeTriVisible[tid] = true;
-				TriStatusKnown[tid] = true;
-				return;
+			ThreadSafeTriVisible[tid] = true;
+			TriStatusKnown[tid] = true;
+			return;
 		}
 
 		// triangle is not definitely visible or hidden
 	});
 
-
 	//
 	// Construct set of exterior sample points, for each triangle sample point
 	// below we will check if it is visible from any of these sample points.
@@ -978,19 +970,132 @@ static void RemoveHiddenFaces_ExteriorVisibility(FDynamicMesh3& TargetMesh, doub
 }
 
 
+// internal struct used in PostProcessHiddenFaceRemovedMesh
+struct FMergeTriInfo
+{
+	int32 MaterialID = 0;
+	FIndex3i ExternalGroupingID = FIndex3i::Zero();
+
+	bool operator==(const FMergeTriInfo& Other) const { return MaterialID == Other.MaterialID && ExternalGroupingID == Other.ExternalGroupingID; }
+	uint32 GetTypeHash() const
+	{
+		return HashCombineFast( ::GetTypeHash(MaterialID), FCrc::MemCrc_DEPRECATED(&ExternalGroupingID, sizeof(ExternalGroupingID)) );
+	}
+	friend uint32 GetTypeHash(const FMergeTriInfo& TriInfo)
+	{
+		return TriInfo.GetTypeHash();
+	}
+};
 
 
-static void PostProcessHiddenFaceRemovedMesh(FDynamicMesh3& TargetMesh, double Tolerance)
+//
+// After hidden face removal, a mesh can often be optimized to at least save some vertices (by welding open borders),
+// and then in some cases, now-connected triangle areas can be retriangulated to require fewer triangles.
+// The latter is really only possible if UVs/Normal seams are not involved, and generally such merging
+// of areas needs to be prevented between different Material regions.
+// To support Materials that define different material regions internally (eg indexed colors encoded in vertex colors,
+// custom primitive data, etc) a function is provided to allow external code to provide 3 "unique triangle group" integers.
+// All integers must match for a triangle regions to be merged for retriangulation.
+//
+static void PostProcessHiddenFaceRemovedMesh(
+	FDynamicMesh3& TargetMesh, 
+	double Tolerance,
+	bool bTryToMergeFaces,
+	TFunctionRef<FIndex3i(const FDynamicMesh3& Mesh, int32 TriangleID)> GetTriangleGroupingIDFunc)
 {
+	bool bVerbose = CVarGeometryCombineMeshInstancesVerbose.GetValueOnAnyThread();
+
 	// weld edges in case input was unwelded...
-	FMergeCoincidentMeshEdges Welder(&TargetMesh);
-	Welder.MergeVertexTolerance = Tolerance * 0.001;
-	Welder.OnlyUniquePairs = false;
-	Welder.Apply();
+	{
+		FMergeCoincidentMeshEdges Welder(&TargetMesh);
+		Welder.MergeVertexTolerance = Tolerance * 0.01;
+		Welder.OnlyUniquePairs = false;
+		Welder.Apply();
+	}
+
+	if (!bTryToMergeFaces)
+	{
+		TargetMesh.CompactInPlace();
+		return;
+	}
+
+	const FDynamicMeshMaterialAttribute* MaterialIDs =
+		(TargetMesh.HasAttributes() && TargetMesh.Attributes()->HasMaterialID()) ? TargetMesh.Attributes()->GetMaterialID() : nullptr;
+	
+	TMap<FMergeTriInfo, int> UniqueMatIndices;
+	TArray<int32> TriSortIndex;
+	TriSortIndex.SetNum(TargetMesh.MaxTriangleID());
+	for (int32 tid : TargetMesh.TriangleIndicesItr())
+	{
+		FMergeTriInfo TriInfo;
+		TriInfo.MaterialID = (MaterialIDs) ? MaterialIDs->GetValue(tid) : -1;
+		TriInfo.ExternalGroupingID = GetTriangleGroupingIDFunc(TargetMesh, tid);
+		int32* Found = UniqueMatIndices.Find(TriInfo);
+		if (Found == nullptr)
+		{
+			int32 NewIndex = UniqueMatIndices.Num();
+			UniqueMatIndices.Add(TriInfo, NewIndex);
+			TriSortIndex[tid] = NewIndex;
+		}
+		else
+		{
+			TriSortIndex[tid] = *Found;
+		}
+	}
+
+
+	TArray<FDynamicMesh3> SplitMeshes;
+	if ( UniqueMatIndices.Num() == 1 )
+	{
+		SplitMeshes.Add(MoveTemp(TargetMesh));
+	}
+	else
+	{
+		FDynamicMeshEditor::SplitMesh(&TargetMesh, SplitMeshes, [&](int32 tid) { return TriSortIndex[tid]; });
+	}
 
-	// todo: try to simplify? need to be able to constrain by things like vertex color...
+	for (FDynamicMesh3& SubRegionMesh : SplitMeshes)
+	{
+		// resolving T-junctions tends to make things worse...
+		//FMeshResolveTJunctions Resolver(&SubRegionMesh);
+		//Resolver.DistanceTolerance = 0.01;
+		//Resolver.Apply();
+
+		// try weld again just in case
+		FMergeCoincidentMeshEdges Welder(&SubRegionMesh);
+		Welder.MergeVertexTolerance = Tolerance * 0.01;
+		Welder.OnlyUniquePairs = false;
+		Welder.Apply();
+
+		// simplify to planar
+		FQEMSimplification Simplifier(&SubRegionMesh);
+		Simplifier.CollapseMode = FQEMSimplification::ESimplificationCollapseModes::AverageVertexPosition;
+		Simplifier.SimplifyToMinimalPlanar( 0.01 );
+	}
+
+	TargetMesh.Clear();
+	TargetMesh.EnableMatchingAttributes(SplitMeshes[0], true, true);
+	FDynamicMeshEditor Editor(&TargetMesh);
+	for (FDynamicMesh3& SubRegionMesh : SplitMeshes)
+	{
+		FMeshIndexMappings Mappings;
+		Editor.AppendMesh(&SubRegionMesh, Mappings);
+	}
+
+	// weld edges back together again
+	{
+		FMergeCoincidentMeshEdges Welder(&TargetMesh);
+		Welder.MergeVertexTolerance = Tolerance * 0.01;
+		Welder.OnlyUniquePairs = false;
+		Welder.Apply();
+	}
 
 	TargetMesh.CompactInPlace();
+
+	if ( bVerbose )
+	{
+		UE_LOG(LogGeometry, Log, TEXT("    Merge Faces           [Tris %6d Verts %6d]"), TargetMesh.TriangleCount(), TargetMesh.VertexCount());
+	}
 }
 
 
@@ -1234,8 +1339,11 @@ static void SortMesh(FDynamicMesh3& Mesh)
 
 void ComputeHiddenRemovalForLOD(
 	FDynamicMesh3& MeshLOD,
+	int32 LODIndex,
 	IGeometryProcessing_CombineMeshInstances::FOptions CombineOptions)
 {
+	bool bVerbose = CVarGeometryCombineMeshInstancesVerbose.GetValueOnAnyThread();
+
 	TRACE_CPUPROFILER_EVENT_SCOPE(RemoveHidden_LOD);
 	bool bModified = false;
 	switch (CombineOptions.RemoveHiddenFacesMethod)
@@ -1251,9 +1359,23 @@ void ComputeHiddenRemovalForLOD(
 			break;
 	}
 
+	if ( bVerbose )
+	{
+		UE_LOG(LogGeometry, Log, TEXT("    Remove Hidden Faces - [Tris %6d Verts %6d]"), MeshLOD.TriangleCount(), MeshLOD.VertexCount());
+	}
+
 	if ( bModified )
 	{
-		PostProcessHiddenFaceRemovedMesh(MeshLOD, CombineOptions.SimplifyBaseTolerance);
+		TFunction<FIndex3i(const FDynamicMesh3& Mesh, int32 TriangleID)> GroupingIDFunc = CombineOptions.TriangleGroupingIDFunc;
+		if (!GroupingIDFunc)
+		{
+			GroupingIDFunc = [](const FDynamicMesh3&, int32) { return FIndex3i::Zero(); };
+		}
+
+		PostProcessHiddenFaceRemovedMesh(MeshLOD, 
+			CombineOptions.SimplifyBaseTolerance,
+			CombineOptions.bMergeCoplanarFaces && LODIndex >= CombineOptions.MergeCoplanarFacesStartLOD,
+			GroupingIDFunc );
 	}
 }
 
@@ -1267,6 +1389,7 @@ void BuildCombinedMesh(
 	TArray<FDynamicMesh3>& CombinedMeshLODs)
 {
 	using namespace UE::Geometry;
+	bool bVerbose = CVarGeometryCombineMeshInstancesVerbose.GetValueOnGameThread();
 
 	int32 NumLODs = CombineOptions.NumLODs;
 	TArray<FCombinedMeshLOD> MeshLODs;
@@ -1405,13 +1528,23 @@ void BuildCombinedMesh(
 		&& CVarGeometryCombineMeshInstancesRemoveHidden.GetValueOnGameThread() > 0);
 	if (bRemoveHiddenFaces)
 	{
-		for (int32 k = CombineOptions.RemoveHiddenStartLOD; k < MeshLODs.Num() && k < FirstVoxWrappedIndex; ++k)
+		for (int32 LODIndex = CombineOptions.RemoveHiddenStartLOD; LODIndex < MeshLODs.Num() && LODIndex < FirstVoxWrappedIndex; ++LODIndex)
 		{
-			UE::Tasks::FTask RemoveHiddenTask = UE::Tasks::Launch(UE_SOURCE_LOCATION, [&MeshLODs, &CombineOptions, k]()
+			if ( bVerbose )
+			{
+				UE_LOG(LogGeometry, Log, TEXT("  Optimizing LOD%d - Tris %6d Verts %6d"), LODIndex, MeshLODs[LODIndex].Mesh.TriangleCount(), MeshLODs[LODIndex].Mesh.VertexCount());
+			}
+
+			UE::Tasks::FTask RemoveHiddenTask = UE::Tasks::Launch(UE_SOURCE_LOCATION, [&MeshLODs, &CombineOptions, LODIndex]()
 			{ 
-				ComputeHiddenRemovalForLOD(MeshLODs[k].Mesh, CombineOptions);
+				ComputeHiddenRemovalForLOD(MeshLODs[LODIndex].Mesh, LODIndex, CombineOptions);
 			});
 			PendingRemoveHiddenTasks.Add(RemoveHiddenTask);
+
+			if (bVerbose)
+			{
+				RemoveHiddenTask.BusyWait();
+			}
 		}
 	}
 
@@ -1424,17 +1557,24 @@ void BuildCombinedMesh(
 		FDynamicMesh3 SourceVoxWrapMesh = MoveTemp(MeshLODs[FirstVoxWrappedIndex].Mesh);
 		FDynamicMeshAABBTree3 Spatial(&SourceVoxWrapMesh, true);
 
+		const double CLOSURE_DIST = 10.0;		// TODO: this needs to be exposed as an option, perhaps per-part
+
 		FDynamicMesh3 TempBaseVoxWrapMesh;
 		double VoxelDimension = 2.0;	// may be modified by ComputeVoxWrapMesh call
 		{
 			TRACE_CPUPROFILER_EVENT_SCOPE(ComputeVoxWrap);
-			ComputeVoxWrapMesh(SourceVoxWrapMesh, Spatial, TempBaseVoxWrapMesh, 10.0, VoxelDimension);
+			ComputeVoxWrapMesh(SourceVoxWrapMesh, Spatial, TempBaseVoxWrapMesh, CLOSURE_DIST, VoxelDimension);
 			// currently need to re-sort output to remove non-determinism...
 			SortMesh(TempBaseVoxWrapMesh);
 
 			//UE_LOG(LogGeometry, Warning, TEXT("VoxWrapMesh has %d triangles %d vertices"), TempBaseVoxWrapMesh.TriangleCount(), TempBaseVoxWrapMesh.VertexCount());
 		}
 
+		if ( bVerbose )
+		{
+			UE_LOG(LogGeometry, Log, TEXT("  Generated Base VoxWrap Mesh - Tris %8d Verts %8d - CellSize is %4.3f"), TempBaseVoxWrapMesh.TriangleCount(), TempBaseVoxWrapMesh.VertexCount(), VoxelDimension);
+		}
+
 		{
 			TRACE_CPUPROFILER_EVENT_SCOPE(FastCollapsePrePass);
 			TempBaseVoxWrapMesh.DiscardAttributes();
@@ -1443,6 +1583,11 @@ void BuildCombinedMesh(
 			Simplifier.FastCollapsePass(VoxelDimension*0.5, 10, false, 50000);
 		}
 
+		if ( bVerbose )
+		{
+			UE_LOG(LogGeometry, Log, TEXT("         FastCollapse         - Tris %8d Verts %8d"), TempBaseVoxWrapMesh.TriangleCount(), TempBaseVoxWrapMesh.VertexCount());
+		}
+
 		const FDynamicMesh3& LastApproxLOD = MeshLODs[FirstVoxWrappedIndex-1].Mesh;
 
 		int32 MaxTriCount = CombineOptions.VoxWrapMaxTriCountBase;
@@ -1460,6 +1605,12 @@ void BuildCombinedMesh(
 			ComputeSimplifiedVoxWrapMesh(MeshLODs[LODIndex].Mesh, &SourceVoxWrapMesh, &Spatial, 
 				SimplifyTolerance, MaxTriCount);
 
+			if ( bVerbose )
+			{
+				UE_LOG(LogGeometry, Log, TEXT("         LOD %2d               - Tris %8d Verts %8d - MaxTriCount %4d"), 
+					LODIndex, MeshLODs[LODIndex].Mesh.TriangleCount(), MeshLODs[LODIndex].Mesh.VertexCount(), MaxTriCount);
+			}
+
 			InitializeAttributes(MeshLODs[LODIndex].Mesh, CombineOptions.HardNormalAngleDeg,
 				/*bProjectAttributes*/true, &SourceVoxWrapMesh, &Spatial);
 
@@ -1472,17 +1623,22 @@ void BuildCombinedMesh(
 	// wait...
 	UE::Tasks::Wait(PendingRemoveHiddenTasks);
 
-	// remove hiddel faces on voxel LODs (todo: can do this via shape sorting, much faster)
+	// remove hidden faces on voxel LODs (todo: can do this via shape sorting, much faster)
 	if (bRemoveHiddenFaces)
 	{
 		TRACE_CPUPROFILER_EVENT_SCOPE(RemoveHidden);
-		ParallelFor(MeshLODs.Num(), [&](int32 k)
+		ParallelFor(MeshLODs.Num(), [&](int32 LODIndex)
 		{
-			if ( k >= FirstVoxWrappedIndex )
+			if ( LODIndex >= FirstVoxWrappedIndex )
 			{ 
-				ComputeHiddenRemovalForLOD(MeshLODs[k].Mesh, CombineOptions);
+				if ( bVerbose )
+				{
+					UE_LOG(LogGeometry, Log, TEXT("  Optimizing LOD%d - Tris %6d Verts %6d"), LODIndex, MeshLODs[LODIndex].Mesh.TriangleCount(), MeshLODs[LODIndex].Mesh.VertexCount());
+				}
+
+				ComputeHiddenRemovalForLOD(MeshLODs[LODIndex].Mesh, LODIndex, CombineOptions);
 			}
-		});
+		}, (bVerbose) ? EParallelForFlags::ForceSingleThread :  EParallelForFlags::None );
 	}
 
 
@@ -1730,11 +1886,27 @@ void FCombineMeshInstancesImpl::CombineMeshInstances(
 {
 	TRACE_CPUPROFILER_EVENT_SCOPE(CombineMeshInstances);
 
+	bool bVerbose = CVarGeometryCombineMeshInstancesVerbose.GetValueOnGameThread();
+	if (bVerbose)
+	{
+		UE_LOG(LogGeometry, Log, TEXT("CombineMeshInstances: processing %d Instances into %d LODs (%d Copied, %d Simplified, %d Approx, %d VoxWrapped)"),
+			MeshInstances.StaticMeshInstances.Num(),
+			Options.NumLODs, Options.NumCopiedLODs, Options.NumSimplifiedLODs, 
+			FMath::Max(0, Options.NumLODs - Options.NumCopiedLODs - Options.NumSimplifiedLODs - Options.NumVoxWrapLODs),
+			Options.NumVoxWrapLODs);
+	}
+
 	FMeshInstanceAssembly InstanceAssembly;
 
 	{
 		TRACE_CPUPROFILER_EVENT_SCOPE(CombineMeshInst_Setup);
 		InitializeMeshInstanceAssembly(MeshInstances, InstanceAssembly);
+		if (bVerbose)
+		{
+			UE_LOG(LogGeometry, Log, TEXT("  InstanceAssembly contains %d InstanceSets, %d Unique Materials"), 
+				InstanceAssembly.InstanceSets.Num(), InstanceAssembly.UniqueMaterials.Num());
+		}
+
 		InitializeAssemblySourceMeshesFromLOD(InstanceAssembly, 0, Options.NumCopiedLODs);
 		InitializeInstanceAssemblySpatials(InstanceAssembly);
 	}
@@ -1768,6 +1940,12 @@ void FCombineMeshInstancesImpl::CombineMeshInstances(
 	{
 		TRACE_CPUPROFILER_EVENT_SCOPE(CombineMeshInst_BuildCollision);
 		BuildCombinedCollisionShapes(InstanceAssembly, Options, CombinedCollisionShapes);
+
+		if (bVerbose)
+		{
+			UE_LOG(LogGeometry, Log, TEXT("  CombinedCollisionShapes contains %d Boxes, %d Convexes"), 
+				CombinedCollisionShapes.Boxes.Num(), CombinedCollisionShapes.Convexes.Num());
+		}
 	}
 	FPhysicsDataCollection PhysicsData;
 	PhysicsData.Geometry = CombinedCollisionShapes;
diff --git a/Engine/Source/Developer/GeometryProcessingInterfaces/Private/GeometryProcessingInterfacesModule.cpp b/Engine/Source/Developer/GeometryProcessingInterfaces/Private/GeometryProcessingInterfacesModule.cpp
--- a/Engine/Source/Developer/GeometryProcessingInterfaces/Private/GeometryProcessingInterfacesModule.cpp
+++ b/Engine/Source/Developer/GeometryProcessingInterfaces/Private/GeometryProcessingInterfacesModule.cpp
@@ -39,7 +39,7 @@ IGeometryProcessing_ApproximateActors* FGeometryProcessingInterfacesModule::GetA
 }
 IGeometryProcessing_CombineMeshInstances* FGeometryProcessingInterfacesModule::GetCombineMeshInstancesImplementation()
 {
-	if (ApproximateActors == nullptr)
+	if (CombineMeshInstances == nullptr)
 	{
 		TArray<IGeometryProcessing_CombineMeshInstances*> CombineMeshInstancesOptions =
 			IModularFeatures::Get().GetModularFeatureImplementations<IGeometryProcessing_CombineMeshInstances>(IGeometryProcessing_CombineMeshInstances::GetModularFeatureName());
diff --git a/Engine/Source/Developer/GeometryProcessingInterfaces/Public/GeometryProcessingInterfaces/CombineMeshInstances.h b/Engine/Source/Developer/GeometryProcessingInterfaces/Public/GeometryProcessingInterfaces/CombineMeshInstances.h
--- a/Engine/Source/Developer/GeometryProcessingInterfaces/Public/GeometryProcessingInterfaces/CombineMeshInstances.h
+++ b/Engine/Source/Developer/GeometryProcessingInterfaces/Public/GeometryProcessingInterfaces/CombineMeshInstances.h
@@ -115,6 +115,11 @@ public:
 
 		// opening angle used to detect/assign sharp edges
 		double HardNormalAngleDeg = 15.0;
+
+
+		bool bMergeCoplanarFaces = true;
+		int32 MergeCoplanarFacesStartLOD = 1;
+		TFunction<UE::Geometry::FIndex3i(const UE::Geometry::FDynamicMesh3& Mesh, int32 TriangleID)> TriangleGroupingIDFunc;
 	};
 
 
