CombineMeshInstances: reorganize logging to resolve PVS error
#rnx

diff --git a/Engine/Plugins/Experimental/MeshModelingToolsetExp/Source/GeometryProcessingAdapters/Private/CombineMeshInstancesImpl.cpp b/Engine/Plugins/Experimental/MeshModelingToolsetExp/Source/GeometryProcessingAdapters/Private/CombineMeshInstancesImpl.cpp
--- a/Engine/Plugins/Experimental/MeshModelingToolsetExp/Source/GeometryProcessingAdapters/Private/CombineMeshInstancesImpl.cpp
+++ b/Engine/Plugins/Experimental/MeshModelingToolsetExp/Source/GeometryProcessingAdapters/Private/CombineMeshInstancesImpl.cpp
@@ -1889,11 +1889,10 @@ void FCombineMeshInstancesImpl::CombineMeshInstances(
 	bool bVerbose = CVarGeometryCombineMeshInstancesVerbose.GetValueOnGameThread();
 	if (bVerbose)
 	{
+		int32 NumApproxLODs = FMath::Max(0, Options.NumLODs - Options.NumCopiedLODs - Options.NumSimplifiedLODs - Options.NumVoxWrapLODs);
 		UE_LOG(LogGeometry, Log, TEXT("CombineMeshInstances: processing %d Instances into %d LODs (%d Copied, %d Simplified, %d Approx, %d VoxWrapped)"),
 			MeshInstances.StaticMeshInstances.Num(),
-			Options.NumLODs, Options.NumCopiedLODs, Options.NumSimplifiedLODs, 
-			FMath::Max(0, Options.NumLODs - Options.NumCopiedLODs - Options.NumSimplifiedLODs - Options.NumVoxWrapLODs),
-			Options.NumVoxWrapLODs);
+			Options.NumLODs, Options.NumCopiedLODs, Options.NumSimplifiedLODs, NumApproxLODs, Options.NumVoxWrapLODs);
 	}
 
 	FMeshInstanceAssembly InstanceAssembly;
