CombineMeshInstances: split bowties before welding edges to avoid some current issues that can be introduced by welding

diff --git a/Engine/Plugins/Experimental/MeshModelingToolsetExp/Source/GeometryProcessingAdapters/Private/CombineMeshInstancesImpl.cpp b/Engine/Plugins/Experimental/MeshModelingToolsetExp/Source/GeometryProcessingAdapters/Private/CombineMeshInstancesImpl.cpp
--- a/Engine/Plugins/Experimental/MeshModelingToolsetExp/Source/GeometryProcessingAdapters/Private/CombineMeshInstancesImpl.cpp
+++ b/Engine/Plugins/Experimental/MeshModelingToolsetExp/Source/GeometryProcessingAdapters/Private/CombineMeshInstancesImpl.cpp
@@ -1007,6 +1007,12 @@ static void PostProcessHiddenFaceRemovedMesh(
 
 	// weld edges in case input was unwelded...
 	{
+		// currently FMergeCoincidentMeshEdges can break the mesh if it has bowties, remove
+		// them to work around the issue
+		FDynamicMeshEditor MeshEditor(&TargetMesh);
+		FDynamicMeshEditResult EditResult;
+		MeshEditor.SplitBowties(EditResult);
+
 		FMergeCoincidentMeshEdges Welder(&TargetMesh);
 		Welder.MergeVertexTolerance = Tolerance * 0.01;
 		Welder.OnlyUniquePairs = false;
@@ -1084,6 +1090,12 @@ static void PostProcessHiddenFaceRemovedMesh(
 
 	// weld edges back together again
 	{
+		// currently FMergeCoincidentMeshEdges can break the mesh if it has bowties, remove
+		// them to work around the issue
+		FDynamicMeshEditor MeshEditor(&TargetMesh);
+		FDynamicMeshEditResult EditResult;
+		MeshEditor.SplitBowties(EditResult);
+
 		FMergeCoincidentMeshEdges Welder(&TargetMesh);
 		Welder.MergeVertexTolerance = Tolerance * 0.01;
 		Welder.OnlyUniquePairs = false;
