fix 'extrude path' not generating a flat mesh if the user clicks 'complete' too early

diff --git a/Engine/Plugins/Experimental/MeshModelingToolsetExp/Source/MeshModelingToolsExp/Private/DrawPolyPathTool.cpp b/Engine/Plugins/Experimental/MeshModelingToolsetExp/Source/MeshModelingToolsExp/Private/DrawPolyPathTool.cpp
--- a/Engine/Plugins/Experimental/MeshModelingToolsetExp/Source/MeshModelingToolsExp/Private/DrawPolyPathTool.cpp
+++ b/Engine/Plugins/Experimental/MeshModelingToolsetExp/Source/MeshModelingToolsExp/Private/DrawPolyPathTool.cpp
@@ -266,7 +266,7 @@ void UDrawPolyPathTool::Shutdown(EToolShutdownType ShutdownType)
 	// it's a bit more work.
 	if (State == EState::SettingHeight && ShutdownType != EToolShutdownType::Cancel)
 	{
-		CurHeight = TransformProps->ExtrudeHeight;
+		CurHeight = TransformProps->ExtrudeMode == EDrawPolyPathExtrudeMode::Flat ? 0.0 : TransformProps->ExtrudeHeight;
 		OnCompleteExtrudeHeight();
 	}
 
