Fix C++20 compile error. Make UCubeGridTool::FSelection's comparison operators const to address ambiguity issues.

diff --git a/Engine/Plugins/Experimental/MeshModelingToolsetExp/Source/MeshModelingToolsExp/Public/CubeGridTool.h b/Engine/Plugins/Experimental/MeshModelingToolsetExp/Source/MeshModelingToolsExp/Public/CubeGridTool.h
--- a/Engine/Plugins/Experimental/MeshModelingToolsetExp/Source/MeshModelingToolsExp/Public/CubeGridTool.h
+++ b/Engine/Plugins/Experimental/MeshModelingToolsetExp/Source/MeshModelingToolsExp/Public/CubeGridTool.h
@@ -341,14 +341,14 @@ public:
 		// Direction must be initialized to a valid enum value (0 is not one).
 		FCubeGrid::EFaceDirection Direction = FCubeGrid::EFaceDirection::PositiveX;
 
-		bool operator==(const FSelection& Other)
+		bool operator==(const FSelection& Other) const
 		{
 			return Box == Other.Box
 				&& StartBox == Other.StartBox
 				&& Direction == Other.Direction;
 		}
 
-		bool operator!=(const FSelection& Other)
+		bool operator!=(const FSelection& Other) const
 		{
 			return !(*this == Other);
 		}
