Skin paint tool using wrong index.

diff --git a/Engine/Plugins/Experimental/MeshModelingToolsetExp/Source/MeshModelingToolsExp/Private/SkinWeightsPaintTool.cpp b/Engine/Plugins/Experimental/MeshModelingToolsetExp/Source/MeshModelingToolsExp/Private/SkinWeightsPaintTool.cpp
--- a/Engine/Plugins/Experimental/MeshModelingToolsetExp/Source/MeshModelingToolsExp/Private/SkinWeightsPaintTool.cpp
+++ b/Engine/Plugins/Experimental/MeshModelingToolsetExp/Source/MeshModelingToolsExp/Private/SkinWeightsPaintTool.cpp
@@ -212,7 +212,7 @@ void FSkinToolWeights::InitializeSkinWeights(
 			check(InfluenceIndex < MAX_TOTAL_INFLUENCES);
 			const int32 BoneIndex = BoneWeight.GetBoneIndex();
 			const float Weight = BoneWeight.GetWeight();
-			const FVector& RefPoseVertexPosition = Deformer.RefPoseVertexPositions[VertexID];
+			const FVector& RefPoseVertexPosition = Deformer.RefPoseVertexPositions[VertexIndex];
 			const FTransform& InvRefPoseTransform = Deformer.InvCSRefPoseTransforms[BoneIndex];
 			const FVector& BoneLocalPositionInRefPose = InvRefPoseTransform.TransformPosition(RefPoseVertexPosition);
 			CurrentWeights[VertexIndex].Emplace(BoneIndex, BoneLocalPositionInRefPose, Weight);
