Fix vox offset tool issues:
 - Warn about artifacts on meshes with open boundaries
 - Fix some settings not being saved/restored
 - Fix transaction name on undo/redo not matching the tool name

Also: Make the asset warning for base interactive tool only clear the tool's display message if it set a display message when it was last called, so that it is less likely to clear unrelated display messages when called with no warning to display.

diff --git a/Engine/Plugins/Experimental/MeshModelingToolsetExp/Source/MeshModelingToolsEditorOnlyExp/Private/PolygonOnMeshTool.cpp b/Engine/Plugins/Experimental/MeshModelingToolsetExp/Source/MeshModelingToolsEditorOnlyExp/Private/PolygonOnMeshTool.cpp
--- a/Engine/Plugins/Experimental/MeshModelingToolsetExp/Source/MeshModelingToolsEditorOnlyExp/Private/PolygonOnMeshTool.cpp
+++ b/Engine/Plugins/Experimental/MeshModelingToolsetExp/Source/MeshModelingToolsEditorOnlyExp/Private/PolygonOnMeshTool.cpp
@@ -122,10 +122,13 @@ void UPolygonOnMeshTool::Setup()
 				GetToolManager()->DisplayMessage(LOCTEXT("BoundaryEdgesNotification", "Cut resulted in one or more boundary edges."),
 					EToolMessageLevel::UserWarning);
 			}
+			else if (UpdatedPreview->HaveEmptyResult())
+			{
+				UpdateAcceptWarnings(EAcceptWarning::EmptyForbidden);
+			}
 			else
 			{
-				// This clears the warning if needed
-				UpdateAcceptWarnings(UpdatedPreview->HaveEmptyResult() ? EAcceptWarning::EmptyForbidden : EAcceptWarning::NoWarning);
+				GetToolManager()->DisplayMessage(FText(), EToolMessageLevel::UserWarning);
 			}
 		}
 	);
diff --git a/Engine/Plugins/Experimental/MeshModelingToolsetExp/Source/MeshModelingToolsExp/Private/VoxelBlendMeshesTool.cpp b/Engine/Plugins/Experimental/MeshModelingToolsetExp/Source/MeshModelingToolsExp/Private/VoxelBlendMeshesTool.cpp
--- a/Engine/Plugins/Experimental/MeshModelingToolsetExp/Source/MeshModelingToolsExp/Private/VoxelBlendMeshesTool.cpp
+++ b/Engine/Plugins/Experimental/MeshModelingToolsetExp/Source/MeshModelingToolsExp/Private/VoxelBlendMeshesTool.cpp
@@ -53,6 +53,19 @@ void UVoxelBlendMeshesTool::SaveProperties()
 }
 
 
+void UVoxelBlendMeshesTool::ConvertInputsAndSetPreviewMaterials(bool bSetPreviewMesh)
+{
+	Super::ConvertInputsAndSetPreviewMaterials(bSetPreviewMesh);
+
+	if (!BlendProperties->bVoxWrap && HasOpenBoundariesInMeshInputs())
+	{
+		GetToolManager()->DisplayMessage(
+			LOCTEXT("WarnOpenEdges", "Open edges found: Consider using the 'Vox Wrap Preprocess' settings to avoid artifacts."),
+			EToolMessageLevel::UserWarning);
+	}
+}
+
+
 TUniquePtr<FDynamicMeshOperator> UVoxelBlendMeshesTool::MakeNewOperator()
 {
 	TUniquePtr<FVoxelBlendMeshesOp> Op = MakeUnique<FVoxelBlendMeshesOp>();
diff --git a/Engine/Plugins/Experimental/MeshModelingToolsetExp/Source/MeshModelingToolsExp/Private/VoxelMorphologyMeshesTool.cpp b/Engine/Plugins/Experimental/MeshModelingToolsetExp/Source/MeshModelingToolsExp/Private/VoxelMorphologyMeshesTool.cpp
--- a/Engine/Plugins/Experimental/MeshModelingToolsetExp/Source/MeshModelingToolsExp/Private/VoxelMorphologyMeshesTool.cpp
+++ b/Engine/Plugins/Experimental/MeshModelingToolsetExp/Source/MeshModelingToolsExp/Private/VoxelMorphologyMeshesTool.cpp
@@ -48,8 +48,20 @@ void UVoxelMorphologyMeshesTool::SetupProperties()
 void UVoxelMorphologyMeshesTool::SaveProperties()
 {
 	Super::SaveProperties();
+	MorphologyProperties->SaveProperties(this);
+}
+
+
+void UVoxelMorphologyMeshesTool::ConvertInputsAndSetPreviewMaterials(bool bSetPreviewMesh)
+{
+	Super::ConvertInputsAndSetPreviewMaterials(bSetPreviewMesh);
 
-	VoxProperties->SaveProperties(this);
+	if (!MorphologyProperties->bVoxWrap && HasOpenBoundariesInMeshInputs())
+	{
+		GetToolManager()->DisplayMessage(
+			LOCTEXT("WarnOpenEdges", "Open edges found: Consider using the 'Vox Wrap Preprocess' settings to avoid artifacts."),
+			EToolMessageLevel::UserWarning);
+	}
 }
 
 
@@ -79,12 +91,12 @@ TUniquePtr<FDynamicMeshOperator> UVoxelMorphologyMeshesTool::MakeNewOperator()
 
 FString UVoxelMorphologyMeshesTool::GetCreatedAssetName() const
 {
-	return TEXT("Morphology");
+	return TEXT("Offset");
 }
 
 FText UVoxelMorphologyMeshesTool::GetActionName() const
 {
-	return LOCTEXT("VoxelMorphologyMeshes", "Voxel Morphology");
+	return LOCTEXT("VoxelMorphologyMeshes", "Voxel Offset");
 }
 
 
diff --git a/Engine/Plugins/Experimental/MeshModelingToolsetExp/Source/MeshModelingToolsExp/Public/VoxelBlendMeshesTool.h b/Engine/Plugins/Experimental/MeshModelingToolsetExp/Source/MeshModelingToolsExp/Public/VoxelBlendMeshesTool.h
--- a/Engine/Plugins/Experimental/MeshModelingToolsetExp/Source/MeshModelingToolsExp/Public/VoxelBlendMeshesTool.h
+++ b/Engine/Plugins/Experimental/MeshModelingToolsetExp/Source/MeshModelingToolsExp/Public/VoxelBlendMeshesTool.h
@@ -79,6 +79,10 @@ protected:
 	virtual FString GetCreatedAssetName() const override;
 	virtual FText GetActionName() const override;
 
+	virtual void ConvertInputsAndSetPreviewMaterials(bool bSetPreviewMesh) override;
+
+	// TODO: this tool also needs a warning for meshes w/ open boundaries
+
 	// IDynamicMeshOperatorFactory API
 	virtual TUniquePtr<UE::Geometry::FDynamicMeshOperator> MakeNewOperator() override;
 
diff --git a/Engine/Plugins/Experimental/MeshModelingToolsetExp/Source/MeshModelingToolsExp/Public/VoxelMorphologyMeshesTool.h b/Engine/Plugins/Experimental/MeshModelingToolsetExp/Source/MeshModelingToolsExp/Public/VoxelMorphologyMeshesTool.h
--- a/Engine/Plugins/Experimental/MeshModelingToolsetExp/Source/MeshModelingToolsExp/Public/VoxelMorphologyMeshesTool.h
+++ b/Engine/Plugins/Experimental/MeshModelingToolsetExp/Source/MeshModelingToolsExp/Public/VoxelMorphologyMeshesTool.h
@@ -62,6 +62,8 @@ protected:
 	virtual FString GetCreatedAssetName() const override;
 	virtual FText GetActionName() const override;
 
+	virtual void ConvertInputsAndSetPreviewMaterials(bool bSetPreviewMesh) override;
+
 	// IDynamicMeshOperatorFactory API
 	virtual TUniquePtr<UE::Geometry::FDynamicMeshOperator> MakeNewOperator() override;
 
diff --git a/Engine/Plugins/Runtime/MeshModelingToolset/Source/ModelingComponents/Private/BaseTools/BaseVoxelTool.cpp b/Engine/Plugins/Runtime/MeshModelingToolset/Source/ModelingComponents/Private/BaseTools/BaseVoxelTool.cpp
--- a/Engine/Plugins/Runtime/MeshModelingToolset/Source/ModelingComponents/Private/BaseTools/BaseVoxelTool.cpp
+++ b/Engine/Plugins/Runtime/MeshModelingToolset/Source/ModelingComponents/Private/BaseTools/BaseVoxelTool.cpp
@@ -35,6 +35,18 @@ void UBaseVoxelTool::SaveProperties()
 	VoxProperties->SaveProperties(this);
 }
 
+bool UBaseVoxelTool::HasOpenBoundariesInMeshInputs()
+{
+	for (int32 Idx = 0; Idx < OriginalDynamicMeshes.Num(); ++Idx)
+	{
+		if (!OriginalDynamicMeshes[Idx]->IsClosed() && OriginalDynamicMeshes[Idx]->TriangleCount() > 0)
+		{
+			return true;
+		}
+	}
+	return false;
+}
+
 
 TArray<UMaterialInterface*> UBaseVoxelTool::GetOutputMaterials() const
 {
diff --git a/Engine/Plugins/Runtime/MeshModelingToolset/Source/ModelingComponents/Public/BaseTools/BaseVoxelTool.h b/Engine/Plugins/Runtime/MeshModelingToolset/Source/ModelingComponents/Public/BaseTools/BaseVoxelTool.h
--- a/Engine/Plugins/Runtime/MeshModelingToolset/Source/ModelingComponents/Public/BaseTools/BaseVoxelTool.h
+++ b/Engine/Plugins/Runtime/MeshModelingToolset/Source/ModelingComponents/Public/BaseTools/BaseVoxelTool.h
@@ -32,6 +32,9 @@ protected:
 	/** Sets the output material to the default "world grid" material */
 	virtual TArray<UMaterialInterface*> GetOutputMaterials() const override;
 
+	// Test whether any of the OriginalDynamicMeshes has any open boundary edges
+	virtual bool HasOpenBoundariesInMeshInputs();
+
 
 	UPROPERTY()
 	TObjectPtr<UVoxelProperties> VoxProperties;
diff --git a/Engine/Source/Runtime/InteractiveToolsFramework/Private/InteractiveTool.cpp b/Engine/Source/Runtime/InteractiveToolsFramework/Private/InteractiveTool.cpp
--- a/Engine/Source/Runtime/InteractiveToolsFramework/Private/InteractiveTool.cpp
+++ b/Engine/Source/Runtime/InteractiveToolsFramework/Private/InteractiveTool.cpp
@@ -282,8 +282,11 @@ void UInteractiveTool::UpdateAcceptWarnings(EAcceptWarning Warning)
 	switch (Warning)
 	{
 	case EAcceptWarning::NoWarning:
-		// clear warning
-		GetToolManager()->DisplayMessage(FText(), EToolMessageLevel::UserWarning);
+		if (bLastShowedAcceptWarning)
+		{
+			// clear warning
+			GetToolManager()->DisplayMessage(FText(), EToolMessageLevel::UserWarning);
+		}
 		break;
 	case EAcceptWarning::EmptyForbidden:
 		GetToolManager()->DisplayMessage(LOCTEXT("CannotCreateEmptyMesh", "WARNING: Tool doesn't allow creation of an empty mesh."),
@@ -292,6 +295,7 @@ void UInteractiveTool::UpdateAcceptWarnings(EAcceptWarning Warning)
 	default:
 		check(false);
 	}
+	bLastShowedAcceptWarning = Warning != EAcceptWarning::NoWarning;
 }
 
 
diff --git a/Engine/Source/Runtime/InteractiveToolsFramework/Public/InteractiveTool.h b/Engine/Source/Runtime/InteractiveToolsFramework/Public/InteractiveTool.h
--- a/Engine/Source/Runtime/InteractiveToolsFramework/Public/InteractiveTool.h
+++ b/Engine/Source/Runtime/InteractiveToolsFramework/Public/InteractiveTool.h
@@ -538,6 +538,11 @@ protected:
 	 */
 	virtual void UpdateAcceptWarnings(EAcceptWarning Warning);
 
+private:
+	// Tracks whether the UpdateAcceptWarnings function showed a warning the last time it was called.
+	// Used to avoid clearing the tool display message in cases where it was not set by this function.
+	bool bLastShowedAcceptWarning = false;
+
 
 
 protected:
