Skeleton Editing Tool

diff --git a/Engine/Plugins/Editor/ModelingToolsEditorMode/Source/ModelingToolsEditorMode/Private/ModelingToolsManagerActions.cpp b/Engine/Plugins/Editor/ModelingToolsEditorMode/Source/ModelingToolsEditorMode/Private/ModelingToolsManagerActions.cpp
--- a/Engine/Plugins/Editor/ModelingToolsEditorMode/Source/ModelingToolsEditorMode/Private/ModelingToolsManagerActions.cpp
+++ b/Engine/Plugins/Editor/ModelingToolsEditorMode/Source/ModelingToolsEditorMode/Private/ModelingToolsManagerActions.cpp
@@ -332,6 +332,8 @@ void FModelingToolsManagerCommands::RegisterCommands()
 	UI_COMMAND(BeginSkinWeightsPaintTool, "SkinWts", "Start the Paint Skin Weights Tool", EUserInterfaceActionType::Button, FInputChord());
 	UI_COMMAND(BeginSkinWeightsBindingTool, "SkinBind", "Start the Skin Weights Binding Tool", EUserInterfaceActionType::Button, FInputChord());
 
+	REGISTER_MODELING_TOOL_COMMAND(BeginSkeletonEditingTool, "SkeletonEditing", "Start the Skeleton Editing Tool");
+
 	REGISTER_MODELING_TOOL_COMMAND(BeginPolyModelTool_Inset, "Inset", "Inset the current Selected Faces");
 	REGISTER_MODELING_TOOL_COMMAND(BeginPolyModelTool_Outset, "Outset", "Outset the current Selected Faces");
 	REGISTER_MODELING_TOOL_COMMAND(BeginPolyModelTool_CutFaces, "Cut", "Cut the current Selected Faces");
diff --git a/Engine/Plugins/Editor/ModelingToolsEditorMode/Source/ModelingToolsEditorMode/Public/ModelingToolsManagerActions.h b/Engine/Plugins/Editor/ModelingToolsEditorMode/Source/ModelingToolsEditorMode/Public/ModelingToolsManagerActions.h
--- a/Engine/Plugins/Editor/ModelingToolsEditorMode/Source/ModelingToolsEditorMode/Public/ModelingToolsManagerActions.h
+++ b/Engine/Plugins/Editor/ModelingToolsEditorMode/Source/ModelingToolsEditorMode/Public/ModelingToolsManagerActions.h
@@ -126,6 +126,8 @@ public:
 
 	TSharedPtr<FUICommandInfo> BeginSkinWeightsPaintTool;
 	TSharedPtr<FUICommandInfo> BeginSkinWeightsBindingTool;
+
+	TSharedPtr<FUICommandInfo> BeginSkeletonEditingTool;
 	
 	TSharedPtr<FUICommandInfo> BeginLODManagerTool;
 	TSharedPtr<FUICommandInfo> BeginGenerateStaticMeshLODAssetTool;
diff --git a/Engine/Plugins/Experimental/Animation/SkeletalMeshModelingTools/Source/SkeletalMeshModelingTools/Private/SkeletalMeshModelingToolsCommands.cpp b/Engine/Plugins/Experimental/Animation/SkeletalMeshModelingTools/Source/SkeletalMeshModelingTools/Private/SkeletalMeshModelingToolsCommands.cpp
--- a/Engine/Plugins/Experimental/Animation/SkeletalMeshModelingTools/Source/SkeletalMeshModelingTools/Private/SkeletalMeshModelingToolsCommands.cpp
+++ b/Engine/Plugins/Experimental/Animation/SkeletalMeshModelingTools/Source/SkeletalMeshModelingTools/Private/SkeletalMeshModelingToolsCommands.cpp
@@ -2,6 +2,8 @@
 
 #include "SkeletalMeshModelingToolsCommands.h"
 
+#include "SkeletalMesh/SkeletonEditingTool.h"
+
 #define LOCTEXT_NAMESPACE "SkeletalMeshModelingToolsCommands"
 
 
@@ -15,4 +17,47 @@ const FSkeletalMeshModelingToolsCommands& FSkeletalMeshModelingToolsCommands::Ge
 	return TCommands<FSkeletalMeshModelingToolsCommands>::Get();
 }
 
+FSkeletalMeshModelingToolsActionCommands::FSkeletalMeshModelingToolsActionCommands() : 
+	TInteractiveToolCommands<FSkeletalMeshModelingToolsActionCommands>(
+		"SeletalMeshModelingToolsEditMode", // Context name for fast lookup
+		NSLOCTEXT("Contexts", "SeletalMeshModelingToolsEditMode", "Skeletal Mesh Modeling Tools - Shared Shortcuts"), // Localized context name for displaying
+		NAME_None, // Parent
+		FAppStyle::GetAppStyleSetName() // Icon Style Set
+	)
+{
+}
+
+void FSkeletalMeshModelingToolsActionCommands::GetToolDefaultObjectList(TArray<UInteractiveTool*>& ToolCDOs)
+{}
+
+void FSkeletalMeshModelingToolsActionCommands::RegisterAllToolActions()
+{
+	FSkeletonEditingToolActionCommands::Register();
+}
+
+void FSkeletalMeshModelingToolsActionCommands::UnregisterAllToolActions()
+{
+	FSkeletonEditingToolActionCommands::Unregister();
+}
+
+void FSkeletalMeshModelingToolsActionCommands::UpdateToolCommandBinding(UInteractiveTool* Tool, TSharedPtr<FUICommandList> UICommandList, bool bUnbind)
+{
+#define UPDATE_BINDING(CommandsType)  if (!bUnbind) CommandsType::Get().BindCommandsForCurrentTool(UICommandList, Tool); else CommandsType::Get().UnbindActiveCommands(UICommandList);
+
+	if (ExactCast<USkeletonEditingTool>(Tool) != nullptr)
+	{
+		UPDATE_BINDING(FSkeletonEditingToolActionCommands);
+	}
+}
+
+#define DEFINE_TOOL_ACTION_COMMANDS(CommandsClassName, ContextNameString, SettingsDialogString, ToolClassName ) \
+CommandsClassName::CommandsClassName() : TInteractiveToolCommands<CommandsClassName>( \
+ContextNameString, NSLOCTEXT("Contexts", ContextNameString, SettingsDialogString), NAME_None, FAppStyle::GetAppStyleSetName()) {} \
+void CommandsClassName::GetToolDefaultObjectList(TArray<UInteractiveTool*>& ToolCDOs) \
+{\
+ToolCDOs.Add(GetMutableDefault<ToolClassName>()); \
+}
+
+DEFINE_TOOL_ACTION_COMMANDS(FSkeletonEditingToolActionCommands, "SkeletalMeshModelingToolsSkeletonEditing", "Skeletal Mesh Modeling Tools - Skeleton Editing Tool", USkeletonEditingTool);
+
 #undef LOCTEXT_NAMESPACE
diff --git a/Engine/Plugins/Experimental/Animation/SkeletalMeshModelingTools/Source/SkeletalMeshModelingTools/Private/SkeletalMeshModelingToolsCommands.h b/Engine/Plugins/Experimental/Animation/SkeletalMeshModelingTools/Source/SkeletalMeshModelingTools/Private/SkeletalMeshModelingToolsCommands.h
--- a/Engine/Plugins/Experimental/Animation/SkeletalMeshModelingTools/Source/SkeletalMeshModelingTools/Private/SkeletalMeshModelingToolsCommands.h
+++ b/Engine/Plugins/Experimental/Animation/SkeletalMeshModelingTools/Source/SkeletalMeshModelingTools/Private/SkeletalMeshModelingToolsCommands.h
@@ -5,6 +5,7 @@
 #include "CoreMinimal.h"
 #include "Styling/AppStyle.h"
 #include "Framework/Commands/Commands.h"
+#include "Tools/InteractiveToolsCommands.h"
 
 
 class FSkeletalMeshModelingToolsCommands : public TCommands<FSkeletalMeshModelingToolsCommands>
@@ -19,9 +20,32 @@ public:
 	{
 	}
 
-	void RegisterCommands() override;
+	virtual void RegisterCommands() override;
 	static const FSkeletalMeshModelingToolsCommands& Get();
 
 	// Modeling tools commands
 	TSharedPtr<FUICommandInfo> ToggleModelingToolsMode;
 };
+
+class FSkeletalMeshModelingToolsActionCommands : public TInteractiveToolCommands<FSkeletalMeshModelingToolsActionCommands>
+{
+public:
+	FSkeletalMeshModelingToolsActionCommands();
+
+	// TInteractiveToolCommands
+	virtual void GetToolDefaultObjectList(TArray<UInteractiveTool*>& ToolCDOs) override;
+	
+	static void RegisterAllToolActions();
+	static void UnregisterAllToolActions();
+	static void UpdateToolCommandBinding(UInteractiveTool* Tool, TSharedPtr<FUICommandList> UICommandList, bool bUnbind = false);
+};
+
+#define DECLARE_TOOL_ACTION_COMMANDS(CommandsClassName) \
+class CommandsClassName : public TInteractiveToolCommands<CommandsClassName> \
+{\
+public:\
+CommandsClassName();\
+virtual void GetToolDefaultObjectList(TArray<UInteractiveTool*>& ToolCDOs) override;\
+};\
+
+DECLARE_TOOL_ACTION_COMMANDS(FSkeletonEditingToolActionCommands);
\ No newline at end of file
diff --git a/Engine/Plugins/Experimental/Animation/SkeletalMeshModelingTools/Source/SkeletalMeshModelingTools/Private/SkeletalMeshModelingToolsEditorMode.cpp b/Engine/Plugins/Experimental/Animation/SkeletalMeshModelingTools/Source/SkeletalMeshModelingTools/Private/SkeletalMeshModelingToolsEditorMode.cpp
--- a/Engine/Plugins/Experimental/Animation/SkeletalMeshModelingTools/Source/SkeletalMeshModelingTools/Private/SkeletalMeshModelingToolsEditorMode.cpp
+++ b/Engine/Plugins/Experimental/Animation/SkeletalMeshModelingTools/Source/SkeletalMeshModelingTools/Private/SkeletalMeshModelingToolsEditorMode.cpp
@@ -40,6 +40,8 @@
 #include "SkeletalMeshNotifier.h"
 #include "SkeletalMesh/SkeletalMeshEditionInterface.h"
 
+#include "SkeletalMesh/SkeletonEditingTool.h"
+
 #define LOCTEXT_NAMESPACE "SkeletalMeshModelingToolsEditorMode"
 
 
@@ -229,6 +231,9 @@ void USkeletalMeshModelingToolsEditorMode::Enter()
 	RegisterTool(ToolManagerCommands.BeginSkinWeightsPaintTool, TEXT("BeginSkinWeightsPaintTool"), NewObject<USkinWeightsPaintToolBuilder>());
 	RegisterTool(ToolManagerCommands.BeginSkinWeightsBindingTool, TEXT("BeginSkinWeightsBindingTool"), NewObject<USkinWeightsBindingToolBuilder>());
 
+	// Skeleton Editing
+	RegisterTool(ToolManagerCommands.BeginSkeletonEditingTool, TEXT("BeginSkeletonEditingTool"), NewObject<USkeletonEditingToolBuilder>());
+	
 	GetInteractiveToolsContext()->ToolManager->SelectActiveToolType(EToolSide::Left, TEXT("BeginSkinWeightsPaintTool"));
 }
 
@@ -311,21 +316,37 @@ bool USkeletalMeshModelingToolsEditorMode::ComputeBoundingBoxForViewportFocus(AA
 					InOutBox += Component->GetBoneLocation(BoneName);
 				}
 				
-				return true;
+				return true; 
 			}
 		}
 	}
+
+	// if Tool supports custom Focus box, use that
+	if (GetToolManager()->HasAnyActiveTool())
+	{
+		UInteractiveTool* Tool = GetToolManager()->GetActiveTool(EToolSide::Mouse);
+		IInteractiveToolCameraFocusAPI* FocusAPI = Cast<IInteractiveToolCameraFocusAPI>(Tool);
+		if (FocusAPI && FocusAPI->SupportsWorldSpaceFocusBox())
+		{
+			InOutBox = FocusAPI->GetWorldSpaceFocusBox();
+			return true;
+		}
+	}
 	
 	return Super::ComputeBoundingBoxForViewportFocus(Actor, PrimitiveComponent, InOutBox);
 }
 
 void USkeletalMeshModelingToolsEditorMode::OnToolStarted(UInteractiveToolManager* Manager, UInteractiveTool* Tool)
 {
+	FSkeletalMeshModelingToolsActionCommands::UpdateToolCommandBinding(Tool, Toolkit->GetToolkitCommands(), false);
+	
 	ConnectTool(Tool);
 }
 
 void USkeletalMeshModelingToolsEditorMode::OnToolEnded(UInteractiveToolManager* Manager, UInteractiveTool* Tool)
 {
+	FSkeletalMeshModelingToolsActionCommands::UpdateToolCommandBinding(Tool, Toolkit->GetToolkitCommands(), true);
+	
 	DisconnectTool(Tool);
 }
 
diff --git a/Engine/Plugins/Experimental/Animation/SkeletalMeshModelingTools/Source/SkeletalMeshModelingTools/Private/SkeletalMeshModelingToolsEditorModeToolkit.cpp b/Engine/Plugins/Experimental/Animation/SkeletalMeshModelingTools/Source/SkeletalMeshModelingTools/Private/SkeletalMeshModelingToolsEditorModeToolkit.cpp
--- a/Engine/Plugins/Experimental/Animation/SkeletalMeshModelingTools/Source/SkeletalMeshModelingTools/Private/SkeletalMeshModelingToolsEditorModeToolkit.cpp
+++ b/Engine/Plugins/Experimental/Animation/SkeletalMeshModelingTools/Source/SkeletalMeshModelingTools/Private/SkeletalMeshModelingToolsEditorModeToolkit.cpp
@@ -20,10 +20,19 @@
 #include "ModelingToolsManagerActions.h"
 #include "BaseTools/MeshSurfacePointTool.h"
 #include "SPrimaryButton.h"
+#include "Toolkits/AssetEditorModeUILayer.h"
 
 
 #define LOCTEXT_NAMESPACE "SkeletalMeshModelingToolsEditorModeToolkit"
 
+FSkeletalMeshModelingToolsEditorModeToolkit::~FSkeletalMeshModelingToolsEditorModeToolkit()
+{
+	if (UEditorInteractiveToolsContext* Context = GetScriptableEditorMode()->GetInteractiveToolsContext(EToolsContextScope::EdMode))
+	{
+		Context->OnToolNotificationMessage.RemoveAll(this);
+		Context->OnToolWarningMessage.RemoveAll(this);
+	}
+}
 
 void FSkeletalMeshModelingToolsEditorModeToolkit::Init(
 	const TSharedPtr<IToolkitHost>& InToolkitHost, 
@@ -95,9 +104,12 @@ void FSkeletalMeshModelingToolsEditorModeToolkit::Init(
 
 	ActiveToolName = FText::GetEmpty();
 	ActiveToolMessage = FText::GetEmpty();
-	
-	GetScriptableEditorMode()->GetInteractiveToolsContext()->OnToolNotificationMessage.AddSP(this, &FSkeletalMeshModelingToolsEditorModeToolkit::PostNotification);
-	GetScriptableEditorMode()->GetInteractiveToolsContext()->OnToolWarningMessage.AddSP(this, &FSkeletalMeshModelingToolsEditorModeToolkit::PostWarning);
+
+	if (UEditorInteractiveToolsContext* Context = GetScriptableEditorMode()->GetInteractiveToolsContext(EToolsContextScope::EdMode))
+	{
+		Context->OnToolNotificationMessage.AddSP(this, &FSkeletalMeshModelingToolsEditorModeToolkit::PostNotification);
+		Context->OnToolWarningMessage.AddSP(this, &FSkeletalMeshModelingToolsEditorModeToolkit::PostWarning);
+	}
 
 	SAssignNew(ViewportOverlayWidget, SHorizontalBox)
 
@@ -238,11 +250,12 @@ static const FName EditTabName(TEXT("Edit"));
 static const FName ProcessingTabName(TEXT("MeshOps"));
 static const FName DeformTabName(TEXT("Deform"));
 static const FName SkinWeightsTabName(TEXT("Skin"));
+static const FName SkeletonTabName(TEXT("Skeleton"));
 
 
 void FSkeletalMeshModelingToolsEditorModeToolkit::GetToolPaletteNames(TArray<FName>& InPaletteName) const
 {
-	InPaletteName = { EditTabName, ProcessingTabName, DeformTabName, SkinWeightsTabName };
+	InPaletteName = { EditTabName, ProcessingTabName, DeformTabName, SkinWeightsTabName, SkeletonTabName };
 }
 
 
@@ -288,6 +301,10 @@ void FSkeletalMeshModelingToolsEditorModeToolkit::BuildToolPalette(FName Palette
 		ToolbarBuilder.AddToolBarButton(Commands.BeginSkinWeightsPaintTool);
 		ToolbarBuilder.AddToolBarButton(Commands.BeginSkinWeightsBindingTool);
 	}
+	else if (PaletteName == SkeletonTabName)
+	{
+		ToolbarBuilder.AddToolBarButton(Commands.BeginSkeletonEditingTool);
+	}
 }
 
 
@@ -298,13 +315,28 @@ void FSkeletalMeshModelingToolsEditorModeToolkit::OnToolPaletteChanged(FName Pal
 
 void FSkeletalMeshModelingToolsEditorModeToolkit::PostNotification(const FText& InMessage)
 {
+	ClearNotification();
+	
 	ActiveToolMessage = InMessage;
+
+	if (ModeUILayer.IsValid())
+	{
+		const FName StatusBarName = ModeUILayer.Pin()->GetStatusBarName();
+		ActiveToolMessageHandle = GEditor->GetEditorSubsystem<UStatusBarSubsystem>()->PushStatusBarMessage(StatusBarName, ActiveToolMessage);
+	}
 }
 
 
 void FSkeletalMeshModelingToolsEditorModeToolkit::ClearNotification()
 {
 	ActiveToolMessage = FText::GetEmpty();
+
+	if (ModeUILayer.IsValid())
+	{
+		const FName StatusBarName = ModeUILayer.Pin()->GetStatusBarName();
+		GEditor->GetEditorSubsystem<UStatusBarSubsystem>()->PopStatusBarMessage(StatusBarName, ActiveToolMessageHandle);
+	}
+	ActiveToolMessageHandle.Reset();
 }
 
 
diff --git a/Engine/Plugins/Experimental/Animation/SkeletalMeshModelingTools/Source/SkeletalMeshModelingTools/Private/SkeletalMeshModelingToolsEditorModeToolkit.h b/Engine/Plugins/Experimental/Animation/SkeletalMeshModelingTools/Source/SkeletalMeshModelingTools/Private/SkeletalMeshModelingToolsEditorModeToolkit.h
--- a/Engine/Plugins/Experimental/Animation/SkeletalMeshModelingTools/Source/SkeletalMeshModelingTools/Private/SkeletalMeshModelingToolsEditorModeToolkit.h
+++ b/Engine/Plugins/Experimental/Animation/SkeletalMeshModelingTools/Source/SkeletalMeshModelingTools/Private/SkeletalMeshModelingToolsEditorModeToolkit.h
@@ -3,6 +3,7 @@
 #pragma once
 
 #include "CoreMinimal.h"
+#include "StatusBarSubsystem.h"
 #include "Toolkits/BaseToolkit.h"
 
 
@@ -15,6 +16,9 @@ class FSkeletalMeshModelingToolsEditorModeToolkit :
 	 public FModeToolkit
 {
 public:
+
+	~FSkeletalMeshModelingToolsEditorModeToolkit();
+	
 	// IToolkit overrides
 	void Init(const TSharedPtr<IToolkitHost>& InToolkitHost, TWeakObjectPtr<UEdMode> InOwningMode) override;
 	FName GetToolkitFName() const override;
@@ -44,6 +48,7 @@ private:
 
 	FText ActiveToolName;
 	FText ActiveToolMessage;
+	FStatusBarMessageHandle ActiveToolMessageHandle;
 
 	TSharedPtr<SWidget> ToolkitWidget;
 	TSharedPtr<IDetailsView> DetailsView;
diff --git a/Engine/Plugins/Experimental/Animation/SkeletalMeshModelingTools/Source/SkeletalMeshModelingTools/Private/SkeletalMeshModelingToolsModule.cpp b/Engine/Plugins/Experimental/Animation/SkeletalMeshModelingTools/Source/SkeletalMeshModelingTools/Private/SkeletalMeshModelingToolsModule.cpp
--- a/Engine/Plugins/Experimental/Animation/SkeletalMeshModelingTools/Source/SkeletalMeshModelingTools/Private/SkeletalMeshModelingToolsModule.cpp
+++ b/Engine/Plugins/Experimental/Animation/SkeletalMeshModelingTools/Source/SkeletalMeshModelingTools/Private/SkeletalMeshModelingToolsModule.cpp
@@ -40,6 +40,8 @@ void FSkeletalMeshModelingToolsModule::StartupModule()
 
 	ToolbarExtenders.Add(ISkeletalMeshEditorModule::FSkeletalMeshEditorToolbarExtender::CreateRaw(this, &FSkeletalMeshModelingToolsModule::ExtendSkelMeshEditorToolbar));
 	SkelMeshEditorExtenderHandle = ToolbarExtenders.Last().GetHandle();
+
+	FCoreDelegates::OnPostEngineInit.AddRaw(this, &FSkeletalMeshModelingToolsModule::OnPostEngineInit);
 }
 
 void FSkeletalMeshModelingToolsModule::ShutdownModule()
@@ -55,6 +57,9 @@ void FSkeletalMeshModelingToolsModule::ShutdownModule()
 
 	FSkeletalMeshModelingToolsCommands::Unregister();
 	FSkeletalMeshModelingToolsStyle::Unregister();
+	
+	FCoreDelegates::OnPostEngineInit.RemoveAll(this);
+	FSkeletalMeshModelingToolsActionCommands::UnregisterAllToolActions();
 }
 
 
@@ -148,4 +153,9 @@ void FSkeletalMeshModelingToolsModule::OnToggleModelingToolsMode(TWeakPtr<ISkele
 	}
 }
 
+void FSkeletalMeshModelingToolsModule::OnPostEngineInit()
+{
+	FSkeletalMeshModelingToolsActionCommands::RegisterAllToolActions();
+}
+
 #undef LOCTEXT_NAMESPACE
diff --git a/Engine/Plugins/Experimental/Animation/SkeletalMeshModelingTools/Source/SkeletalMeshModelingTools/Private/SkeletalMeshModelingToolsModule.h b/Engine/Plugins/Experimental/Animation/SkeletalMeshModelingTools/Source/SkeletalMeshModelingTools/Private/SkeletalMeshModelingToolsModule.h
--- a/Engine/Plugins/Experimental/Animation/SkeletalMeshModelingTools/Source/SkeletalMeshModelingTools/Private/SkeletalMeshModelingToolsModule.h
+++ b/Engine/Plugins/Experimental/Animation/SkeletalMeshModelingTools/Source/SkeletalMeshModelingTools/Private/SkeletalMeshModelingToolsModule.h
@@ -32,6 +32,8 @@ private:
 	bool IsModelingToolModeActive(TWeakPtr<ISkeletalMeshEditor> InSkeletalMeshEditor) const;
 	void OnToggleModelingToolsMode(TWeakPtr<ISkeletalMeshEditor> InSkeletalMeshEditor);
 
+	void OnPostEngineInit();
+	
 	// The handle for the extender delegate we added. Needed for clean module shutdown.
 	FDelegateHandle SkelMeshEditorExtenderHandle;
 
diff --git a/Engine/Plugins/Experimental/Animation/SkeletalMeshModelingTools/Source/SkeletalMeshModelingTools/SkeletalMeshModelingTools.Build.cs b/Engine/Plugins/Experimental/Animation/SkeletalMeshModelingTools/Source/SkeletalMeshModelingTools/SkeletalMeshModelingTools.Build.cs
--- a/Engine/Plugins/Experimental/Animation/SkeletalMeshModelingTools/Source/SkeletalMeshModelingTools/SkeletalMeshModelingTools.Build.cs
+++ b/Engine/Plugins/Experimental/Animation/SkeletalMeshModelingTools/Source/SkeletalMeshModelingTools/SkeletalMeshModelingTools.Build.cs
@@ -29,7 +29,8 @@ public class SkeletalMeshModelingTools : ModuleRules
 				"StylusInput",
 				"ToolMenus",
 				"ToolWidgets",
-				"UnrealEd", 
+				"UnrealEd",
+				"StatusBar"
 			}
 		);
 	}
diff --git a/Engine/Plugins/Experimental/MeshModelingToolsetExp/Source/MeshModelingToolsExp/Private/SkeletalMesh/SkeletonEditingTool.cpp b/Engine/Plugins/Experimental/MeshModelingToolsetExp/Source/MeshModelingToolsExp/Private/SkeletalMesh/SkeletonEditingTool.cpp
new file mode 100644
--- /dev/null
+++ b/Engine/Plugins/Experimental/MeshModelingToolsetExp/Source/MeshModelingToolsExp/Private/SkeletalMesh/SkeletonEditingTool.cpp
@@ -0,0 +1,982 @@
+﻿// Copyright Epic Games, Inc. All Rights Reserved.
+
+
+#include "SkeletalMesh/SkeletonEditingTool.h"
+
+#include "ContextObjectStore.h"
+#include "InteractiveToolManager.h"
+#include "ModelingToolTargetUtil.h"
+#include "SkeletalDebugRendering.h"
+#include "ToolTargetManager.h"
+#include "BaseBehaviors/SingleClickBehavior.h"
+#include "Engine/SkeletalMesh.h"
+#include "Components/SkeletalMeshComponent.h"
+#include "UnrealClient.h"
+#include "HitProxies.h"
+#include "BaseBehaviors/ClickDragBehavior.h"
+#include "BaseGizmos/GizmoViewContext.h"
+
+#include "Algo/Count.h"
+#include "TargetInterfaces/PrimitiveComponentBackedTarget.h"
+
+#define LOCTEXT_NAMESPACE "USkeletonEditingTool"
+
+namespace SkeletonEditingTool
+{
+	
+FRefSkeletonChange::FRefSkeletonChange(const USkeletonEditingTool* InTool)
+	: FToolCommandChange()
+	, PreChangeSkeleton(InTool->SkeletonModifier.GetReferenceSkeleton())
+	, PreBoneTracker(InTool->SkeletonModifier.GetBoneIndexTracker())
+	, PostChangeSkeleton(InTool->SkeletonModifier.GetReferenceSkeleton())
+	, PostBoneTracker(InTool->SkeletonModifier.GetBoneIndexTracker())
+{}
+
+void FRefSkeletonChange::StoreSkeleton(const USkeletonEditingTool* InTool)
+{
+	PostChangeSkeleton = InTool->SkeletonModifier.GetReferenceSkeleton();
+	PostBoneTracker = InTool->SkeletonModifier.GetBoneIndexTracker();
+}
+
+void FRefSkeletonChange::Apply(UObject* Object)
+{ // redo
+	USkeletonEditingTool* Tool = CastChecked<USkeletonEditingTool>(Object);
+	Tool->SkeletonModifier.ExternalUpdate(PostChangeSkeleton, PostBoneTracker);
+}
+
+void FRefSkeletonChange::Revert(UObject* Object)
+{ // undo
+	USkeletonEditingTool* Tool = CastChecked<USkeletonEditingTool>(Object);
+	Tool->SkeletonModifier.ExternalUpdate(PreChangeSkeleton, PreBoneTracker);
+}
+
+}
+
+/**
+ * USkeletonEditingToolBuilder
+ */
+
+bool USkeletonEditingToolBuilder::CanBuildTool(const FToolBuilderState& SceneState) const
+{
+	return SceneState.TargetManager->CountSelectedAndTargetable(SceneState, GetTargetRequirements()) == 1;
+}
+
+const FToolTargetTypeRequirements& USkeletonEditingToolBuilder::GetTargetRequirements() const
+{
+	static const FToolTargetTypeRequirements TypeRequirements(UPrimitiveComponentBackedTarget::StaticClass());
+	return TypeRequirements;
+}
+
+UInteractiveTool* USkeletonEditingToolBuilder::BuildTool(const FToolBuilderState& SceneState) const
+{
+	USkeletonEditingTool* NewTool = NewObject<USkeletonEditingTool>(SceneState.ToolManager);
+	
+	UToolTarget* Target = SceneState.TargetManager->BuildFirstSelectedTargetable(SceneState, GetTargetRequirements());
+	NewTool->SetTarget(Target);
+	NewTool->Init(SceneState);
+	
+	return NewTool;
+}
+
+/**
+ * USkeletonEditingTool
+ */
+
+void USkeletonEditingTool::Init(const FToolBuilderState& InSceneState)
+{
+	TargetWorld = InSceneState.World;
+	ViewContext = InSceneState.ToolManager->GetContextObjectStore()->FindContext<UGizmoViewContext>();
+}
+
+void USkeletonEditingTool::Setup()
+{
+	Super::Setup();
+
+	const IPrimitiveComponentBackedTarget* TargetComponent = Cast<IPrimitiveComponentBackedTarget>(Target);
+	const USkeletalMeshComponent* Component = Cast<USkeletalMeshComponent>(TargetComponent->GetOwnerComponent());
+
+	USkeletalMesh* SkeletalMesh = Component ? Component->GetSkeletalMeshAsset() : nullptr;
+	if (!SkeletalMesh)
+	{
+		return;
+	}
+
+	// setup modifier
+	SkeletonModifier.Init(SkeletalMesh);
+
+	// setup current bone
+	const FReferenceSkeleton& RefSkeleton = SkeletonModifier.GetReferenceSkeleton();
+	const int32 NumBones = RefSkeleton.GetNum();
+	const FName& RootBoneName = NumBones ? RefSkeleton.GetBoneName(0) : NAME_None;
+
+	if (NumBones)
+	{
+		CurrentBone = RootBoneName;
+	}
+
+	// setup preview
+	{
+		PreviewMesh = NewObject<UPreviewMesh>(this);
+		PreviewMesh->bBuildSpatialDataStructure = true;
+		PreviewMesh->CreateInWorld(TargetWorld.Get(), FTransform::Identity);
+
+		PreviewMesh->SetTransform(UE::ToolTarget::GetLocalToWorldTransform(Target));
+
+		PreviewMesh->ReplaceMesh(UE::ToolTarget::GetDynamicMeshCopy(Target));
+
+		const FComponentMaterialSet MaterialSet = UE::ToolTarget::GetMaterialSet(Target);
+		PreviewMesh->SetMaterials(MaterialSet.Materials);
+
+		// hide the skeletal mesh component
+		UE::ToolTarget::HideSourceObject(Target);
+	}
+
+	// setup properties
+	{
+		Properties = NewObject<USkeletonEditingProperties>();
+		Properties->Initialize(this);
+		Properties->RestoreProperties(this);
+		AddToolPropertySource(Properties);
+
+		ProjectionProperties = NewObject<UProjectionProperties>();
+		ProjectionProperties->Initialize(PreviewMesh);
+		ProjectionProperties->RestoreProperties(this);
+		AddToolPropertySource(ProjectionProperties);
+
+		MirroringProperties = NewObject<UMirroringProperties>();
+		MirroringProperties->Initialize(this);
+		MirroringProperties->RestoreProperties(this);
+		AddToolPropertySource(MirroringProperties);
+
+		OrientingProperties = NewObject<UOrientingProperties>();
+		OrientingProperties->Initialize(this);
+		OrientingProperties->RestoreProperties(this);
+		AddToolPropertySource(OrientingProperties);
+	}
+
+	// setup drag & drop behaviour
+	{
+		UClickDragInputBehavior* ClickDragBehavior = NewObject<UClickDragInputBehavior>(this);
+		ClickDragBehavior->Initialize(this);
+		AddInputBehavior(ClickDragBehavior);
+	}
+}
+
+void USkeletonEditingTool::Shutdown(EToolShutdownType ShutdownType)
+{
+	if (ShutdownType == EToolShutdownType::Accept)
+	{
+		GetToolManager()->BeginUndoTransaction(LOCTEXT("SkeletonEditingTool", "Commit Skeleton Editing"));
+		SkeletonModifier.CommitSkeletonToSkeletalMesh();
+		GetToolManager()->EndUndoTransaction();
+
+		// to force to refresh the tree
+		if (NeedsNotification())
+		{
+			GetNotifier().Notify({}, ESkeletalMeshNotifyType::BonesAdded);
+		}
+	}
+	
+	Super::Shutdown(ShutdownType);
+
+	// remove preview mesh
+	if (PreviewMesh != nullptr)
+	{
+		PreviewMesh->SetVisible(false);
+		PreviewMesh->Disconnect();
+		PreviewMesh = nullptr;
+	}
+
+	// show the skeletal mesh component
+	UE::ToolTarget::ShowSourceObject(Target);
+
+	//save properties	
+	Properties->SaveProperties(this);
+	ProjectionProperties->SaveProperties(this);
+	MirroringProperties->SaveProperties(this);
+	OrientingProperties->SaveProperties(this);
+}
+
+void USkeletonEditingTool::RegisterActions(FInteractiveToolActionSet& ActionSet)
+{
+	Super::RegisterActions(ActionSet);
+
+	int32 ActionId = static_cast<int32>(EStandardToolActions::BaseClientDefinedActionID) + 400;
+	auto GetActionId = [&ActionId]
+	{
+		return ActionId++;
+	};
+	
+	// register New key
+	ActionSet.RegisterAction(this, GetActionId(), TEXT("CreateNewBone"),
+		LOCTEXT("CreateNewBone", "Create New Bone"),
+		LOCTEXT("CreateNewBoneDesc", "Create New Bone"),
+		EModifierKey::None, EKeys::N,
+		[this]()
+		{
+			Operation = EEditingOperation::Create;
+			GetToolManager()->DisplayMessage(LOCTEXT("Create", "Click & Drag to place a new bone."), EToolMessageLevel::UserNotification);
+		});
+	
+	// register Delete key
+	ActionSet.RegisterAction(this, GetActionId(), TEXT("DeleteSelectedBones"),
+		LOCTEXT("DeleteSelectedBones", "Delete Selected Bone(s)"),
+		LOCTEXT("DeleteSelectedBonesDesc", "Delete Selected Bone(s)"),
+		EModifierKey::None, EKeys::Delete,
+		[this]()
+		{
+			RemoveBones();
+		});
+
+	// register Select key
+	ActionSet.RegisterAction(this, GetActionId(), TEXT("SelectBones"),
+		LOCTEXT("SelectBone", "Select Bone"),
+		LOCTEXT("SelectDesc", "Select Bone"),
+		EModifierKey::None, EKeys::Escape,
+		[this]()
+		{
+			Operation = EEditingOperation::Select;
+			GetToolManager()->DisplayMessage(LOCTEXT("Select", "Click on a bone to select it."), EToolMessageLevel::UserNotification);
+		});
+
+	// register UnParent key
+	ActionSet.RegisterAction(this, GetActionId(), TEXT("UnparentBones"),
+		LOCTEXT("UnparentBones", "Unparent Bones"),
+		LOCTEXT("UnparentBonesDesc", "Unparent Bones"),
+		EModifierKey::Shift, EKeys::P,
+		[this]()
+		{
+			UnParentBones();
+		});
+		
+	// register Parent key
+	ActionSet.RegisterAction(this, GetActionId(), TEXT("ParentBones"),
+		LOCTEXT("ParentBones", "Parent Bones"),
+		LOCTEXT("ParentBonesDesc", "Parent Bones"),
+		EModifierKey::None, EKeys::B, // FIXME find another shortcut
+		[this]()
+		{
+			Operation = EEditingOperation::Parent;
+			GetToolManager()->DisplayMessage(LOCTEXT("Parent", "Click on a bone to be set as the new parent."), EToolMessageLevel::UserNotification);
+		});
+}
+
+void USkeletonEditingTool::CreateNewBone()
+{
+	if (Operation != EEditingOperation::Create)
+	{
+		return;
+	}
+
+	BeginChange();
+
+	const FName BoneName = SkeletonModifier.GetUniqueName(Properties->DefaultName);
+	const bool bBoneAdded = SkeletonModifier.AddBone(BoneName, CurrentBone, Properties->Transform);
+	if (bBoneAdded)
+	{
+		if (NeedsNotification())
+		{
+			GetNotifier().Notify({BoneName}, ESkeletalMeshNotifyType::BonesAdded);
+		}
+	
+		CurrentBone = BoneName;
+		Properties->Name = CurrentBone;
+
+		EndChange();
+		return;
+	}
+
+	CancelChange();
+}
+
+void USkeletonEditingTool::MirrorBones()
+{
+	TGuardValue OperationGuard(Operation, EEditingOperation::Mirror);
+	BeginChange();
+
+	const bool bBonesMirrored = SkeletonModifier.MirrorBones(GetSelectedBones(), MirroringProperties->Options);
+	if (bBonesMirrored)
+	{
+		EndChange();
+		return;		
+	}
+	
+	CancelChange();
+}
+
+void USkeletonEditingTool::RemoveBones()
+{
+	const TArray<FName> BonesToRemove = GetSelectedBones();
+	
+	TGuardValue OperationGuard(Operation, EEditingOperation::Remove);
+	BeginChange();
+
+	const bool bBonesRemoved = SkeletonModifier.RemoveBones(BonesToRemove, true);
+	if (bBonesRemoved)
+	{
+		// if (NeedsNotification())
+		// {
+		// 	GetNotifier().Notify(BonesToRemove, ESkeletalMeshNotifyType::BonesRemoved);
+		// }
+		
+		EndChange();
+		return;
+	}
+
+	CancelChange();
+}
+
+void USkeletonEditingTool::UnParentBones()
+{
+	static const TArray<FName> Dummy;
+
+	TGuardValue OperationGuard(Operation, EEditingOperation::Parent);
+	BeginChange();
+
+	const bool bBonesUnParented = SkeletonModifier.ParentBones(GetSelectedBones(), Dummy);
+	if (bBonesUnParented)
+	{
+		GetToolManager()->DisplayMessage(LOCTEXT("Unparent", "Selected bones have been unparented."), EToolMessageLevel::UserNotification);
+		
+		EndChange();
+		return;
+	}
+	
+	CancelChange();
+}
+
+void USkeletonEditingTool::ParentBones(const FName& InParentName)
+{
+	if (Operation != EEditingOperation::Parent)
+	{
+		return;
+	}
+
+	BeginChange();
+	const bool bBonesParented = SkeletonModifier.ParentBones(GetSelectedBones(), {InParentName});
+	if (bBonesParented)
+	{
+		Operation = EEditingOperation::Select;
+		EndChange();
+		return;
+	}
+
+	Operation = EEditingOperation::Select;
+	CancelChange();
+}
+
+void USkeletonEditingTool::MoveBones()
+{
+	const FReferenceSkeleton& RefSkeleton = SkeletonModifier.GetReferenceSkeleton();
+
+	const TArray<FName> Bones = GetSelectedBones();
+	const bool bHasValidBone = Bones.ContainsByPredicate([&](const FName& InBoneName)
+	{
+		return RefSkeleton.FindRawBoneIndex(InBoneName) > INDEX_NONE;
+	});
+
+	if (!bHasValidBone)
+	{
+		return;
+	}
+
+	TGuardValue OperationGuard(Operation, EEditingOperation::Transform);
+	BeginChange();
+
+	const bool bBonesMoved = SkeletonModifier.SetBoneTransform(Bones[0], Properties->Transform, Properties->bUpdateChildren);
+	if (bBonesMoved)
+	{
+		// if (NeedsNotification())
+		// {
+		// 	GetNotifier().Notify(Bones, ESkeletalMeshNotifyType::BonesMoved);
+		// }
+		
+		EndChange();
+		return;
+	}
+
+	CancelChange();
+}
+
+void USkeletonEditingTool::RenameBones()
+{
+	if (CurrentBone == Properties->Name || Properties->Name == NAME_None)
+	{
+		return;
+	}
+	
+	const FReferenceSkeleton& RefSkeleton = SkeletonModifier.GetReferenceSkeleton();
+	if (RefSkeleton.FindRawBoneIndex(CurrentBone) == INDEX_NONE)
+	{
+		return;
+	}
+
+	TGuardValue OperationGuard(Operation, EEditingOperation::Rename);
+	BeginChange();
+
+	const bool bBoneRenamed = SkeletonModifier.RenameBone(CurrentBone, Properties->Name);
+	if (bBoneRenamed)
+	{
+		CurrentBone = Properties->Name;
+
+		// if (NeedsNotification())
+		// {
+		// 	GetNotifier().Notify({CurrentBone}, ESkeletalMeshNotifyType::BonesRenamed);
+		// }
+		
+		EndChange();
+		return;
+	}
+
+	CancelChange();
+}
+
+void USkeletonEditingTool::OnClickPress(const FInputDeviceRay& PressPos)
+{
+	BeginChange();
+}
+
+void USkeletonEditingTool::OnClickDrag(const FInputDeviceRay& DragPos)
+{
+	FVector HitPoint;
+	if (ProjectionProperties->GetProjectionPoint(DragPos, HitPoint))
+	{
+		const FTransform& ParentGlobal = SkeletonModifier.GetTransform(ParentIndex, true);
+		Properties->Transform.SetLocation(ParentGlobal.InverseTransformPosition(HitPoint));
+
+		if (!ActiveChange)
+		{
+			TGuardValue OperationGuard(Operation, EEditingOperation::Transform);
+			BeginChange();
+		}
+
+		const bool bBoneMoved = SkeletonModifier.SetBoneTransform(CurrentBone, Properties->Transform, Properties->bUpdateChildren);
+		if (!bBoneMoved)
+		{
+			CancelChange();
+			return;
+		}
+
+		const bool bOrient = Operation == EEditingOperation::Create && OrientingProperties->bAutoOrient;
+		if (bOrient && ParentIndex != INDEX_NONE)
+		{
+			const FReferenceSkeleton& RefSkeleton = SkeletonModifier.GetReferenceSkeleton();
+			const FName ParentName = RefSkeleton.GetRawRefBoneInfo()[ParentIndex].Name;
+			SkeletonModifier.OrientBone(ParentName, OrientingProperties->Options);
+		}
+	}
+}
+
+void USkeletonEditingTool::OrientBones()
+{
+	const FReferenceSkeleton& RefSkeleton = SkeletonModifier.GetReferenceSkeleton();
+
+	const TArray<FName> Bones = GetSelectedBones();
+	const bool bHasValidBone = Bones.ContainsByPredicate([&](const FName& InBoneName)
+	{
+		return RefSkeleton.FindRawBoneIndex(InBoneName) > INDEX_NONE;
+	});
+
+	if (!bHasValidBone)
+	{
+		return;
+	}
+
+	TGuardValue OperationGuard(Operation, EEditingOperation::Transform);
+	BeginChange();
+	
+	const bool bBoneOriented = SkeletonModifier.OrientBones(Bones, OrientingProperties->Options);
+	if (bBoneOriented)
+	{
+		// if (NeedsNotification())
+		// {
+		// 	GetNotifier().Notify(Bones, ESkeletalMeshNotifyType::BonesMoved);
+		// }
+		
+		EndChange();
+		return;
+	}
+
+	CancelChange();
+}
+
+void USkeletonEditingTool::OnClickRelease(const FInputDeviceRay& ReleasePos)
+{
+	TGuardValue OperationGuard(Operation, EEditingOperation::Transform);
+	EndChange();
+}
+
+void USkeletonEditingTool::OnTerminateDragSequence()
+{}
+
+FInputRayHit USkeletonEditingTool::CanBeginClickDragSequence(const FInputDeviceRay& ClickPos)
+{
+	auto PickBone = [&]() -> int32
+	{
+		FViewport* FocusedViewport = GetToolManager()->GetContextQueriesAPI()->GetFocusedViewport();
+		if (HHitProxy* HitProxy = FocusedViewport->GetHitProxy(ClickPos.ScreenPosition.X, ClickPos.ScreenPosition.Y))
+		{
+			if (TOptional<FName> OptBoneName = GetBoneName(HitProxy))
+			{
+				const FReferenceSkeleton& ReferenceSkeleton = SkeletonModifier.GetReferenceSkeleton();
+				return ReferenceSkeleton.FindRawBoneIndex(*OptBoneName);
+			}
+		}
+		return INDEX_NONE;
+	};
+	
+	// pick bone in viewport
+	const int32 BoneIndex = PickBone();
+
+	// update parent
+	const FReferenceSkeleton& ReferenceSkeleton = SkeletonModifier.GetReferenceSkeleton();
+	ParentIndex = INDEX_NONE;
+
+	// update projection properties
+	const FVector GlobalPosition = SkeletonModifier.GetTransform(BoneIndex, true).GetTranslation();
+	ProjectionProperties->UpdatePlane(*ViewContext, GlobalPosition);
+
+	// if we picked a new bone
+	if (BoneIndex > INDEX_NONE)
+	{
+		// parent selection without changing the selection
+		if (Operation == EEditingOperation::Parent)
+		{	
+			ParentBones(ReferenceSkeleton.GetBoneName(BoneIndex));
+			return FInputRayHit();
+		}
+		
+		// otherwise, update current selection
+		CurrentBone = ReferenceSkeleton.GetBoneName(BoneIndex);
+
+		Properties->Name = CurrentBone;
+		Properties->Transform = ReferenceSkeleton.GetRefBonePose()[BoneIndex];
+		ParentIndex = ReferenceSkeleton.GetParentIndex(BoneIndex);
+		
+		if (NeedsNotification())
+		{
+			GetNotifier().Notify({CurrentBone}, ESkeletalMeshNotifyType::BonesSelected);
+		}
+	
+		return FInputRayHit(0.0);
+	}
+
+	// if we didn't pick anything
+	if (Operation == EEditingOperation::Select)
+	{
+		// unselect all
+		CurrentBone = NAME_None;
+		Properties->Name = CurrentBone; 
+		
+		if (NeedsNotification())
+		{
+			GetNotifier().Notify({CurrentBone}, ESkeletalMeshNotifyType::BonesSelected);
+		}
+		
+		return FInputRayHit();
+	}
+
+	// if we're in creation mode then create a new bone
+	if (Operation == EEditingOperation::Create)
+	{
+		FVector HitPoint;
+		if (ProjectionProperties->GetProjectionPoint(ClickPos, HitPoint))
+		{
+			// CurrentBone is gonna be the parent
+			ParentIndex = ReferenceSkeleton.FindRawBoneIndex(CurrentBone);
+			const FTransform& ParentGlobalTransform = SkeletonModifier.GetTransform(ParentIndex, true);
+
+			// Create the new bone under mouse
+			Properties->Transform.SetLocation(ParentGlobalTransform.InverseTransformPosition(HitPoint));
+			CreateNewBone();
+			
+			return FInputRayHit(0.0);
+		}
+	}
+	
+	return FInputRayHit();
+}
+
+void USkeletonEditingTool::HandleSkeletalMeshModified(const TArray<FName>& BoneNames, const ESkeletalMeshNotifyType InNotifyType)
+{
+	const FName BoneName = BoneNames.IsEmpty() ? NAME_None : BoneNames[0];
+	switch (InNotifyType)
+	{
+		case ESkeletalMeshNotifyType::BonesAdded:
+			CurrentBone = BoneName;
+			break;
+		case ESkeletalMeshNotifyType::BonesRemoved:
+			if (BoneNames.Contains(CurrentBone))
+			{
+				CurrentBone = NAME_None;
+			}
+			break;
+		case ESkeletalMeshNotifyType::BonesMoved:
+			CurrentBone = BoneName;
+			break;
+		case ESkeletalMeshNotifyType::BonesSelected:
+			CurrentBone = BoneName;
+			break;
+		case ESkeletalMeshNotifyType::BonesRenamed:
+			CurrentBone = BoneName;
+			break;
+		default:
+			break;
+	}
+	Properties->Name = CurrentBone;
+}
+
+void USkeletonEditingTool::Render(IToolsContextRenderAPI* RenderAPI)
+{
+	// FIXME many things could be caches here and updated lazilly
+	if (!Target)
+	{
+		return;
+	}
+
+	static const FLinearColor DefaultBoneColor(0.0f,0.0f,0.025f,1.0f);
+	static const FLinearColor SelectedBoneColor(0.2f,1.0f,0.2f,1.0f);
+	static const FLinearColor AffectedBoneColor(1.0f,1.0f,1.0f,1.0f);
+	static const FLinearColor ParentOfSelectedBoneColor(0.85f,0.45f,0.12f,1.0f);
+	static FSkelDebugDrawConfig DrawConfig;
+		DrawConfig.BoneDrawMode = EBoneDrawMode::Type::All;
+		DrawConfig.BoneDrawSize = 1.f;
+		DrawConfig.bAddHitProxy = true;
+		DrawConfig.bForceDraw = false;
+		DrawConfig.DefaultBoneColor = DefaultBoneColor;
+		DrawConfig.AffectedBoneColor = AffectedBoneColor;
+		DrawConfig.SelectedBoneColor = SelectedBoneColor;
+		DrawConfig.ParentOfSelectedBoneColor = ParentOfSelectedBoneColor;
+		DrawConfig.AxisConfig.Thickness = Properties->AxisThickness;
+		DrawConfig.AxisConfig.Length = Properties->AxisLength;
+	
+	FPrimitiveDrawInterface* PDI = RenderAPI->GetPrimitiveDrawInterface();
+
+	const IPrimitiveComponentBackedTarget* TargetComponent = Cast<IPrimitiveComponentBackedTarget>(Target);
+	const FTransform ComponentTransform = TargetComponent->GetWorldTransform();
+	
+	const FReferenceSkeleton& RefSkeleton = SkeletonModifier.GetReferenceSkeleton();
+
+	const int32 NumBones = RefSkeleton.GetRawBoneNum();
+	TArray<TRefCountPtr<HHitProxy>> HitProxies; HitProxies.Reserve(NumBones);
+	TArray<FBoneIndexType> RequiredBones; RequiredBones.AddUninitialized(NumBones);
+	TArray<FTransform> WorldTransforms; WorldTransforms.AddUninitialized(NumBones);
+	TArray<FLinearColor> BoneColors; BoneColors.AddUninitialized(NumBones);
+	
+	for (int32 Index = 0; Index < NumBones; ++Index)
+	{
+		const FTransform& BoneTransform = SkeletonModifier.GetTransform(Index, true);
+		WorldTransforms[Index] = BoneTransform;
+		RequiredBones[Index] = Index;
+		BoneColors[Index] = DefaultBoneColor;
+		HitProxies.Add(new HBoneHitProxy(Index, RefSkeleton.GetBoneName(Index)));
+	}
+
+	// FIXME cache this
+	TArray<int32> SelectedBones;
+	for (const FName& BoneName: GetSelectedBones())
+	{
+		int32 SelectedIndex = RefSkeleton.FindRawBoneIndex(CurrentBone);
+		if (SelectedIndex > INDEX_NONE)
+		{
+			SelectedBones.Add(SelectedIndex);
+		}
+	}
+
+	SkeletalDebugRendering::DrawBones(
+		PDI,
+		ComponentTransform.GetLocation(),
+		RequiredBones,
+		RefSkeleton,
+		WorldTransforms,
+		SelectedBones,
+		BoneColors,
+		HitProxies,
+		DrawConfig
+	);
+}
+
+FBox USkeletonEditingTool::GetWorldSpaceFocusBox()
+{
+	const TArray<FName> Selection = GetSelectedBones();
+	if (!Selection.IsEmpty())
+	{
+		FBox Box(EForceInit::ForceInit);
+		TSet<int32> AllChildren;
+
+		const FReferenceSkeleton& RefSkeleton = SkeletonModifier.GetReferenceSkeleton();
+		
+		for (const FName& BoneName: Selection)
+		{
+			const int32 BoneIndex = RefSkeleton.FindRawBoneIndex(BoneName);
+			Box += SkeletonModifier.GetTransform(BoneIndex, true).GetTranslation();
+
+			// get direct children
+			TArray<int32> Children;
+			RefSkeleton.GetDirectChildBones(BoneIndex, Children);
+			AllChildren.Append(Children);
+		}
+
+		for (const int32 ChildIndex: AllChildren)
+		{
+			Box += SkeletonModifier.GetTransform(ChildIndex, true).GetTranslation();
+		}
+		
+		return Box;
+	}
+
+	if (PreviewMesh && PreviewMesh->GetActor())
+	{
+		return PreviewMesh->GetActor()->GetComponentsBoundingBox();
+	}
+
+	return USingleSelectionTool::GetWorldSpaceFocusBox();
+}
+
+TArray<FName> USkeletonEditingTool::GetSelectedBones() const
+{
+	if (Binding.IsValid())
+	{
+		const TArray<FName> Selection = Binding.Pin()->GetSelectedBones();
+		if (!Selection.IsEmpty())
+		{
+			return Selection;
+		}
+	}
+
+	if (CurrentBone != NAME_None)
+	{
+		return {CurrentBone};
+	}
+
+	const TArray<FName> Dummy;
+	return Dummy;
+}
+
+void USkeletonEditingTool::BeginChange()
+{
+	if (Operation == EEditingOperation::Select)
+	{
+		return;
+	}
+	
+	ensure( ActiveChange == nullptr );
+	ActiveChange = MakeUnique<SkeletonEditingTool::FRefSkeletonChange>(this); 
+}
+
+void USkeletonEditingTool::EndChange()
+{
+	if (!ActiveChange.IsValid())
+	{
+		return;
+	}
+
+	if (Operation == EEditingOperation::Select)
+	{
+		return CancelChange();
+	}
+	
+	ActiveChange->StoreSkeleton(this);
+
+	static const UEnum* OperationEnum = StaticEnum<EEditingOperation>();
+	const FName OperationName = OperationEnum->GetNameByValue(static_cast<int64>(Operation));
+	const FText TransactionDesc = FText::Format(LOCTEXT("RefSkeletonChanged", "Skeleton Edit - {0}"), FText::FromName(OperationName));
+	
+	UInteractiveToolManager* ToolManager = GetToolManager();
+	ToolManager->BeginUndoTransaction(TransactionDesc);
+	ToolManager->EmitObjectChange(this, MoveTemp(ActiveChange), TransactionDesc);
+	ToolManager->EndUndoTransaction();
+}
+
+void USkeletonEditingTool::CancelChange()
+{
+	ActiveChange.Reset();
+}
+
+void USkeletonEditingProperties::Initialize(USkeletonEditingTool* ParentToolIn)
+{
+	ParentTool = ParentToolIn;
+	Name = ParentTool->CurrentBone;
+}
+
+#if WITH_EDITOR
+void USkeletonEditingProperties::PostEditChangeProperty(FPropertyChangedEvent& PropertyChangedEvent)
+{
+	Super::PostEditChangeProperty(PropertyChangedEvent);
+	
+	if (PropertyChangedEvent.ChangeType == EPropertyChangeType::ValueSet)
+	{
+		const FName PropertyName = PropertyChangedEvent.GetMemberPropertyName();
+		if (PropertyName == GET_MEMBER_NAME_CHECKED(USkeletonEditingProperties, Name))
+		{
+			ParentTool->RenameBones();
+		}
+
+		if (PropertyName == GET_MEMBER_NAME_CHECKED(USkeletonEditingProperties, Transform))
+		{
+			ParentTool->MoveBones();
+		}
+	}
+}
+#endif
+
+void UMirroringProperties::Initialize(USkeletonEditingTool* ParentToolIn)
+{
+	ParentTool = ParentToolIn;
+}
+
+void UMirroringProperties::MirrorBones()
+{
+	ParentTool->MirrorBones();
+}
+
+void UOrientingProperties::Initialize(USkeletonEditingTool* ParentToolIn)
+{
+	ParentTool = ParentToolIn;
+}
+
+void UOrientingProperties::OrientBones()
+{
+	ParentTool->OrientBones();
+}
+
+#if WITH_EDITOR
+
+void UOrientingProperties::PostEditChangeProperty(FPropertyChangedEvent &PropertyChangedEvent)
+{
+	Super::PostEditChangeProperty(PropertyChangedEvent);
+
+	if (PropertyChangedEvent.ChangeType == EPropertyChangeType::ValueSet)
+	{
+		auto CheckAxis = [&](const TEnumAsByte<EAxis::Type>& InRef, TEnumAsByte<EAxis::Type>& OutOther)
+		{
+			if (OutOther != InRef)
+			{
+				return;
+			}
+
+			switch (InRef)
+			{
+			case EAxis::X:
+				OutOther = EAxis::Y;
+				break;
+			case EAxis::Y:
+				OutOther = EAxis::Z;
+				break;
+			case EAxis::Z:
+				OutOther = EAxis::X;
+				break;
+			default:
+				break;
+			}
+		};
+		
+		const FName PropertyName = PropertyChangedEvent.GetPropertyName();
+		if (PropertyName == GET_MEMBER_NAME_CHECKED(FOrientOptions, Primary))
+		{
+			if (Options.Primary == EAxis::None)
+			{
+				Options.Primary = EAxis::X;
+				Options.Secondary = EAxis::Y;
+				return;
+			}
+			CheckAxis(Options.Primary, Options.Secondary);
+			return;
+		}
+		
+		if (PropertyName == GET_MEMBER_NAME_CHECKED(FOrientOptions, Secondary))
+		{
+			CheckAxis(Options.Secondary, Options.Primary);
+			return;
+		}
+	}
+}
+
+#endif
+
+void UProjectionProperties::Initialize(TObjectPtr<UPreviewMesh> InPreviewMesh) 
+{
+	PreviewMesh = InPreviewMesh;
+}
+
+void UProjectionProperties::UpdatePlane(const UGizmoViewContext& InViewContext, const FVector& InOrigin)
+{
+	PlaneNormal = -InViewContext.GetViewDirection();
+	PlaneOrigin = InOrigin;
+}
+
+bool UProjectionProperties::GetProjectionPoint(const FInputDeviceRay& InRay, FVector& OutHitPoint) const
+{
+	const FRay& WorldRay = InRay.WorldRay;
+
+	if (PreviewMesh.IsValid())
+	{
+		if (ProjectionType == EProjectionType::OnMesh)
+		{
+			FHitResult Hit;
+			if (PreviewMesh->FindRayIntersection(WorldRay, Hit))
+			{
+				OutHitPoint = WorldRay.Origin + WorldRay.Direction * Hit.Distance;
+				return true;
+			}
+		}
+
+		if (ProjectionType == EProjectionType::WithinMesh)
+		{
+			using namespace UE::Geometry;
+
+			if (FDynamicMeshAABBTree3* MeshAABBTree = PreviewMesh->GetSpatial())
+			{
+				using HitResult = MeshIntersection::FHitIntersectionResult;
+				TArray<HitResult> Hits;
+				
+				if (MeshAABBTree->FindAllHitTriangles(WorldRay, Hits))
+				{
+					if (Hits.Num() == 1)
+					{
+						OutHitPoint = WorldRay.Origin + WorldRay.Direction * Hits[0].Distance;
+						return true;						
+					}
+
+					// const double AverageDistance = Algo::Accumulate(Distances, 0.0) / static_cast<double>(Distances.Num());
+					// const int32 Index0 = Distances.IndexOfByPredicate([AverageDistance](double Distance) {return Distance <= AverageDistance;} );
+					// const int32 Index1 = Distances.IndexOfByPredicate([AverageDistance](double Distance) {return Distance >= AverageDistance;} );
+
+					static constexpr int32 Index0 = 0;
+					static constexpr int32 Index1 = 1;
+
+					const double d0 = Hits[Index0].Distance;
+					const double d1 = Hits[Index1].Distance;
+					OutHitPoint = WorldRay.Origin + WorldRay.Direction * ((d0+d1)*0.5);
+					return true;
+				}
+			}
+
+			FHitResult Hit;
+			if (PreviewMesh->FindRayIntersection(WorldRay, Hit))
+			{
+				OutHitPoint = WorldRay.Origin + WorldRay.Direction * Hit.Distance;
+				return true;
+			}
+		}
+	}
+	
+	// if ray is parallel to plane, nothing has been hit
+	if (FMath::IsNearlyZero(FVector::DotProduct(PlaneNormal, WorldRay.Direction)))
+	{
+		return false;
+	}
+
+	const FPlane Plane(PlaneOrigin, PlaneNormal);
+	const float HitDepth = FMath::RayPlaneIntersectionParam(WorldRay.Origin, WorldRay.Direction, Plane);
+	if (HitDepth < 0)
+	{
+		return false;
+	}
+
+	OutHitPoint = WorldRay.Origin + WorldRay.Direction * HitDepth;
+	
+	return true;
+}
+
+#undef LOCTEXT_NAMESPACE
\ No newline at end of file
diff --git a/Engine/Plugins/Experimental/MeshModelingToolsetExp/Source/MeshModelingToolsExp/Public/SkeletalMesh/SkeletonEditingTool.h b/Engine/Plugins/Experimental/MeshModelingToolsetExp/Source/MeshModelingToolsExp/Public/SkeletalMesh/SkeletonEditingTool.h
new file mode 100644
--- /dev/null
+++ b/Engine/Plugins/Experimental/MeshModelingToolsetExp/Source/MeshModelingToolsExp/Public/SkeletalMesh/SkeletonEditingTool.h
@@ -0,0 +1,310 @@
+﻿// Copyright Epic Games, Inc. All Rights Reserved.
+
+#pragma once
+
+#include "DynamicMeshBrushTool.h"
+#include "BaseTools/MeshSurfacePointMeshEditingTool.h"
+
+#include "SkeletalMeshNotifier.h"
+#include "SkeletonModifier.h"
+#include "SkeletalMesh/SkeletalMeshEditionInterface.h"
+#include "Engine/World.h"
+
+#include "SkeletonEditingTool.generated.h"
+
+class USingleClickInputBehavior;
+class UGizmoViewContext;
+class UTransformGizmo;
+class USkeletonEditingTool;
+
+namespace SkeletonEditingTool
+{
+
+/**
+ * A wrapper change class that stores a reference skeleton and the bones' indexes trackers to be used for undo/redo.
+ */
+class FRefSkeletonChange : public FToolCommandChange
+{
+public:
+	FRefSkeletonChange(const USkeletonEditingTool* InTool);
+
+	virtual FString ToString() const override
+	{
+		return FString(TEXT("Edit Skeleton"));
+	}
+
+	virtual void Apply(UObject* Object) override;
+	virtual void Revert(UObject* Object) override;
+
+	void StoreSkeleton(const USkeletonEditingTool* InTool);
+
+private:
+	FReferenceSkeleton PreChangeSkeleton;
+	TArray<int32> PreBoneTracker;
+	FReferenceSkeleton PostChangeSkeleton;
+	TArray<int32> PostBoneTracker;
+};
+	
+}
+
+/**
+ * USkeletonEditingToolBuilder
+ */
+
+UCLASS()
+class MESHMODELINGTOOLSEXP_API USkeletonEditingToolBuilder : public UInteractiveToolWithToolTargetsBuilder
+{
+	GENERATED_BODY()
+
+public:
+	// UInteractiveToolBuilder overrides
+	virtual bool CanBuildTool(const FToolBuilderState& SceneState) const override;
+	virtual UInteractiveTool* BuildTool(const FToolBuilderState& SceneState) const override;
+	
+protected:
+	// UInteractiveToolWithToolTargetsBuilder overrides
+	virtual const FToolTargetTypeRequirements& GetTargetRequirements() const override;
+};
+
+/**
+ * EEditingOperation represents the current tool's operation 
+ */
+
+UENUM()
+enum class EEditingOperation : uint8
+{
+	Select,		// Selecting bones in the viewport.
+	Create,		// Creating bones in the viewport.
+	Remove,		// Removing current selection.
+	Transform,	// Transforming bones in the viewport or thru the details panel.
+	Parent,		// Parenting bones in the viewport.
+	Rename,		// Renaming bones thru the details panel.
+	Mirror		// Mirroring bones thru the details panel.
+};
+
+/**
+ * USkeletonEditingTool is a tool to edit a the ReferenceSkeleton of a SkeletalMesh (target)
+ * Changed are actually commit to the SkeletalMesh and it's mesh description on Accept.
+ */
+
+UCLASS()
+class MESHMODELINGTOOLSEXP_API USkeletonEditingTool :
+	public USingleSelectionTool,
+	public IClickDragBehaviorTarget,
+	public ISkeletalMeshEditionInterface
+{
+	GENERATED_BODY()
+
+public:
+
+	void Init(const FToolBuilderState& InSceneState);
+	
+	// UInteractiveTool overrides
+	virtual void Setup() override;
+	virtual void Shutdown(EToolShutdownType ShutdownType) override;
+
+	virtual void Render(IToolsContextRenderAPI* RenderAPI) override;
+	
+	virtual bool HasCancel() const override { return true; }
+	virtual bool HasAccept() const override { return true; }
+	virtual bool CanAccept() const override { return true; }
+
+	// ICLickDragBehaviorTarget overrides
+	virtual FInputRayHit CanBeginClickDragSequence(const FInputDeviceRay& PressPos) override;
+	
+	virtual void OnClickPress(const FInputDeviceRay& PressPos) override;
+	virtual void OnClickDrag(const FInputDeviceRay& DragPos) override;
+	virtual void OnClickRelease(const FInputDeviceRay& ReleasePos) override;
+	virtual void OnTerminateDragSequence() override;
+
+	// IInteractiveToolCameraFocusAPI overrides
+	virtual FBox GetWorldSpaceFocusBox() override;
+
+	// Modifier functions
+	void MirrorBones();
+	void RenameBones();
+	void MoveBones();
+	void OrientBones();
+	
+	int32 ParentIndex = INDEX_NONE;
+	FName CurrentBone = NAME_None;
+
+	FSkeletonModifier SkeletonModifier;
+	
+protected:
+
+	// UInteractiveTool overrides
+	virtual void RegisterActions(FInteractiveToolActionSet& ActionSet) override;
+
+	// ISkeletalMeshEditionInterface overrides
+	virtual void HandleSkeletalMeshModified(const TArray<FName>& BoneNames, const ESkeletalMeshNotifyType InNotifyType) override;
+
+	// Modifier functions
+	void CreateNewBone();
+	void RemoveBones();
+	void UnParentBones();
+	void ParentBones(const FName& InParentName);
+
+	TArray<FName> GetSelectedBones() const;
+
+	UPROPERTY()
+	TObjectPtr<USkeletonEditingProperties> Properties;
+
+	UPROPERTY()
+	TObjectPtr<UProjectionProperties> ProjectionProperties;
+	
+	UPROPERTY()
+	TObjectPtr<UMirroringProperties> MirroringProperties;
+
+	UPROPERTY()
+	TObjectPtr<UOrientingProperties> OrientingProperties;
+	
+	UPROPERTY()
+	TObjectPtr<USingleClickInputBehavior> LeftClickBehavior = nullptr;
+
+	UPROPERTY()
+	TObjectPtr<UPreviewMesh> PreviewMesh = nullptr;
+
+	UPROPERTY()
+	TWeakObjectPtr<UWorld> TargetWorld = nullptr;
+
+	UPROPERTY()
+	TWeakObjectPtr<UGizmoViewContext> ViewContext = nullptr;
+
+	UPROPERTY()
+	EEditingOperation Operation = EEditingOperation::Select;
+
+	// ref skeleton transactions
+	void BeginChange();
+	void EndChange();
+	void CancelChange();
+	TUniquePtr<SkeletonEditingTool::FRefSkeletonChange> ActiveChange;
+};
+
+/**
+ * USkeletonEditingProperties
+ */
+
+UCLASS()
+class MESHMODELINGTOOLSEXP_API USkeletonEditingProperties : public UInteractiveToolPropertySet
+{
+	GENERATED_BODY()
+public:
+	void Initialize(USkeletonEditingTool* ParentToolIn);
+
+	UPROPERTY(EditAnywhere, Category = "Details")
+	FName Name;
+
+	UPROPERTY(EditAnywhere, Category = "Details")
+	FTransform Transform;
+	
+	UPROPERTY(EditAnywhere, Category = "Move")
+	bool bUpdateChildren = false;
+
+	UPROPERTY(EditAnywhere, Category = "Create")
+	FName DefaultName = "joint";
+
+	UPROPERTY(EditAnywhere, Category = "Viewport Axis Settings",  meta = (ClampMin = "0.0", UIMin = "0.0"))
+	float AxisLength = 1.f;
+
+	UPROPERTY(EditAnywhere, Category = "Viewport Axis Settings",  meta = (ClampMin = "0.0", UIMin = "0.0"))
+	float AxisThickness = 0.f;
+
+#if WITH_EDITOR
+	virtual void PostEditChangeProperty(FPropertyChangedEvent &PropertyChangedEvent) override;
+#endif
+	
+private:
+	TWeakObjectPtr<USkeletonEditingTool> ParentTool = nullptr;
+};
+
+/**
+ * EProjectionType
+ */
+
+UENUM()
+enum class EProjectionType : uint8
+{
+	CameraPlane,	// The camera plane is used as the projection plane 
+	OnMesh,			// The mesh surface is used for projection
+	WithinMesh,		// The mesh volume is used for projection
+};
+
+/**
+ * UProjectionProperties
+ */
+
+UCLASS()
+class MESHMODELINGTOOLSEXP_API UProjectionProperties: public UInteractiveToolPropertySet
+{
+	GENERATED_BODY()
+
+public:
+	void Initialize(TObjectPtr<UPreviewMesh> PreviewMesh);
+
+	void UpdatePlane(const UGizmoViewContext& InViewContext, const FVector& InOrigin);
+	bool GetProjectionPoint(const FInputDeviceRay& InRay, FVector& OutHitPoint) const;
+	
+	UPROPERTY(EditAnywhere, Category = "Project")
+	EProjectionType ProjectionType = EProjectionType::WithinMesh;
+	
+private:
+	TWeakObjectPtr<UPreviewMesh> PreviewMesh = nullptr;
+	
+	UPROPERTY()
+	FVector PlaneOrigin = FVector::ZeroVector;
+	
+	UPROPERTY()
+	FVector PlaneNormal =  FVector::ZAxisVector;
+};
+
+/**
+ * UMirroringProperties
+ */
+
+UCLASS()
+class MESHMODELINGTOOLSEXP_API UMirroringProperties: public UInteractiveToolPropertySet
+{
+	GENERATED_BODY()
+
+public:
+	void Initialize(USkeletonEditingTool* ParentToolIn);
+
+	UFUNCTION(CallInEditor, Category = "Mirror")
+	void MirrorBones();
+
+	UPROPERTY(EditAnywhere, Category = "Mirror")
+	FMirrorOptions Options;
+
+private:
+	TWeakObjectPtr<USkeletonEditingTool> ParentTool = nullptr;
+};
+
+/**
+ * UOrientingProperties
+ */
+
+UCLASS()
+class MESHMODELINGTOOLSEXP_API UOrientingProperties: public UInteractiveToolPropertySet
+{
+	GENERATED_BODY()
+
+public:
+	void Initialize(USkeletonEditingTool* ParentToolIn);
+
+	UFUNCTION(CallInEditor, Category = "Orient")
+	void OrientBones();
+	
+	UPROPERTY(EditAnywhere, Category = "Orient")
+	FOrientOptions Options;
+
+	UPROPERTY(EditAnywhere, Category = "Orient")
+	bool bAutoOrient = false;
+
+#if WITH_EDITOR
+	virtual void PostEditChangeProperty(FPropertyChangedEvent &PropertyChangedEvent) override;
+#endif
+	
+private:
+	TWeakObjectPtr<USkeletonEditingTool> ParentTool = nullptr;
+};
\ No newline at end of file
diff --git a/Engine/Source/Developer/SkeletalMeshUtilitiesCommon/Private/SkeletonModifier.cpp b/Engine/Source/Developer/SkeletalMeshUtilitiesCommon/Private/SkeletonModifier.cpp
--- a/Engine/Source/Developer/SkeletalMeshUtilitiesCommon/Private/SkeletonModifier.cpp
+++ b/Engine/Source/Developer/SkeletalMeshUtilitiesCommon/Private/SkeletonModifier.cpp
@@ -52,11 +52,63 @@ FTransform FMirrorOptions::MirrorTransform(const FTransform& InTransform) const
 
 FVector FMirrorOptions::MirrorVector(const FVector& InVector) const
 {
-	FVector Axis = FVector::ZeroVector;
-	Axis.SetComponentForAxis(MirrorAxis, 1.f);
+	FVector Axis(ForceInitToZero); Axis.SetComponentForAxis(MirrorAxis, 1.f);
 	return InVector.MirrorByVector(Axis);
 }
 
+FTransform FOrientOptions::OrientTransform(const FVector& InPrimaryTarget, const FTransform& InTransform) const
+{
+	if (Primary == EAxis::None || InPrimaryTarget.IsNearlyZero())
+	{
+		return InTransform;
+	}
+
+	FTransform Transform = InTransform;
+
+	FVector PrimaryAxis(ForceInitToZero); PrimaryAxis.SetComponentForAxis(Primary, 1.f);
+	const FVector PrimaryTarget = InPrimaryTarget.GetSafeNormal();
+	
+	// orient primary axis towards InPrimaryTarget
+	{
+		const FVector Axis = Transform.TransformVectorNoScale(PrimaryAxis).GetSafeNormal();
+		const FQuat Rotation = FQuat::FindBetweenNormals(Axis, PrimaryTarget);
+		const FQuat NewRotation = (Rotation * Transform.GetRotation()).GetNormalized();
+		Transform.SetRotation(NewRotation);
+	}
+
+	if (Secondary == EAxis::None || SecondaryTarget.IsNearlyZero())
+	{ // no need to use secondary axis
+		return Transform;
+	}
+
+	FVector Target = SecondaryTarget.GetSafeNormal();
+	if (FMath::IsNearlyEqual(FMath::Abs(FVector::DotProduct(PrimaryTarget, Target)), 1.0f))
+	{ // both targets are parallel
+		return Transform;
+	}
+
+	// orient secondary axis towards SecondaryDirection
+	{
+		FVector SecondaryAxis(ForceInitToZero); SecondaryAxis.SetComponentForAxis(Secondary, 1.f);
+		
+		{
+			const FVector Axis = Transform.TransformVectorNoScale(PrimaryAxis).GetSafeNormal();
+			Target = Target - FVector::DotProduct(Target, Axis) * Axis;
+		}
+
+		if (!Target.IsNearlyZero())
+		{
+			Target = Target.GetSafeNormal();
+			const FVector Axis = Transform.TransformVectorNoScale(SecondaryAxis).GetSafeNormal();
+			const FQuat Rotation = FQuat::FindBetweenNormals(Axis, Target);
+			const FQuat NewRotation = (Rotation * Transform.GetRotation()).GetNormalized();
+			Transform.SetRotation(NewRotation);
+		}
+	}
+	
+	return Transform;
+}
+
 void FSkeletonModifier::ExternalUpdate(const FReferenceSkeleton& InRefSkeleton, const TArray<int32>& InIndexTracker)
 {
 	if (!ReferenceSkeleton)
@@ -452,22 +504,22 @@ bool FSkeletonModifier::MirrorBones(const TArray<FName>& InBonesName, const FMir
 	TArray<FName> MirroredNames; GetMirroredNames(BonesToMirror, InOptions, MirroredNames);
 
 	// add bones first if they are missing
-	TArray<int32> MirroredIndex; GetMirroredBones(BonesToMirror, MirroredNames, MirroredIndex);
-	if (MirroredIndex.Num() != NumBonesToMirror)
+	TArray<int32> MirroredBones; GetMirroredBones(BonesToMirror, MirroredNames, MirroredBones);
+	if (MirroredBones.Num() != NumBonesToMirror)
 	{
 		UE_LOG(LogAnimation, Error, TEXT("Skeleton Modifier - Mirror: Couldn't find mirrored bones."));
 		return false;
 	}
 
 	// compute mirrored transforms
-	TArray<FTransform> MirroredTransforms; GetMirroredTransforms(BonesToMirror, InOptions, MirroredTransforms);
+	TArray<FTransform> MirroredTransforms; GetMirroredTransforms(BonesToMirror, MirroredBones, InOptions, MirroredTransforms);
 
 	// update reference skeleton
 	{
 		FReferenceSkeletonModifier Modifier(*ReferenceSkeleton, nullptr);
 		for (int32 Index = 0; Index < NumBonesToMirror; Index++)
 		{
-			Modifier.UpdateRefPoseTransform(MirroredIndex[Index], MirroredTransforms[Index]);
+			Modifier.UpdateRefPoseTransform(MirroredBones[Index], MirroredTransforms[Index]);
 		}
 	}
 
@@ -625,7 +677,8 @@ void FSkeletonModifier::GetMirroredBones(
 }
 
 void FSkeletonModifier::GetMirroredTransforms(
-	const TArray<int32>& InBonesToMirror, const FMirrorOptions& InOptions, TArray<FTransform>& OutMirroredTransforms) const
+	const TArray<int32>& InBonesToMirror, const TArray<int32>& InMirroredBones,
+	const FMirrorOptions& InOptions, TArray<FTransform>& OutMirroredTransforms) const
 {
 	OutMirroredTransforms.Reset();
 	
@@ -689,8 +742,9 @@ void FSkeletonModifier::GetMirroredTransforms(
 		const int32 RefBoneIndex = InBonesToMirror[Index];		
 		const int32 RefParentIndex = BoneInfos[RefBoneIndex].ParentIndex;
 		const int32 ParentMirroredIndex = InBonesToMirror.IndexOfByKey(RefParentIndex);
+		const int32 ParentIndex = BoneInfos[InMirroredBones[Index]].ParentIndex;
 		const FTransform& ParentGlobal = ParentMirroredIndex != INDEX_NONE ? MirroredGlobal[ParentMirroredIndex] :
-							TransformComposer->GetGlobalTransform(RefParentIndex);
+							TransformComposer->GetGlobalTransform(ParentIndex);
 		OutMirroredTransforms.Add(MirroredGlobal[Index].GetRelativeTransform(ParentGlobal));
 	}
 }
@@ -699,26 +753,26 @@ void FSkeletonModifier::GetMirroredTransforms(
 // atm, we update the bone's local ref transform so children's global transforms are changed (we just need to cache the
 // global transforms then restore them back)
 // orienting the bone for example should change the children's global transform
-bool FSkeletonModifier::MoveBone( const FName InBoneToMove, const FTransform& InNewTransform, const bool bMoveChildren)
+bool FSkeletonModifier::SetBoneTransform( const FName InBoneName, const FTransform& InNewTransform, const bool bMoveChildren)
 {
-	if (InBoneToMove == NAME_None)
+	if (InBoneName == NAME_None)
 	{
 		UE_LOG(LogAnimation, Error, TEXT("Skeleton Modifier - Move: Cannot move bone with no name."));
 		return false;
 	}
 	
-	return MoveBones({InBoneToMove}, {InNewTransform}, bMoveChildren);
+	return SetBonesTransforms({InBoneName}, {InNewTransform}, bMoveChildren);
 }
 
-bool FSkeletonModifier::MoveBones(
-	const TArray<FName>& InBonesToMove, const TArray<FTransform>& InNewTransforms, const bool bMoveChildren)
+bool FSkeletonModifier::SetBonesTransforms(
+	const TArray<FName>& InBoneNames, const TArray<FTransform>& InNewTransforms, const bool bMoveChildren)
 {
 	if (!IsReferenceSkeletonValid())
 	{
 		return false;
 	}
 	
-	const int32 NumBonesToMove = InBonesToMove.Num();
+	const int32 NumBonesToMove = InBoneNames.Num();
 	if (NumBonesToMove == 0 || NumBonesToMove != InNewTransforms.Num())
 	{
 		UE_LOG(LogAnimation, Error, TEXT("Skeleton Modifier - Move: Discrepancy between bones and transforms (%d / %d)."), NumBonesToMove, InNewTransforms.Num());
@@ -731,7 +785,7 @@ bool FSkeletonModifier::MoveBones(
 
 	for (int32 Index = 0; Index < NumBonesToMove; Index++)
 	{
-		const int32 BoneIndex = ReferenceSkeleton->FindRawBoneIndex(InBonesToMove[Index]);
+		const int32 BoneIndex = ReferenceSkeleton->FindRawBoneIndex(InBoneNames[Index]);
 		if (BoneIndex != INDEX_NONE)
 		{
 			BoneIndices.Add(BoneIndex);
@@ -756,7 +810,13 @@ bool FSkeletonModifier::MoveBones(
 		{
 			TArray<int32> Children;
 			ReferenceSkeleton->GetDirectChildBones(BoneIndices[Index], Children);
-			ChildrenToFix.Append(Children);
+			for (int32 ChildIndex: Children)
+			{
+				if (!BoneIndices.Contains(ChildIndex))
+				{
+					ChildrenToFix.Add(ChildIndex);
+				}
+			}
 		}
 
 		// sort them from highest index to lowest
@@ -973,18 +1033,18 @@ bool FSkeletonModifier::ParentBones(const TArray<FName>& InBoneNames, const TArr
 	return true;
 }
 
-bool FSkeletonModifier::AlignBone(const FName InBoneName, const bool bAlignChildren)
+bool FSkeletonModifier::OrientBone(const FName InBoneName, const FOrientOptions& InOptions)
 {
 	if (InBoneName == NAME_None)
 	{
-		UE_LOG(LogAnimation, Error, TEXT("Skeleton Modifier - Align: Cannot align a bone with no name."));
+		UE_LOG(LogAnimation, Error, TEXT("Skeleton Modifier - Orient: Cannot orient a bone with no name."));
 		return false;
 	}
 	
-	return AlignBones({InBoneName}, bAlignChildren);
+	return OrientBones({InBoneName}, InOptions);
 }
 
-bool FSkeletonModifier::AlignBones(const TArray<FName>& InBoneNames, const bool bAlignChildren)
+bool FSkeletonModifier::OrientBones(const TArray<FName>& InBoneNames, const FOrientOptions& InOptions)
 {
 	if (!IsReferenceSkeletonValid())
 	{
@@ -993,73 +1053,121 @@ bool FSkeletonModifier::AlignBones(const TArray<FName>& InBoneNames, const bool
 	
 	if (InBoneNames.IsEmpty())
 	{
-		UE_LOG(LogAnimation, Error, TEXT("Skeleton Modifier - Align: No bone provided."));
+		UE_LOG(LogAnimation, Error, TEXT("Skeleton Modifier - Orient: No bone provided."));
 		return false;
 	}
 
-	auto GetAlignedTransform = [&](const FName& InBoneName) -> TOptional<FTransform>
+	// get bones to mirror
+	TArray<int32> BonesToOrient; GetBonesToOrient(InBoneNames, InOptions, BonesToOrient);
+	const int32 NumBonesToOrient = BonesToOrient.Num();
+	if (NumBonesToOrient == 0)
 	{
-		TOptional<FTransform> OptTransform;
+		UE_LOG(LogAnimation, Error, TEXT("Skeleton Modifier - Orient: None of the provided names has been found."));
+		return false;
+	}
+	
+	auto GetAlignedTransform = [&](const int32 BoneIndex) -> FTransform
+	{
+		const FTransform& BoneGlobal = TransformComposer->GetGlobalTransform(BoneIndex);
 		
-		const int32 BoneIndex = ReferenceSkeleton->FindRawBoneIndex(InBoneName);
 		const int32 ParentIndex = BoneIndex != INDEX_NONE ? ReferenceSkeleton->GetRawParentIndex(BoneIndex) : INDEX_NONE;
-		if (ParentIndex == INDEX_NONE)
-		{ // we can't align if there's no parent
-			return OptTransform;
-		}
 
 		TArray<int32> Children; ReferenceSkeleton->GetDirectChildBones(BoneIndex, Children);
-		if (Children.Num() != 1)
+		const int32 NumChildren = Children.Num();
+		if (NumChildren > 1)
 		{ // we can't align if there are more than one children
-			return OptTransform;
+			return BoneGlobal;
 		}
 
-		const FTransform& ParentGlobal = GetTransform(ParentIndex, true);
-		const FTransform& BoneGlobal = GetTransform(BoneIndex, true);
-		const FTransform& ChildGlobal = GetTransform(Children[0], true);
-
-		// X/Y axis
-		const FVector XAxis = (ChildGlobal.GetLocation() - BoneGlobal.GetLocation()).GetSafeNormal();
-		const FVector YAxis = (BoneGlobal.GetLocation() - ParentGlobal.GetLocation()).GetSafeNormal();
+		const FTransform& ParentGlobal = TransformComposer->GetGlobalTransform(ParentIndex);
+		FVector Direction = (BoneGlobal.GetLocation() - ParentGlobal.GetLocation()).GetSafeNormal();
 
-		if (XAxis.IsNearlyZero() || YAxis.IsNearlyZero())
+		if (NumChildren > 0)
 		{
-			return OptTransform;
+			const FTransform& ChildGlobal = TransformComposer->GetGlobalTransform(Children[0]);
+			Direction = (ChildGlobal.GetLocation() - BoneGlobal.GetLocation()).GetSafeNormal();
 		}
+		
+		if (Direction.IsNearlyZero())
+		{
+			return BoneGlobal;
+		}
+
+		return InOptions.OrientTransform(Direction, BoneGlobal);
+	};
 
-		// Z axis
-		const bool bComputePlane = (FMath::Abs(FVector::DotProduct(XAxis, YAxis)) - 1.0f) < KINDA_SMALL_NUMBER;
-		const FVector ZAxis = bComputePlane ? FVector::CrossProduct(XAxis, YAxis) : BoneGlobal.TransformVector(FVector::ZAxisVector);
+	const TArray<FMeshBoneInfo>& BoneInfos = ReferenceSkeleton->GetRawRefBoneInfo();
 	
-		FQuat Rotation = FRotationMatrix::MakeFromXZ(XAxis, ZAxis).ToQuat();
-		Rotation.Normalize();
+	TArray<FName> BonesToAlign; BonesToAlign.Reserve(NumBonesToOrient);
+	TArray<FTransform> OrientedGlobal; OrientedGlobal.Reserve(NumBonesToOrient);
+	for (int32 Index = 0; Index < NumBonesToOrient; ++Index)	
+	{
+		const int32 BoneIndex = BonesToOrient[Index];
+		BonesToAlign.Add(BoneInfos[BoneIndex].Name);
+		OrientedGlobal.Add(GetAlignedTransform(BoneIndex));
+	}
+
+	// switch back to local
+	TArray<FTransform> Transforms; Transforms.Reserve(NumBonesToOrient);
+	for (int32 Index = 0; Index < BonesToAlign.Num(); Index++)
+	{
+		const FName& BoneName = BonesToAlign[Index];
+		const int32 BoneIndex = ReferenceSkeleton->FindRawBoneIndex(BoneName);
+		const int32 ParentIndex = BoneInfos[BoneIndex].ParentIndex;
+		const int32 ParentOrientedIndex = ParentIndex != INDEX_NONE ? BonesToAlign.IndexOfByKey(BoneInfos[ParentIndex].Name) : INDEX_NONE;
+		const FTransform& ParentGlobal = ParentOrientedIndex != INDEX_NONE ? OrientedGlobal[ParentOrientedIndex] :
+					TransformComposer->GetGlobalTransform(ParentIndex);
+		Transforms.Add(OrientedGlobal[Index].GetRelativeTransform(ParentGlobal));
+	}
+
+	if (BonesToAlign.IsEmpty())
+	{
+		UE_LOG(LogAnimation, Error, TEXT("Skeleton Modifier - Orient: No bone to orient."));
+		return false;
+	}
 
-		const FTransform GlobalTransform(Rotation, BoneGlobal.GetLocation());
-		const FTransform LocalTransform = GlobalTransform.GetRelativeTransform(ParentGlobal);
+	static constexpr bool bMoveChildren = false;
+	return SetBonesTransforms(BonesToAlign, Transforms, bMoveChildren);
+}
 
-		return LocalTransform;
-	};
+void FSkeletonModifier::GetBonesToOrient(
+	const TArray<FName>& InBonesName, const FOrientOptions& InOptions, TArray<int32>& OutBonesToOrient) const
+{
+	OutBonesToOrient.Reset();
 	
-	TArray<FName> BonesToAlign; BonesToAlign.Reserve(InBoneNames.Num());
-	TArray<FTransform> Transforms; Transforms.Reserve(InBoneNames.Num());
-	for (int32 Index = 0; Index < InBoneNames.Num(); ++Index)	
+	TSet<int32> IndicesToOrient;
+	auto GetBonesToOrient = [&](const int32 BoneIndex, auto&& GetBonesToOrient2) -> void
 	{
-		TOptional<FTransform> OptTransform = GetAlignedTransform(InBoneNames[Index]);
-		if (OptTransform)
+		if (BoneIndex == INDEX_NONE)
 		{
-			BonesToAlign.Add(InBoneNames[Index]);
-			Transforms.Add(*OptTransform);
+			return;
 		}
+		
+		IndicesToOrient.Add(BoneIndex);
+
+		if (InOptions.bOrientChildren)
+		{
+			TArray<int32> Children; ReferenceSkeleton->GetDirectChildBones(BoneIndex, Children);
+			for (int32 ChildIndex: Children)
+			{
+				GetBonesToOrient2(ChildIndex, GetBonesToOrient2);
+			}
+		}
+	};
+	
+	for (const FName& BoneName: InBonesName)
+	{
+		GetBonesToOrient(ReferenceSkeleton->FindRawBoneIndex(BoneName), GetBonesToOrient);
 	}
 
-	if (BonesToAlign.IsEmpty())
+	const int32 NumBonesToOrient = IndicesToOrient.Num();
+	if (NumBonesToOrient == 0)
 	{
-		UE_LOG(LogAnimation, Error, TEXT("Skeleton Modifier - Align: No bone to align."));
-		return false;
+		return;
 	}
 
-	static constexpr bool bMoveChildren = false;
-	return MoveBones(BonesToAlign, Transforms, bMoveChildren);
+	IndicesToOrient.Sort([](const int32 Index0, const int32 Index1) {return Index0 < Index1;});
+	OutBonesToOrient = IndicesToOrient.Array();
 }
 
 void FSkeletonModifier::UpdateBoneTracker(const TArray<FMeshBoneInfo>& InOtherInfos)
diff --git a/Engine/Source/Developer/SkeletalMeshUtilitiesCommon/Public/SkeletalMeshNotifier.h b/Engine/Source/Developer/SkeletalMeshUtilitiesCommon/Public/SkeletalMeshNotifier.h
--- a/Engine/Source/Developer/SkeletalMeshUtilitiesCommon/Public/SkeletalMeshNotifier.h
+++ b/Engine/Source/Developer/SkeletalMeshUtilitiesCommon/Public/SkeletalMeshNotifier.h
@@ -14,7 +14,8 @@ enum class ESkeletalMeshNotifyType
 	BonesAdded,
 	BonesRemoved,
 	BonesMoved,
-	BonesSelected
+	BonesSelected,
+	BonesRenamed
 };
 
 // A delegate for monitoring to skeletal mesh global notifications.
diff --git a/Engine/Source/Developer/SkeletalMeshUtilitiesCommon/Public/SkeletonModifier.h b/Engine/Source/Developer/SkeletalMeshUtilitiesCommon/Public/SkeletonModifier.h
--- a/Engine/Source/Developer/SkeletalMeshUtilitiesCommon/Public/SkeletonModifier.h
+++ b/Engine/Source/Developer/SkeletalMeshUtilitiesCommon/Public/SkeletonModifier.h
@@ -16,6 +16,10 @@ struct FMeshDescription;
 struct FReferenceSkeleton;
 struct FTransformComposer;
 
+/**
+ * FMirrorOptions
+ */
+
 USTRUCT(BlueprintType)
 struct SKELETALMESHUTILITIESCOMMON_API FMirrorOptions
 {
@@ -40,6 +44,30 @@ struct SKELETALMESHUTILITIESCOMMON_API FMirrorOptions
 	FVector MirrorVector(const FVector& InVector) const;
 };
 
+/**
+ * FOrientOptions
+ */
+
+USTRUCT(BlueprintType)
+struct SKELETALMESHUTILITIESCOMMON_API FOrientOptions
+{
+	GENERATED_USTRUCT_BODY()
+
+	UPROPERTY(BlueprintReadWrite, EditAnywhere, Category = "Orient")
+	TEnumAsByte<EAxis::Type> Primary = EAxis::X;
+
+	UPROPERTY(BlueprintReadWrite, EditAnywhere, Category = "Orient")
+	TEnumAsByte<EAxis::Type> Secondary = EAxis::Y;
+
+	UPROPERTY(BlueprintReadWrite, EditAnywhere, Category = "Orient")
+	FVector SecondaryTarget = FVector::YAxisVector;
+
+	UPROPERTY(BlueprintReadWrite, EditAnywhere, Category = "Orient")
+	bool bOrientChildren = true;
+
+	FTransform OrientTransform(const FVector& InPrimaryTarget, const FTransform& InTransform) const;
+};
+
 /**
  * FSkeletalMeshSkeletonModifier
  */
@@ -68,16 +96,15 @@ public:
 	 */
 	bool MirrorBone(const FName InBoneName, const FMirrorOptions& InOptions = FMirrorOptions());
 	bool MirrorBones(const TArray<FName>& InBonesName, const FMirrorOptions& InOptions = FMirrorOptions());
-	bool MirrorBonesDep(const TArray<FName>& InBonesName);
 
-	/** Moves a bone at a new desired local transform
-	 *  @param InBoneToMove The new bone's name that needs to be moved.
+	/** Sets the bone the desired local transform
+	 *  @param InBoneName The new bone's name that needs to be moved.
 	 *  @param InNewTransform The new local transform in the bone's parent space.
 	 *  @param bMoveChildren Propagate new transform to children
 	 *  @return \c true if the operation succeeded, false otherwise. 
 	 */
-	bool MoveBone(const FName InBoneToMove, const FTransform& InNewTransform, const bool bMoveChildren);
-	bool MoveBones(const TArray<FName>& InBonesToMove, const TArray<FTransform>& InNewTransforms, const bool bMoveChildren);
+	bool SetBoneTransform(const FName InBoneName, const FTransform& InNewTransform, const bool bMoveChildren);
+	bool SetBonesTransforms(const TArray<FName>& InBoneNames, const TArray<FTransform>& InNewTransforms, const bool bMoveChildren);
 
 	/** Remove a bone in the skeleton hierarchy
 	 *  @param InBoneName The new bone's name.
@@ -105,11 +132,11 @@ public:
 
 	/** Align bones
 	 *  @param InBoneName The current bone's name.
-	 *  @param bAlignChildren
+	 *  @param InOptions The orienting options
 	 *  @return \c true if the operation succeeded, false otherwise. 
 	 */
-	bool AlignBone(const FName InBoneName, const bool bAlignChildren);
-	bool AlignBones(const TArray<FName>& InBoneNames, const bool bAlignChildren);
+	bool OrientBone(const FName InBoneName, const FOrientOptions& InOptions = FOrientOptions());
+	bool OrientBones(const TArray<FName>& InBoneNames, const FOrientOptions& InOptions = FOrientOptions());
 	
 	/**
 	 * Actually applies the skeleton modifications to the skeletal mesh.
@@ -133,7 +160,11 @@ private:
 	void GetBonesToMirror(const TArray<FName>& InBonesName, const FMirrorOptions& InOptions, TArray<int32>& OutBonesToMirror) const;
 	void GetMirroredNames(const TArray<int32>& InBonesToMirror, const FMirrorOptions& InOptions, TArray<FName>& OutBonesName) const;
 	void GetMirroredBones(const TArray<int32>& InBonesToMirror, const TArray<FName>& InMirroredNames, TArray<int32>& OutMirroredBones);
-	void GetMirroredTransforms(const TArray<int32>& InBonesToMirror, const FMirrorOptions& InOptions, TArray<FTransform>& OutMirroredTransforms) const;
+	void GetMirroredTransforms(	const TArray<int32>& InBonesToMirror, const TArray<int32>& InMirroredBones,
+								const FMirrorOptions& InOptions, TArray<FTransform>& OutMirroredTransforms) const;
+
+	// orient function
+	void GetBonesToOrient(const TArray<FName>& InBonesName, const FOrientOptions& InOptions, TArray<int32>& OutBonesToOrient) const;
 	
 	TObjectPtr<USkeletalMesh> SkeletalMesh = nullptr;
 	TUniquePtr<FMeshDescription> MeshDescription;
diff --git a/Engine/Source/Editor/SkeletalMeshEditor/Private/SkeletalMeshEditor.cpp b/Engine/Source/Editor/SkeletalMeshEditor/Private/SkeletalMeshEditor.cpp
--- a/Engine/Source/Editor/SkeletalMeshEditor/Private/SkeletalMeshEditor.cpp
+++ b/Engine/Source/Editor/SkeletalMeshEditor/Private/SkeletalMeshEditor.cpp
@@ -347,11 +347,30 @@ void FSkeletalMeshEditor::OnEditorModeIdChanged(const FEditorModeID& ModeChanged
 		return;
 	}
 
+	static const TArray<FTabId> DefaultTabs{
+		SkeletalMeshEditorTabs::DetailsTab,
+		SkeletalMeshEditorTabs::AssetDetailsTab,
+		SkeletalMeshEditorTabs::SkeletonTreeTab,
+		SkeletalMeshEditorTabs::AdvancedPreviewTab,
+		SkeletalMeshEditorTabs::MorphTargetsTab,
+		SkeletalMeshEditorTabs::CurveMetadataTab,
+		SkeletalMeshEditorTabs::FindReplaceTab
+	};
+	
 	if (bIsEnteringMode)
 	{
 		// FIXME: We should get the hosted toolkit from here.
 		if (GetEditorModeManager().GetActiveScriptableMode(ModeChangedID)->UsesToolkits())
 		{
+			// for (const FTabId& TabId: DefaultTabs)
+			// {
+			// 	if (const TSharedPtr<SDockTab> Tab = TabManager->FindExistingLiveTab(TabId))
+			// 	{
+			// 		TabToRestore.Emplace(TabId, Tab->IsForeground());
+			// 		Tab->RequestCloseTab();
+			// 	}
+			// }
+
 			TabManager->TryInvokeTab(SkeletalMeshEditorTabs::ToolboxDetailsTab);
 		}
 	}
@@ -362,6 +381,12 @@ void FSkeletalMeshEditor::OnEditorModeIdChanged(const FEditorModeID& ModeChanged
 		{
 			ToolboxTab->RequestCloseTab();
 		}
+
+		// for (const TPair<FTabId, bool>& Tab: TabToRestore)
+		// {
+		// 	TabManager->TryInvokeTab(Tab.Get<0>());
+		// }
+		TabToRestore.Empty();
 	}
 }
 
@@ -696,7 +721,7 @@ void FSkeletalMeshEditor::InitToolMenuContext(FToolMenuContext& MenuContext)
 
 void FSkeletalMeshEditor::OnToolkitHostingStarted(const TSharedRef<IToolkit>& Toolkit)
 {
-	if (ensure(!HostedToolkit.IsValid()))
+	if (!HostedToolkit.IsValid())
 	{
 		HostedToolkit = Toolkit;
 
@@ -710,7 +735,7 @@ void FSkeletalMeshEditor::OnToolkitHostingStarted(const TSharedRef<IToolkit>& To
 
 void FSkeletalMeshEditor::OnToolkitHostingFinished(const TSharedRef<IToolkit>& Toolkit)
 {
-	if (ensure(Toolkit == HostedToolkit))
+	if (Toolkit == HostedToolkit)
 	{
 		if (ToolboxWidget.IsValid())
 		{
@@ -1630,9 +1655,11 @@ void FSkeletalMeshEditorNotifier::HandleNotification(const TArray<FName>& BoneNa
 	switch (InNotifyType)
 	{
 	case ESkeletalMeshNotifyType::BonesAdded:
+		SkeletonTree->DeselectAll();
 		SkeletonTree->Refresh();
 		break;
 	case ESkeletalMeshNotifyType::BonesRemoved:
+		SkeletonTree->DeselectAll();
 		SkeletonTree->Refresh();
 		break;
 	case ESkeletalMeshNotifyType::BonesMoved:
@@ -1644,6 +1671,14 @@ void FSkeletalMeshEditorNotifier::HandleNotification(const TArray<FName>& BoneNa
 			SkeletonTree->SetSelectedBone(BoneNames[0], ESelectInfo::Direct);
 		}
 		break;
+	case ESkeletalMeshNotifyType::BonesRenamed:
+		SkeletonTree->DeselectAll();
+		SkeletonTree->Refresh();
+		if (!BoneNames.IsEmpty())
+		{
+			SkeletonTree->SetSelectedBone(BoneNames[0], ESelectInfo::Direct);
+		}
+		break;
 	}
 }
 
diff --git a/Engine/Source/Editor/SkeletalMeshEditor/Private/SkeletalMeshEditor.h b/Engine/Source/Editor/SkeletalMeshEditor/Private/SkeletalMeshEditor.h
--- a/Engine/Source/Editor/SkeletalMeshEditor/Private/SkeletalMeshEditor.h
+++ b/Engine/Source/Editor/SkeletalMeshEditor/Private/SkeletalMeshEditor.h
@@ -220,6 +220,9 @@ private:
 
 	// Binding to send/receive skeletal mesh modifications
 	TSharedPtr<FSkeletalMeshEditorBinding> Binding;
+
+	// Tabs to restore when getting out of the tool
+	TArray<TPair<FTabId, bool>> TabToRestore;
 };
 
 /**
