fix tool name in properties section of voxel offset tool

diff --git a/Engine/Plugins/Experimental/MeshModelingToolsetExp/Source/MeshModelingToolsExp/Private/VoxelMorphologyMeshesTool.cpp b/Engine/Plugins/Experimental/MeshModelingToolsetExp/Source/MeshModelingToolsExp/Private/VoxelMorphologyMeshesTool.cpp
--- a/Engine/Plugins/Experimental/MeshModelingToolsetExp/Source/MeshModelingToolsExp/Private/VoxelMorphologyMeshesTool.cpp
+++ b/Engine/Plugins/Experimental/MeshModelingToolsetExp/Source/MeshModelingToolsExp/Private/VoxelMorphologyMeshesTool.cpp
@@ -37,9 +37,9 @@ void UVoxelMorphologyMeshesTool::SetupProperties()
 	MorphologyProperties->RestoreProperties(this);
 	AddToolPropertySource(MorphologyProperties);
 
-	SetToolDisplayName(LOCTEXT("VoxelMorphologyMeshesToolName", "Voxel Morphology"));
+	SetToolDisplayName(LOCTEXT("VoxelMorphologyMeshesToolName", "Voxel Offset"));
 	GetToolManager()->DisplayMessage(
-		LOCTEXT("OnStartTool", "Apply Morphological operations to the input meshes to create a new Mesh, using voxelization techniques. UVs, sharp edges, and small/thin features will be lost. Increase Voxel Count to enhance accuracy."),
+		LOCTEXT("OnStartTool", "Apply Offset (aka Morphology) operations to the input meshes to create a new Mesh, using voxelization techniques. UVs, sharp edges, and small/thin features will be lost. Increase Voxel Count to enhance accuracy."),
 		EToolMessageLevel::UserNotification);
 
 }
