ModelingTools: clear out MaterialID values on DynamicMesh if not transferring Materials in Convert Tool

diff --git a/Engine/Plugins/Experimental/MeshModelingToolsetExp/Source/MeshModelingToolsExp/Private/ConvertMeshesTool.cpp b/Engine/Plugins/Experimental/MeshModelingToolsetExp/Source/MeshModelingToolsExp/Private/ConvertMeshesTool.cpp
--- a/Engine/Plugins/Experimental/MeshModelingToolsetExp/Source/MeshModelingToolsExp/Private/ConvertMeshesTool.cpp
+++ b/Engine/Plugins/Experimental/MeshModelingToolsetExp/Source/MeshModelingToolsExp/Private/ConvertMeshesTool.cpp
@@ -11,6 +11,7 @@
 #include "Selection/ToolSelectionUtil.h"
 #include "Physics/ComponentCollisionUtil.h"
 #include "ShapeApproximation/SimpleShapeSet3.h"
+#include "DynamicMesh/DynamicMeshAttributeSet.h"
 
 #include "TargetInterfaces/MeshDescriptionProvider.h"
 #include "TargetInterfaces/PrimitiveComponentBackedTarget.h"
@@ -101,6 +102,19 @@ void UConvertMeshesTool::OnShutdown(EToolShutdownType ShutdownType)
 
 			FTransform SourceTransform = (FTransform)UE::ToolTarget::GetLocalToWorldTransform(Targets[k]);
 			FDynamicMesh3 SourceMesh = UE::ToolTarget::GetDynamicMeshCopy(Targets[k], true);
+
+			// if not transferring materials, need to clear out any existing MaterialIDs
+			if (BasicProperties->bTransferMaterials == false && SourceMesh.HasAttributes())
+			{
+				if (FDynamicMeshMaterialAttribute* MaterialIDs = SourceMesh.Attributes()->GetMaterialID())
+				{
+					for (int32 tid : SourceMesh.TriangleIndicesItr())
+					{
+						MaterialIDs->SetValue(tid, 0);
+					}
+				}
+			}
+
 			FString AssetName = TargetActor->GetActorNameOrLabel();
 			FComponentMaterialSet Materials = UE::ToolTarget::GetMaterialSet(Targets[k]);
 			const FComponentMaterialSet* TransferMaterials = (BasicProperties->bTransferMaterials) ? &Materials : nullptr;
