CombineMeshInstances: add support for pinning corner vertices in simplification pass, based on detecting polygroups with a fixed opening angle (see comments). Expose configuration parameters in options struct. Set up with suitable values in UAtomModelBlueprintLibrary::CreateMergeMeshFromComponents_Optimized().

diff --git a/Engine/Plugins/Experimental/MeshModelingToolsetExp/Source/GeometryProcessingAdapters/Private/CombineMeshInstancesImpl.cpp b/Engine/Plugins/Experimental/MeshModelingToolsetExp/Source/GeometryProcessingAdapters/Private/CombineMeshInstancesImpl.cpp
--- a/Engine/Plugins/Experimental/MeshModelingToolsetExp/Source/GeometryProcessingAdapters/Private/CombineMeshInstancesImpl.cpp
+++ b/Engine/Plugins/Experimental/MeshModelingToolsetExp/Source/GeometryProcessingAdapters/Private/CombineMeshInstancesImpl.cpp
@@ -24,6 +24,9 @@
 #include "ShapeApproximation/MeshSimpleShapeApproximation.h"
 #include "Generators/GridBoxMeshGenerator.h"
 
+#include "Polygroups/PolygroupsGenerator.h"
+#include "GroupTopology.h"
+
 #include "MeshSimplification.h"
 #include "DynamicMesh/ColliderMesh.h"
 #include "MeshConstraintsUtil.h"
@@ -467,11 +470,99 @@ void ReplaceBadSimplifiedLODs(FMeshInstanceAssembly& Assembly)
 
 
 
+// This function tries to find "corners" of the mesh that should be exactly preserved,
+// which can help to maintain important shape features (but this is a very rough heuristic)
+static void SetupSimplifyConstraints(
+	FDynamicMesh3& Mesh, 
+	FMeshConstraints& Constraints,
+	double HardEdgeAngleThresholdDeg,
+	double LargeAreaThreshold)
+{
+	// save polygroups if they exist
+	TArray<int32> ExistingGroups;
+	if (Mesh.HasTriangleGroups())
+	{
+		ExistingGroups.SetNum(Mesh.MaxTriangleID());
+		for (int32 tid : Mesh.TriangleIndicesItr())
+		{
+			ExistingGroups[tid] = Mesh.GetTriangleGroup(tid);
+		}
+	}
+
+	// generate polygroups for planar areas of the mesh
+	FPolygroupsGenerator Generator(&Mesh);
+	const bool bUVSeams = false, bNormalSeams = false;
+	double DotTolerance = 1.0 - FMathd::Cos(HardEdgeAngleThresholdDeg * FMathd::DegToRad);
+	Generator.FindPolygroupsFromFaceNormals(DotTolerance, bUVSeams, bNormalSeams);
+	Generator.CopyPolygroupsToMesh();
+
+	FGroupTopology GroupTopology(&Mesh, true);
+
+	// find "large" areas, where large is basically defined as larger than a square area.
+	// This is not a good heuristic...
+	TSet<int32> LargeGroups;
+	for (const FGroupTopology::FGroup& Group : GroupTopology.Groups)
+	{
+		double Area = TMeshQueries<FDynamicMesh3>::GetVolumeArea(Mesh, Group.Triangles).Y;
+		if (Area > LargeAreaThreshold)
+		{
+			LargeGroups.Add(Group.GroupID);
+		}
+	}
+
+
+	// iterate over corners, ie junctions between 3 groups. Pin corner if at least
+	// two adjacent groups are "large"
+	int32 NumCorners = 0;
+	for (const FGroupTopology::FCorner& Corner : GroupTopology.Corners)
+	{
+		int32 NumLargeGroups = 0;
+		for (int32 GroupID : Corner.NeighbourGroupIDs)
+		{
+			if (LargeGroups.Contains(GroupID))
+			{
+				NumLargeGroups++;
+			}
+		}
+		if (NumLargeGroups >= 2)
+		{
+			FVertexConstraint Constraint = Constraints.GetVertexConstraint(Corner.VertexID);
+			Constraint.bCanMove = false;
+			Constraint.bCannotDelete = true;
+			Constraints.SetOrUpdateVertexConstraint(
+				Corner.VertexID, Constraint);
+			NumCorners++;
+		}
+	}
+
+	// restore groups
+	if (ExistingGroups.Num() > 0)
+	{
+		for (int32 tid : Mesh.TriangleIndicesItr())
+		{
+			Mesh.SetTriangleGroup(tid, ExistingGroups[tid]);
+		}
+	}
+	else
+	{
+		Mesh.DiscardTriangleGroups();
+	}
+}
+
+
 static void SimplifyPartMesh(
 	FDynamicMesh3& EditMesh, 
 	double Tolerance, 
-	double RecomputeNormalsAngleThreshold)
+	double RecomputeNormalsAngleThreshold,
+	bool bTryToPreserveSalientCorners = false,
+	double PreserveCornersAngleThreshold = 44,
+	double MinSalientPartDimension = 1.0)
 {
+	// currently bowties need to be split for Welder
+	FDynamicMeshEditor MeshEditor(&EditMesh);
+	FDynamicMeshEditResult EditResult;
+	MeshEditor.SplitBowties(EditResult);
+
 	// weld edges in case input was unwelded...
 	FMergeCoincidentMeshEdges Welder(&EditMesh);
 	Welder.MergeVertexTolerance = Tolerance * 0.001;
@@ -484,7 +575,9 @@ static void SimplifyPartMesh(
 		return;
 	}
 
-	FVolPresMeshSimplification Simplifier(&EditMesh);
+	using SimplifierType = FVolPresMeshSimplification;
+	//using SimplifierType = FQEMSimplification;
+	SimplifierType Simplifier(&EditMesh);
 
 	// clear out attributes so it doesn't affect simplification
 	//EditMesh.DiscardAttributes();
@@ -493,7 +586,7 @@ static void SimplifyPartMesh(
 	EditMesh.Attributes()->DisablePrimaryColors();
 	FMeshNormals::InitializeOverlayToPerVertexNormals(EditMesh.Attributes()->PrimaryNormals(), false);
 
-	Simplifier.ProjectionMode = FVolPresMeshSimplification::ETargetProjectionMode::NoProjection;
+	Simplifier.ProjectionMode = SimplifierType::ETargetProjectionMode::NoProjection;
 
 	FColliderMesh ColliderMesh;
 	ColliderMesh.Initialize(EditMesh);
@@ -501,8 +594,12 @@ static void SimplifyPartMesh(
 	Simplifier.SetProjectionTarget(&ProjectionTarget);
 
 	Simplifier.DEBUG_CHECK_LEVEL = 0;
-	Simplifier.bRetainQuadricMemory = false; 
-	// currenly no need for this path, may need to resurrect it in the future
+
+	// Memory seems to work better on low-poly parts...
+	// This should perhaps be based on some heuristics about 'part type'
+	Simplifier.bRetainQuadricMemory = true;
+												
+	// currenly no need for this path, as seam attributes have been cleared. 
 	//if ( bNoSplitAttributes == false )
 	//{
 	//	Simplifier.bAllowSeamCollapse = true;
@@ -524,9 +621,17 @@ static void SimplifyPartMesh(
 	FMeshConstraints Constraints;
 	FMeshConstraintsUtil::ConstrainAllBoundariesAndSeams(Constraints, EditMesh,
 		MeshBoundaryConstraints, GroupBorderConstraints, MaterialBorderConstraints, true, false, true);
+
+	// add optional constraints to try to preserve area
+	if (bTryToPreserveSalientCorners)
+	{
+		SetupSimplifyConstraints(EditMesh, Constraints, PreserveCornersAngleThreshold, MinSalientPartDimension * MinSalientPartDimension);
+	}
+
+
 	Simplifier.SetExternalConstraints(MoveTemp(Constraints));
 
-	Simplifier.GeometricErrorConstraint = FVolPresMeshSimplification::EGeometricErrorCriteria::PredictedPointToProjectionTarget;
+	Simplifier.GeometricErrorConstraint = SimplifierType::EGeometricErrorCriteria::PredictedPointToProjectionTarget;
 	Simplifier.GeometricErrorTolerance = Tolerance;
 
 	Simplifier.SimplifyToTriangleCount( 1 );
@@ -712,7 +817,8 @@ void ComputeMeshApproximations(
 		for (int32 k = 0; k < NumSimplifiedLODs; ++k)
 		{
 			ApproxGeo.SimplifiedMeshLODs[k] = OptimizationSourceMesh;
-			SimplifyPartMesh(ApproxGeo.SimplifiedMeshLODs[k], InitialTolerance, AngleThresholdDeg);
+			SimplifyPartMesh(ApproxGeo.SimplifiedMeshLODs[k], InitialTolerance, AngleThresholdDeg,
+				CombineOptions.bSimplifyPreserveCorners, CombineOptions.SimplifySharpEdgeAngleDeg, CombineOptions.SimplifyMinSalientDimension);
 			InitialTolerance *= CombineOptions.SimplifyLODLevelToleranceScale;
 		}
 
diff --git a/Engine/Source/Developer/GeometryProcessingInterfaces/Public/GeometryProcessingInterfaces/CombineMeshInstances.h b/Engine/Source/Developer/GeometryProcessingInterfaces/Public/GeometryProcessingInterfaces/CombineMeshInstances.h
--- a/Engine/Source/Developer/GeometryProcessingInterfaces/Public/GeometryProcessingInterfaces/CombineMeshInstances.h
+++ b/Engine/Source/Developer/GeometryProcessingInterfaces/Public/GeometryProcessingInterfaces/CombineMeshInstances.h
@@ -88,6 +88,9 @@ public:
 		int32 NumSimplifiedLODs = 3;
 		double SimplifyBaseTolerance = 1.0;
 		double SimplifyLODLevelToleranceScale = 2.0;
+		bool bSimplifyPreserveCorners = true;
+		double SimplifySharpEdgeAngleDeg = 44.0;
+		double SimplifyMinSalientDimension = 1.0;
 
 		double OptimizeBaseTriCost = 0.7;
 		double OptimizeLODLevelTriCostScale = 1.5;
