Fix a crash that occurs when the edit normals tool is shutting down after it was started with more than one mesh selected

diff --git a/Engine/Plugins/Experimental/MeshModelingToolsetExp/Source/MeshModelingToolsExp/Private/EditNormalsTool.cpp b/Engine/Plugins/Experimental/MeshModelingToolsetExp/Source/MeshModelingToolsExp/Private/EditNormalsTool.cpp
--- a/Engine/Plugins/Experimental/MeshModelingToolsetExp/Source/MeshModelingToolsExp/Private/EditNormalsTool.cpp
+++ b/Engine/Plugins/Experimental/MeshModelingToolsetExp/Source/MeshModelingToolsExp/Private/EditNormalsTool.cpp
@@ -159,7 +159,10 @@ void UEditNormalsTool::OnShutdown(EToolShutdownType ShutdownType)
 {
 	BasicProperties->SaveProperties(this);
 	AdvancedProperties->SaveProperties(this);
-	PolygroupLayerProperties->SaveProperties(this, TEXT("EditNormalsTool"));
+	if (PolygroupLayerProperties)
+	{
+		PolygroupLayerProperties->SaveProperties(this, TEXT("EditNormalsTool"));
+	}
 
 	// Restore (unhide) the source meshes
 	for (int32 ComponentIdx = 0, NumTargets = Targets.Num(); ComponentIdx < NumTargets; ComponentIdx++)
