Skeleton Modifier BP/Python API

diff --git a/Engine/Plugins/Experimental/MeshModelingToolsetExp/Source/MeshModelingToolsExp/Private/SkeletalMesh/SkeletonEditingTool.cpp b/Engine/Plugins/Experimental/MeshModelingToolsetExp/Source/MeshModelingToolsExp/Private/SkeletalMesh/SkeletonEditingTool.cpp
--- a/Engine/Plugins/Experimental/MeshModelingToolsetExp/Source/MeshModelingToolsExp/Private/SkeletalMesh/SkeletonEditingTool.cpp
+++ b/Engine/Plugins/Experimental/MeshModelingToolsetExp/Source/MeshModelingToolsExp/Private/SkeletalMesh/SkeletonEditingTool.cpp
@@ -26,28 +26,28 @@ namespace SkeletonEditingTool
 	
 FRefSkeletonChange::FRefSkeletonChange(const USkeletonEditingTool* InTool)
 	: FToolCommandChange()
-	, PreChangeSkeleton(InTool->SkeletonModifier.GetReferenceSkeleton())
-	, PreBoneTracker(InTool->SkeletonModifier.GetBoneIndexTracker())
-	, PostChangeSkeleton(InTool->SkeletonModifier.GetReferenceSkeleton())
-	, PostBoneTracker(InTool->SkeletonModifier.GetBoneIndexTracker())
+	, PreChangeSkeleton(InTool->Modifier->GetReferenceSkeleton())
+	, PreBoneTracker(InTool->Modifier->GetBoneIndexTracker())
+	, PostChangeSkeleton(InTool->Modifier->GetReferenceSkeleton())
+	, PostBoneTracker(InTool->Modifier->GetBoneIndexTracker())
 {}
 
 void FRefSkeletonChange::StoreSkeleton(const USkeletonEditingTool* InTool)
 {
-	PostChangeSkeleton = InTool->SkeletonModifier.GetReferenceSkeleton();
-	PostBoneTracker = InTool->SkeletonModifier.GetBoneIndexTracker();
+	PostChangeSkeleton = InTool->Modifier->GetReferenceSkeleton();
+	PostBoneTracker = InTool->Modifier->GetBoneIndexTracker();
 }
 
 void FRefSkeletonChange::Apply(UObject* Object)
 { // redo
 	USkeletonEditingTool* Tool = CastChecked<USkeletonEditingTool>(Object);
-	Tool->SkeletonModifier.ExternalUpdate(PostChangeSkeleton, PostBoneTracker);
+	Tool->Modifier->ExternalUpdate(PostChangeSkeleton, PostBoneTracker);
 }
 
 void FRefSkeletonChange::Revert(UObject* Object)
 { // undo
 	USkeletonEditingTool* Tool = CastChecked<USkeletonEditingTool>(Object);
-	Tool->SkeletonModifier.ExternalUpdate(PreChangeSkeleton, PreBoneTracker);
+	Tool->Modifier->ExternalUpdate(PreChangeSkeleton, PreBoneTracker);
 }
 
 }
@@ -102,10 +102,11 @@ void USkeletonEditingTool::Setup()
 	}
 
 	// setup modifier
-	SkeletonModifier.Init(SkeletalMesh);
+	Modifier = NewObject<USkeletonModifier>(this);
+	Modifier->Init(SkeletalMesh);
 
 	// setup current bone
-	const FReferenceSkeleton& RefSkeleton = SkeletonModifier.GetReferenceSkeleton();
+	const FReferenceSkeleton& RefSkeleton = Modifier->GetReferenceSkeleton();
 	const int32 NumBones = RefSkeleton.GetNum();
 	const FName& RootBoneName = NumBones ? RefSkeleton.GetBoneName(0) : NAME_None;
 
@@ -167,7 +168,7 @@ void USkeletonEditingTool::Shutdown(EToolShutdownType ShutdownType)
 	if (ShutdownType == EToolShutdownType::Accept)
 	{
 		GetToolManager()->BeginUndoTransaction(LOCTEXT("SkeletonEditingTool", "Commit Skeleton Editing"));
-		SkeletonModifier.CommitSkeletonToSkeletalMesh();
+		Modifier->CommitSkeletonToSkeletalMesh();
 		GetToolManager()->EndUndoTransaction();
 
 		// to force to refresh the tree
@@ -270,8 +271,8 @@ void USkeletonEditingTool::CreateNewBone()
 
 	BeginChange();
 
-	const FName BoneName = SkeletonModifier.GetUniqueName(Properties->DefaultName);
-	const bool bBoneAdded = SkeletonModifier.AddBone(BoneName, CurrentBone, Properties->Transform);
+	const FName BoneName = Modifier->GetUniqueName(Properties->DefaultName);
+	const bool bBoneAdded = Modifier->AddBone(BoneName, CurrentBone, Properties->Transform);
 	if (bBoneAdded)
 	{
 		if (NeedsNotification())
@@ -294,7 +295,7 @@ void USkeletonEditingTool::MirrorBones()
 	TGuardValue OperationGuard(Operation, EEditingOperation::Mirror);
 	BeginChange();
 
-	const bool bBonesMirrored = SkeletonModifier.MirrorBones(GetSelectedBones(), MirroringProperties->Options);
+	const bool bBonesMirrored = Modifier->MirrorBones(GetSelectedBones(), MirroringProperties->Options);
 	if (bBonesMirrored)
 	{
 		EndChange();
@@ -311,7 +312,7 @@ void USkeletonEditingTool::RemoveBones()
 	TGuardValue OperationGuard(Operation, EEditingOperation::Remove);
 	BeginChange();
 
-	const bool bBonesRemoved = SkeletonModifier.RemoveBones(BonesToRemove, true);
+	const bool bBonesRemoved = Modifier->RemoveBones(BonesToRemove, true);
 	if (bBonesRemoved)
 	{
 		// if (NeedsNotification())
@@ -333,7 +334,7 @@ void USkeletonEditingTool::UnParentBones()
 	TGuardValue OperationGuard(Operation, EEditingOperation::Parent);
 	BeginChange();
 
-	const bool bBonesUnParented = SkeletonModifier.ParentBones(GetSelectedBones(), Dummy);
+	const bool bBonesUnParented = Modifier->ParentBones(GetSelectedBones(), Dummy);
 	if (bBonesUnParented)
 	{
 		GetToolManager()->DisplayMessage(LOCTEXT("Unparent", "Selected bones have been unparented."), EToolMessageLevel::UserNotification);
@@ -353,7 +354,7 @@ void USkeletonEditingTool::ParentBones(const FName& InParentName)
 	}
 
 	BeginChange();
-	const bool bBonesParented = SkeletonModifier.ParentBones(GetSelectedBones(), {InParentName});
+	const bool bBonesParented = Modifier->ParentBones(GetSelectedBones(), {InParentName});
 	if (bBonesParented)
 	{
 		EndChange();
@@ -367,7 +368,7 @@ void USkeletonEditingTool::ParentBones(const FName& InParentName)
 
 void USkeletonEditingTool::MoveBones()
 {
-	const FReferenceSkeleton& RefSkeleton = SkeletonModifier.GetReferenceSkeleton();
+	const FReferenceSkeleton& RefSkeleton = Modifier->GetReferenceSkeleton();
 
 	const TArray<FName> Bones = GetSelectedBones();
 	const bool bHasValidBone = Bones.ContainsByPredicate([&](const FName& InBoneName)
@@ -383,7 +384,7 @@ void USkeletonEditingTool::MoveBones()
 	TGuardValue OperationGuard(Operation, EEditingOperation::Transform);
 	BeginChange();
 
-	const bool bBonesMoved = SkeletonModifier.SetBoneTransform(Bones[0], Properties->Transform, Properties->bUpdateChildren);
+	const bool bBonesMoved = Modifier->SetBoneTransform(Bones[0], Properties->Transform, Properties->bUpdateChildren);
 	if (bBonesMoved)
 	{
 		// if (NeedsNotification())
@@ -405,7 +406,7 @@ void USkeletonEditingTool::RenameBones()
 		return;
 	}
 	
-	const FReferenceSkeleton& RefSkeleton = SkeletonModifier.GetReferenceSkeleton();
+	const FReferenceSkeleton& RefSkeleton = Modifier->GetReferenceSkeleton();
 	if (RefSkeleton.FindRawBoneIndex(CurrentBone) == INDEX_NONE)
 	{
 		return;
@@ -414,7 +415,7 @@ void USkeletonEditingTool::RenameBones()
 	TGuardValue OperationGuard(Operation, EEditingOperation::Rename);
 	BeginChange();
 
-	const bool bBoneRenamed = SkeletonModifier.RenameBone(CurrentBone, Properties->Name);
+	const bool bBoneRenamed = Modifier->RenameBone(CurrentBone, Properties->Name);
 	if (bBoneRenamed)
 	{
 		CurrentBone = Properties->Name;
@@ -441,7 +442,7 @@ void USkeletonEditingTool::OnClickDrag(const FInputDeviceRay& DragPos)
 	FVector HitPoint;
 	if (ProjectionProperties->GetProjectionPoint(DragPos, HitPoint))
 	{
-		const FTransform& ParentGlobal = SkeletonModifier.GetTransform(ParentIndex, true);
+		const FTransform& ParentGlobal = Modifier->GetTransform(ParentIndex, true);
 		Properties->Transform.SetLocation(ParentGlobal.InverseTransformPosition(HitPoint));
 
 		if (!ActiveChange)
@@ -450,7 +451,7 @@ void USkeletonEditingTool::OnClickDrag(const FInputDeviceRay& DragPos)
 			BeginChange();
 		}
 
-		const bool bBoneMoved = SkeletonModifier.SetBoneTransform(CurrentBone, Properties->Transform, Properties->bUpdateChildren);
+		const bool bBoneMoved = Modifier->SetBoneTransform(CurrentBone, Properties->Transform, Properties->bUpdateChildren);
 		if (!bBoneMoved)
 		{
 			CancelChange();
@@ -460,16 +461,16 @@ void USkeletonEditingTool::OnClickDrag(const FInputDeviceRay& DragPos)
 		const bool bOrient = Operation == EEditingOperation::Create && OrientingProperties->bAutoOrient;
 		if (bOrient && ParentIndex != INDEX_NONE)
 		{
-			const FReferenceSkeleton& RefSkeleton = SkeletonModifier.GetReferenceSkeleton();
+			const FReferenceSkeleton& RefSkeleton = Modifier->GetReferenceSkeleton();
 			const FName ParentName = RefSkeleton.GetRawRefBoneInfo()[ParentIndex].Name;
-			SkeletonModifier.OrientBone(ParentName, OrientingProperties->Options);
+			Modifier->OrientBone(ParentName, OrientingProperties->Options);
 		}
 	}
 }
 
 void USkeletonEditingTool::OrientBones()
 {
-	const FReferenceSkeleton& RefSkeleton = SkeletonModifier.GetReferenceSkeleton();
+	const FReferenceSkeleton& RefSkeleton = Modifier->GetReferenceSkeleton();
 
 	const TArray<FName> Bones = GetSelectedBones();
 	const bool bHasValidBone = Bones.ContainsByPredicate([&](const FName& InBoneName)
@@ -485,7 +486,7 @@ void USkeletonEditingTool::OrientBones()
 	TGuardValue OperationGuard(Operation, EEditingOperation::Transform);
 	BeginChange();
 	
-	const bool bBoneOriented = SkeletonModifier.OrientBones(Bones, OrientingProperties->Options);
+	const bool bBoneOriented = Modifier->OrientBones(Bones, OrientingProperties->Options);
 	if (bBoneOriented)
 	{
 		// if (NeedsNotification())
@@ -518,7 +519,7 @@ FInputRayHit USkeletonEditingTool::CanBeginClickDragSequence(const FInputDeviceR
 		{
 			if (TOptional<FName> OptBoneName = GetBoneName(HitProxy))
 			{
-				const FReferenceSkeleton& ReferenceSkeleton = SkeletonModifier.GetReferenceSkeleton();
+				const FReferenceSkeleton& ReferenceSkeleton = Modifier->GetReferenceSkeleton();
 				return ReferenceSkeleton.FindRawBoneIndex(*OptBoneName);
 			}
 		}
@@ -529,11 +530,11 @@ FInputRayHit USkeletonEditingTool::CanBeginClickDragSequence(const FInputDeviceR
 	const int32 BoneIndex = PickBone();
 
 	// update parent
-	const FReferenceSkeleton& ReferenceSkeleton = SkeletonModifier.GetReferenceSkeleton();
+	const FReferenceSkeleton& ReferenceSkeleton = Modifier->GetReferenceSkeleton();
 	ParentIndex = INDEX_NONE;
 
 	// update projection properties
-	const FVector GlobalPosition = SkeletonModifier.GetTransform(BoneIndex, true).GetTranslation();
+	const FVector GlobalPosition = Modifier->GetTransform(BoneIndex, true).GetTranslation();
 	ProjectionProperties->UpdatePlane(*ViewContext, GlobalPosition);
 
 	// if we picked a new bone
@@ -584,7 +585,7 @@ FInputRayHit USkeletonEditingTool::CanBeginClickDragSequence(const FInputDeviceR
 		{
 			// CurrentBone is gonna be the parent
 			ParentIndex = ReferenceSkeleton.FindRawBoneIndex(CurrentBone);
-			const FTransform& ParentGlobalTransform = SkeletonModifier.GetTransform(ParentIndex, true);
+			const FTransform& ParentGlobalTransform = Modifier->GetTransform(ParentIndex, true);
 
 			// Create the new bone under mouse
 			Properties->Transform.SetLocation(ParentGlobalTransform.InverseTransformPosition(HitPoint));
@@ -655,7 +656,7 @@ void USkeletonEditingTool::Render(IToolsContextRenderAPI* RenderAPI)
 	const IPrimitiveComponentBackedTarget* TargetComponent = Cast<IPrimitiveComponentBackedTarget>(Target);
 	const FTransform ComponentTransform = TargetComponent->GetWorldTransform();
 	
-	const FReferenceSkeleton& RefSkeleton = SkeletonModifier.GetReferenceSkeleton();
+	const FReferenceSkeleton& RefSkeleton = Modifier->GetReferenceSkeleton();
 
 	const int32 NumBones = RefSkeleton.GetRawBoneNum();
 	TArray<TRefCountPtr<HHitProxy>> HitProxies; HitProxies.Reserve(NumBones);
@@ -665,7 +666,7 @@ void USkeletonEditingTool::Render(IToolsContextRenderAPI* RenderAPI)
 	
 	for (int32 Index = 0; Index < NumBones; ++Index)
 	{
-		const FTransform& BoneTransform = SkeletonModifier.GetTransform(Index, true);
+		const FTransform& BoneTransform = Modifier->GetTransform(Index, true);
 		WorldTransforms[Index] = BoneTransform;
 		RequiredBones[Index] = Index;
 		BoneColors[Index] = DefaultBoneColor;
@@ -704,12 +705,12 @@ FBox USkeletonEditingTool::GetWorldSpaceFocusBox()
 		FBox Box(EForceInit::ForceInit);
 		TSet<int32> AllChildren;
 
-		const FReferenceSkeleton& RefSkeleton = SkeletonModifier.GetReferenceSkeleton();
+		const FReferenceSkeleton& RefSkeleton = Modifier->GetReferenceSkeleton();
 		
 		for (const FName& BoneName: Selection)
 		{
 			const int32 BoneIndex = RefSkeleton.FindRawBoneIndex(BoneName);
-			Box += SkeletonModifier.GetTransform(BoneIndex, true).GetTranslation();
+			Box += Modifier->GetTransform(BoneIndex, true).GetTranslation();
 
 			// get direct children
 			TArray<int32> Children;
@@ -719,7 +720,7 @@ FBox USkeletonEditingTool::GetWorldSpaceFocusBox()
 
 		for (const int32 ChildIndex: AllChildren)
 		{
-			Box += SkeletonModifier.GetTransform(ChildIndex, true).GetTranslation();
+			Box += Modifier->GetTransform(ChildIndex, true).GetTranslation();
 		}
 		
 		return Box;
diff --git a/Engine/Plugins/Experimental/MeshModelingToolsetExp/Source/MeshModelingToolsExp/Public/SkeletalMesh/SkeletonEditingTool.h b/Engine/Plugins/Experimental/MeshModelingToolsetExp/Source/MeshModelingToolsExp/Public/SkeletalMesh/SkeletonEditingTool.h
--- a/Engine/Plugins/Experimental/MeshModelingToolsetExp/Source/MeshModelingToolsExp/Public/SkeletalMesh/SkeletonEditingTool.h
+++ b/Engine/Plugins/Experimental/MeshModelingToolsetExp/Source/MeshModelingToolsExp/Public/SkeletalMesh/SkeletonEditingTool.h
@@ -129,7 +129,8 @@ public:
 	int32 ParentIndex = INDEX_NONE;
 	FName CurrentBone = NAME_None;
 
-	FSkeletonModifier SkeletonModifier;
+	UPROPERTY()
+	TObjectPtr<USkeletonModifier> Modifier;
 	
 protected:
 
diff --git a/Engine/Source/Developer/SkeletalMeshUtilitiesCommon/Private/SkeletonModifier.cpp b/Engine/Source/Developer/SkeletalMeshUtilitiesCommon/Private/SkeletonModifier.cpp
--- a/Engine/Source/Developer/SkeletalMeshUtilitiesCommon/Private/SkeletonModifier.cpp
+++ b/Engine/Source/Developer/SkeletalMeshUtilitiesCommon/Private/SkeletonModifier.cpp
@@ -17,7 +17,7 @@
 #include "Rendering/SkeletalMeshModel.h"
 
 
-namespace FSkeletonModifierLocals
+namespace USkeletonModifierLocals
 {
 	static constexpr int32 LODIndex = 0;
 }
@@ -115,7 +115,7 @@ FTransform FOrientOptions::OrientTransform(const FVector& InPrimaryTarget, const
 	return Transform;
 }
 
-void FSkeletonModifier::ExternalUpdate(const FReferenceSkeleton& InRefSkeleton, const TArray<int32>& InIndexTracker)
+void USkeletonModifier::ExternalUpdate(const FReferenceSkeleton& InRefSkeleton, const TArray<int32>& InIndexTracker)
 {
 	if (!ReferenceSkeleton)
 	{
@@ -127,7 +127,7 @@ void FSkeletonModifier::ExternalUpdate(const FReferenceSkeleton& InRefSkeleton,
 	BoneIndexTracker = InIndexTracker;
 }
 
-bool FSkeletonModifier::Init(USkeletalMesh* InSkeletalMesh)
+bool USkeletonModifier::Init(USkeletalMesh* InSkeletalMesh)
 {
 	SkeletalMesh = nullptr;
 	MeshDescription.Reset();
@@ -192,7 +192,7 @@ bool FSkeletonModifier::Init(USkeletalMesh* InSkeletalMesh)
 
 	// store mesh description to edit
 	MeshDescription = MakeUnique<FMeshDescription>();
-	SkeletalMesh->GetMeshDescription(FSkeletonModifierLocals::LODIndex, *MeshDescription);
+	SkeletalMesh->GetMeshDescription(USkeletonModifierLocals::LODIndex, *MeshDescription);
 	if (MeshDescription->IsEmpty())
 	{
 		UE_LOG(LogAnimation, Error, TEXT("Skeleton Modifier: mesh description is emtpy."));
@@ -214,12 +214,12 @@ bool FSkeletonModifier::Init(USkeletalMesh* InSkeletalMesh)
 	}
 	return true;
 #else
-	ensureMsgf(false, TEXT("FSkeletonModifier is an editor only feature."));
+	ensureMsgf(false, TEXT("Skeleton Modifier is an editor only feature."));
 #endif
 	return false;
 }
 
-bool FSkeletonModifier::IsReferenceSkeletonValid(const bool bLog) const
+bool USkeletonModifier::IsReferenceSkeletonValid(const bool bLog) const
 {
 	if (!ReferenceSkeleton.IsValid())
 	{
@@ -232,7 +232,7 @@ bool FSkeletonModifier::IsReferenceSkeletonValid(const bool bLog) const
 	return true;
 }
 
-bool FSkeletonModifier::CommitSkeletonToSkeletalMesh()
+bool USkeletonModifier::CommitSkeletonToSkeletalMesh()
 {
 	if (!SkeletalMesh || !ReferenceSkeleton || !MeshDescription)
 	{
@@ -347,7 +347,7 @@ bool FSkeletonModifier::CommitSkeletonToSkeletalMesh()
 	SkeletalMesh->CalculateInvRefMatrices();
 	
 	// update skeletal mesh LOD (cf. USkeletalMesh::CommitMeshDescription)
-	SkeletalMesh->CommitMeshDescription(FSkeletonModifierLocals::LODIndex, *MeshDescription);
+	SkeletalMesh->CommitMeshDescription(USkeletonModifierLocals::LODIndex, *MeshDescription);
 
 	SkeletalMesh->PostEditChange();
 
@@ -361,12 +361,12 @@ bool FSkeletonModifier::CommitSkeletonToSkeletalMesh()
 
 	return true;
 #else
-	ensureMsgf(false, TEXT("FSkeletonModifier is an editor only feature."));
+	ensureMsgf(false, TEXT("Skeleton Modifier is an editor only feature."));
 #endif
 	return false;
 }
 
-bool FSkeletonModifier::AddBone(const FName InBoneName, const FName InParentName, const FTransform& InTransform)
+bool USkeletonModifier::AddBone(const FName InBoneName, const FName InParentName, const FTransform& InTransform)
 {
 	if (InBoneName == NAME_None)
 	{
@@ -377,7 +377,7 @@ bool FSkeletonModifier::AddBone(const FName InBoneName, const FName InParentName
 	return AddBones({InBoneName}, {InParentName}, {InTransform});
 }
 
-bool FSkeletonModifier::AddBones(
+bool USkeletonModifier::AddBones(
 	const TArray<FName>& InBoneNames, const TArray<FName>& InParentNames, const TArray<FTransform>& InTransforms)
 {
 	if (!IsReferenceSkeletonValid())
@@ -469,7 +469,7 @@ bool FSkeletonModifier::AddBones(
 	return true;
 }
 
-bool FSkeletonModifier::MirrorBone(const FName InBoneName, const FMirrorOptions& InOptions)
+bool USkeletonModifier::MirrorBone(const FName InBoneName, const FMirrorOptions& InOptions)
 {
 	if (InBoneName == NAME_None)
 	{
@@ -480,7 +480,7 @@ bool FSkeletonModifier::MirrorBone(const FName InBoneName, const FMirrorOptions&
 	return MirrorBones({InBoneName}, InOptions);
 }
 
-bool FSkeletonModifier::MirrorBones(const TArray<FName>& InBonesName, const FMirrorOptions& InOptions)
+bool USkeletonModifier::MirrorBones(const TArray<FName>& InBonesName, const FMirrorOptions& InOptions)
 {
 	if (!IsReferenceSkeletonValid())
 	{
@@ -526,7 +526,7 @@ bool FSkeletonModifier::MirrorBones(const TArray<FName>& InBonesName, const FMir
 	return true;
 }
 
-void FSkeletonModifier::GetBonesToMirror(
+void USkeletonModifier::GetBonesToMirror(
 	const TArray<FName>& InBonesName, const FMirrorOptions& InOptions, TArray<int32>& OutBonesToMirror) const
 {
 	OutBonesToMirror.Reset();
@@ -566,7 +566,7 @@ void FSkeletonModifier::GetBonesToMirror(
 	OutBonesToMirror = IndicesToMirror.Array();
 }
 
-void FSkeletonModifier::GetMirroredNames(
+void USkeletonModifier::GetMirroredNames(
 	const TArray<int32>& InBonesToMirror, const FMirrorOptions& InOptions, TArray<FName>& OutBonesName) const
 {
 	OutBonesName.Reset();
@@ -597,7 +597,7 @@ void FSkeletonModifier::GetMirroredNames(
 	});
 }
 
-void FSkeletonModifier::GetMirroredBones(
+void USkeletonModifier::GetMirroredBones(
 	const TArray<int32>& InBonesToMirror, const TArray<FName>& InMirroredNames, TArray<int32>& OutMirroredBones)
 {
 	OutMirroredBones.Reset();
@@ -673,7 +673,7 @@ void FSkeletonModifier::GetMirroredBones(
 	}
 }
 
-void FSkeletonModifier::GetMirroredTransforms(
+void USkeletonModifier::GetMirroredTransforms(
 	const TArray<int32>& InBonesToMirror, const TArray<int32>& InMirroredBones,
 	const FMirrorOptions& InOptions, TArray<FTransform>& OutMirroredTransforms) const
 {
@@ -750,7 +750,7 @@ void FSkeletonModifier::GetMirroredTransforms(
 // atm, we update the bone's local ref transform so children's global transforms are changed (we just need to cache the
 // global transforms then restore them back)
 // orienting the bone for example should change the children's global transform
-bool FSkeletonModifier::SetBoneTransform( const FName InBoneName, const FTransform& InNewTransform, const bool bMoveChildren)
+bool USkeletonModifier::SetBoneTransform( const FName InBoneName, const FTransform& InNewTransform, const bool bMoveChildren)
 {
 	if (InBoneName == NAME_None)
 	{
@@ -761,7 +761,7 @@ bool FSkeletonModifier::SetBoneTransform( const FName InBoneName, const FTransfo
 	return SetBonesTransforms({InBoneName}, {InNewTransform}, bMoveChildren);
 }
 
-bool FSkeletonModifier::SetBonesTransforms(
+bool USkeletonModifier::SetBonesTransforms(
 	const TArray<FName>& InBoneNames, const TArray<FTransform>& InNewTransforms, const bool bMoveChildren)
 {
 	if (!IsReferenceSkeletonValid())
@@ -860,7 +860,7 @@ bool FSkeletonModifier::SetBonesTransforms(
 	return true;
 }
 
-bool FSkeletonModifier::RemoveBone(const FName InBoneName, const bool bRemoveChildren)
+bool USkeletonModifier::RemoveBone(const FName InBoneName, const bool bRemoveChildren)
 {
 	if (InBoneName == NAME_None)
 	{
@@ -871,7 +871,7 @@ bool FSkeletonModifier::RemoveBone(const FName InBoneName, const bool bRemoveChi
 	return RemoveBones({InBoneName}, bRemoveChildren);
 }
 
-bool FSkeletonModifier::RemoveBones(const TArray<FName>& InBoneNames, const bool bRemoveChildren)
+bool USkeletonModifier::RemoveBones(const TArray<FName>& InBoneNames, const bool bRemoveChildren)
 {
 	if (!IsReferenceSkeletonValid())
 	{
@@ -911,7 +911,7 @@ bool FSkeletonModifier::RemoveBones(const TArray<FName>& InBoneNames, const bool
 	return true;		
 }
 
-bool FSkeletonModifier::RenameBone(const FName InOldBoneName, const FName InNewBoneName)
+bool USkeletonModifier::RenameBone(const FName InOldBoneName, const FName InNewBoneName)
 {
 	if (InOldBoneName == NAME_None || InOldBoneName == NAME_None || InNewBoneName == InOldBoneName)
 	{
@@ -922,7 +922,7 @@ bool FSkeletonModifier::RenameBone(const FName InOldBoneName, const FName InNewB
 	return RenameBones({InOldBoneName}, {InNewBoneName});
 }
 
-bool FSkeletonModifier::RenameBones(const TArray<FName>& InOldBoneNames, const TArray<FName>& InNewBoneNames)
+bool USkeletonModifier::RenameBones(const TArray<FName>& InOldBoneNames, const TArray<FName>& InNewBoneNames)
 {
 	if (!IsReferenceSkeletonValid())
 	{
@@ -956,7 +956,7 @@ bool FSkeletonModifier::RenameBones(const TArray<FName>& InOldBoneNames, const T
 	return true;
 }
 
-bool FSkeletonModifier::ParentBone(const FName InBoneName, const FName InParentName)
+bool USkeletonModifier::ParentBone(const FName InBoneName, const FName InParentName)
 {
 	if (InBoneName == NAME_None)
 	{
@@ -967,7 +967,7 @@ bool FSkeletonModifier::ParentBone(const FName InBoneName, const FName InParentN
 	return ParentBones({InBoneName}, {InParentName});
 }
 
-bool FSkeletonModifier::ParentBones(const TArray<FName>& InBoneNames, const TArray<FName>& InParentNames)
+bool USkeletonModifier::ParentBones(const TArray<FName>& InBoneNames, const TArray<FName>& InParentNames)
 {
 	if (!IsReferenceSkeletonValid())
 	{
@@ -1030,7 +1030,7 @@ bool FSkeletonModifier::ParentBones(const TArray<FName>& InBoneNames, const TArr
 	return true;
 }
 
-bool FSkeletonModifier::OrientBone(const FName InBoneName, const FOrientOptions& InOptions)
+bool USkeletonModifier::OrientBone(const FName InBoneName, const FOrientOptions& InOptions)
 {
 	if (InBoneName == NAME_None)
 	{
@@ -1041,7 +1041,7 @@ bool FSkeletonModifier::OrientBone(const FName InBoneName, const FOrientOptions&
 	return OrientBones({InBoneName}, InOptions);
 }
 
-bool FSkeletonModifier::OrientBones(const TArray<FName>& InBoneNames, const FOrientOptions& InOptions)
+bool USkeletonModifier::OrientBones(const TArray<FName>& InBoneNames, const FOrientOptions& InOptions)
 {
 	if (!IsReferenceSkeletonValid())
 	{
@@ -1127,7 +1127,7 @@ bool FSkeletonModifier::OrientBones(const TArray<FName>& InBoneNames, const FOri
 	return SetBonesTransforms(BonesToAlign, Transforms, bMoveChildren);
 }
 
-void FSkeletonModifier::GetBonesToOrient(
+void USkeletonModifier::GetBonesToOrient(
 	const TArray<FName>& InBonesName, const FOrientOptions& InOptions, TArray<int32>& OutBonesToOrient) const
 {
 	OutBonesToOrient.Reset();
@@ -1167,7 +1167,7 @@ void FSkeletonModifier::GetBonesToOrient(
 	OutBonesToOrient = IndicesToOrient.Array();
 }
 
-void FSkeletonModifier::UpdateBoneTracker(const TArray<FMeshBoneInfo>& InOtherInfos)
+void USkeletonModifier::UpdateBoneTracker(const TArray<FMeshBoneInfo>& InOtherInfos)
 {
 	for (int32 Index = 0; Index < BoneIndexTracker.Num(); ++Index)
 	{
@@ -1181,7 +1181,7 @@ void FSkeletonModifier::UpdateBoneTracker(const TArray<FMeshBoneInfo>& InOtherIn
 	}
 }
 
-FName FSkeletonModifier::GetUniqueName(const FName InBoneName) const
+FName USkeletonModifier::GetUniqueName(const FName InBoneName) const
 {
 	if (!ReferenceSkeleton)
 	{
@@ -1209,18 +1209,39 @@ FName FSkeletonModifier::GetUniqueName(const FName InBoneName) const
 	return OutName;
 }
 
-const FReferenceSkeleton& FSkeletonModifier::GetReferenceSkeleton() const
+const FReferenceSkeleton& USkeletonModifier::GetReferenceSkeleton() const
 {
 	static const FReferenceSkeleton Dummy;
 	return ReferenceSkeleton ? *ReferenceSkeleton.Get() : Dummy;
 }
 
-const TArray<int32>& FSkeletonModifier::GetBoneIndexTracker() const
+const TArray<int32>& USkeletonModifier::GetBoneIndexTracker() const
 {
 	return BoneIndexTracker;	
 }
 
-const FTransform& FSkeletonModifier::GetTransform(const int32 InBoneIndex, const bool bGlobal) const
+FTransform USkeletonModifier::GetBoneTransform(const FName InBoneName, const bool bGlobal) const
+{
+	if (!ReferenceSkeleton.IsValid())
+	{
+		return FTransform::Identity;
+	}
+	return GetTransform(ReferenceSkeleton->FindRawBoneIndex(InBoneName), bGlobal);
+}
+
+FName USkeletonModifier::GetParentName(const FName InBoneName) const
+{
+	const int32 BoneIndex = ReferenceSkeleton->FindBoneIndex(InBoneName);
+	if (BoneIndex == INDEX_NONE)
+	{
+		return NAME_None; 
+	}
+	const TArray<FMeshBoneInfo>& BoneInfos = ReferenceSkeleton->GetRawRefBoneInfo();
+	const int32 ParentIndex = BoneInfos[BoneIndex].ParentIndex;
+	return (ParentIndex > INDEX_NONE) ? BoneInfos[ParentIndex].Name : NAME_None;
+}
+
+const FTransform& USkeletonModifier::GetTransform(const int32 InBoneIndex, const bool bGlobal) const
 {
 	if (!ReferenceSkeleton.IsValid())
 	{
diff --git a/Engine/Source/Developer/SkeletalMeshUtilitiesCommon/Public/SkeletonModifier.h b/Engine/Source/Developer/SkeletalMeshUtilitiesCommon/Public/SkeletonModifier.h
--- a/Engine/Source/Developer/SkeletalMeshUtilitiesCommon/Public/SkeletonModifier.h
+++ b/Engine/Source/Developer/SkeletalMeshUtilitiesCommon/Public/SkeletonModifier.h
@@ -72,13 +72,15 @@ struct SKELETALMESHUTILITIESCOMMON_API FOrientOptions
  * FSkeletalMeshSkeletonModifier
  */
 
-class SKELETALMESHUTILITIESCOMMON_API FSkeletonModifier
+UCLASS(BlueprintType)
+class SKELETALMESHUTILITIESCOMMON_API USkeletonModifier : public UObject
 {
+	GENERATED_BODY()
+	
 public:
 
+	UFUNCTION(BlueprintCallable, Category = "Mesh")
 	bool Init(USkeletalMesh* InSkeletalMesh);
-	
-	void ExternalUpdate(const FReferenceSkeleton& InRefSkeleton, const TArray<int32>& InIndexTracker);
 
 	/** Creates a new bone in the skeleton hierarchy at desired transform
 	 *  @param InBoneName The new bone's name.
@@ -86,7 +88,9 @@ public:
 	 *  @param InTransform The default local transform in the parent space.
 	 *  @return \c true if the operation succeeded, false otherwise. 
 	 */
+	UFUNCTION(BlueprintCallable, Category = "Skeleton")
 	bool AddBone(const FName InBoneName, const FName InParentName, const FTransform& InTransform);
+	UFUNCTION(BlueprintCallable, Category = "Skeleton")
 	bool AddBones(const TArray<FName>& InBonesName, const TArray<FName>& InParentsName, const TArray<FTransform>& InTransforms);
 
 	/** Mirror bones
@@ -94,7 +98,9 @@ public:
 	 *  @param InOptions The mirroring options
 	 *  @return \c true if the operation succeeded, false otherwise. 
 	 */
+	UFUNCTION(BlueprintCallable, Category = "Skeleton")
 	bool MirrorBone(const FName InBoneName, const FMirrorOptions& InOptions = FMirrorOptions());
+	UFUNCTION(BlueprintCallable, Category = "Skeleton")
 	bool MirrorBones(const TArray<FName>& InBonesName, const FMirrorOptions& InOptions = FMirrorOptions());
 
 	/** Sets the bone the desired local transform
@@ -103,7 +109,9 @@ public:
 	 *  @param bMoveChildren Propagate new transform to children
 	 *  @return \c true if the operation succeeded, false otherwise. 
 	 */
+	UFUNCTION(BlueprintCallable, Category = "Skeleton")
 	bool SetBoneTransform(const FName InBoneName, const FTransform& InNewTransform, const bool bMoveChildren);
+	UFUNCTION(BlueprintCallable, Category = "Skeleton")
 	bool SetBonesTransforms(const TArray<FName>& InBoneNames, const TArray<FTransform>& InNewTransforms, const bool bMoveChildren);
 
 	/** Remove a bone in the skeleton hierarchy
@@ -111,7 +119,9 @@ public:
 	 *  @param bRemoveChildren Remove children recursively.
 	 *  @return \c true if the operation succeeded, false otherwise. 
 	 */
+	UFUNCTION(BlueprintCallable, Category = "Skeleton")
 	bool RemoveBone(const FName InBoneName, const bool bRemoveChildren);
+	UFUNCTION(BlueprintCallable, Category = "Skeleton")
 	bool RemoveBones(const TArray<FName>& InBoneNames, const bool bRemoveChildren);
 
 	/** Rename bones
@@ -119,7 +129,9 @@ public:
 	 *  @param InNewBoneName The new bone's name.
 	 *  @return \c true if the operation succeeded, false otherwise. 
 	 */
+	UFUNCTION(BlueprintCallable, Category = "Skeleton")
 	bool RenameBone(const FName InOldBoneName, const FName InNewBoneName);
+	UFUNCTION(BlueprintCallable, Category = "Skeleton")
 	bool RenameBones(const TArray<FName>& InOldBoneNames, const TArray<FName>& InNewBoneNames);
 	
 	/** Parent bones
@@ -127,7 +139,9 @@ public:
 	 *  @param InParentName The new parent's name (Name_NONE to unparent).
 	 *  @return \c true if the operation succeeded, false otherwise. 
 	 */
+	UFUNCTION(BlueprintCallable, Category = "Skeleton")
 	bool ParentBone(const FName InBoneName, const FName InParentName);
+	UFUNCTION(BlueprintCallable, Category = "Skeleton")
 	bool ParentBones(const TArray<FName>& InBoneNames, const TArray<FName>& InParentNames);
 
 	/** Align bones
@@ -135,19 +149,40 @@ public:
 	 *  @param InOptions The orienting options
 	 *  @return \c true if the operation succeeded, false otherwise. 
 	 */
+	 UFUNCTION(BlueprintCallable, Category = "Skeleton")
 	bool OrientBone(const FName InBoneName, const FOrientOptions& InOptions = FOrientOptions());
+	UFUNCTION(BlueprintCallable, Category = "Skeleton")
 	bool OrientBones(const TArray<FName>& InBoneNames, const FOrientOptions& InOptions = FOrientOptions());
+
+	/** Get Bone Transform
+	 *  @param InBoneName The current bone's name.
+	 *  @param bGlobal Whether it should return the parent space or global transform
+	 *  @return \c The current bone's transform 
+	 */
+	UFUNCTION(BlueprintCallable, Category = "Skeleton")
+	FTransform GetBoneTransform(const FName InBoneName, const bool bGlobal = false) const;
+
+	/** Get Parent Name
+	 *  @param InBoneName The current bone's name.
+	 *  @return \c The current bone's parent name 
+	 */
+	UFUNCTION(BlueprintCallable, Category = "Skeleton")
+	FName GetParentName(const FName InBoneName) const;
 	
 	/**
 	 * Actually applies the skeleton modifications to the skeletal mesh.
 	 * @return true if commit succeeded.
 	 */
+	UFUNCTION(BlueprintCallable, Category = "Mesh")
 	bool CommitSkeletonToSkeletalMesh();
 
+	void ExternalUpdate(const FReferenceSkeleton& InRefSkeleton, const TArray<int32>& InIndexTracker);
+	
 	FName GetUniqueName(const FName InBoneName) const;
 	
 	const FReferenceSkeleton& GetReferenceSkeleton() const;
 	const TArray<int32>& GetBoneIndexTracker() const;
+	
 	const FTransform& GetTransform(const int32 InBoneIndex, const bool bGlobal = false) const;
 	
 private:
