Fixed crash using invalid viewport

diff --git a/Engine/Plugins/Experimental/MeshModelingToolsetExp/Source/MeshModelingToolsExp/Private/SkeletalMesh/SkeletonEditingTool.cpp b/Engine/Plugins/Experimental/MeshModelingToolsetExp/Source/MeshModelingToolsExp/Private/SkeletalMesh/SkeletonEditingTool.cpp
--- a/Engine/Plugins/Experimental/MeshModelingToolsetExp/Source/MeshModelingToolsExp/Private/SkeletalMesh/SkeletonEditingTool.cpp
+++ b/Engine/Plugins/Experimental/MeshModelingToolsetExp/Source/MeshModelingToolsExp/Private/SkeletalMesh/SkeletonEditingTool.cpp
@@ -512,10 +512,15 @@ void USkeletonEditingTool::OnTerminateDragSequence()
 
 FInputRayHit USkeletonEditingTool::CanBeginClickDragSequence(const FInputDeviceRay& ClickPos)
 {
+	FViewport* Viewport = GetToolManager()->GetContextQueriesAPI()->GetFocusedViewport();
+	if (!Viewport || !Viewport->HasFocus())
+	{
+		return FInputRayHit();
+	}
+	
 	auto PickBone = [&]() -> int32
 	{
-		FViewport* FocusedViewport = GetToolManager()->GetContextQueriesAPI()->GetFocusedViewport();
-		if (HHitProxy* HitProxy = FocusedViewport->GetHitProxy(ClickPos.ScreenPosition.X, ClickPos.ScreenPosition.Y))
+		if (HHitProxy* HitProxy = Viewport->GetHitProxy(ClickPos.ScreenPosition.X, ClickPos.ScreenPosition.Y))
 		{
 			if (TOptional<FName> OptBoneName = GetBoneName(HitProxy))
 			{
