Moving skeletal mesh editing APIs into their own module inside the skeletal mesh editing plugin.

diff --git a/Engine/Plugins/Experimental/MeshModelingToolsetExp/MeshModelingToolsetExp.uplugin b/Engine/Plugins/Experimental/MeshModelingToolsetExp/MeshModelingToolsetExp.uplugin
--- a/Engine/Plugins/Experimental/MeshModelingToolsetExp/MeshModelingToolsetExp.uplugin
+++ b/Engine/Plugins/Experimental/MeshModelingToolsetExp/MeshModelingToolsetExp.uplugin
@@ -35,6 +35,11 @@
 			"Name": "ModelingEditorUI",
 			"Type": "Editor",
 			"LoadingPhase": "Default"
+		},
+		{
+			"Name": "SkeletalMeshModifiers",
+			"Type": "Editor",
+			"LoadingPhase": "Default"
 		}
 	],
 	"Plugins": [
diff --git a/Engine/Plugins/Experimental/MeshModelingToolsetExp/Source/MeshModelingToolsExp/MeshModelingToolsExp.Build.cs b/Engine/Plugins/Experimental/MeshModelingToolsetExp/Source/MeshModelingToolsExp/MeshModelingToolsExp.Build.cs
--- a/Engine/Plugins/Experimental/MeshModelingToolsetExp/Source/MeshModelingToolsExp/MeshModelingToolsExp.Build.cs
+++ b/Engine/Plugins/Experimental/MeshModelingToolsetExp/Source/MeshModelingToolsExp/MeshModelingToolsExp.Build.cs
@@ -38,6 +38,7 @@ public class MeshModelingToolsExp : ModuleRules
 				"ModelingComponents",
 				"ModelingOperators",
 				"SkeletalMeshUtilitiesCommon",
+				"SkeletalMeshModifiers",
 				// ... add other public dependencies that you statically link with here ...
 			}
 		);
diff --git a/Engine/Plugins/Experimental/MeshModelingToolsetExp/Source/SkeletalMeshModifiers/Private/SkeletalMeshModifiersModule.cpp b/Engine/Plugins/Experimental/MeshModelingToolsetExp/Source/SkeletalMeshModifiers/Private/SkeletalMeshModifiersModule.cpp
new file mode 100644
--- /dev/null
+++ b/Engine/Plugins/Experimental/MeshModelingToolsetExp/Source/SkeletalMeshModifiers/Private/SkeletalMeshModifiersModule.cpp
@@ -0,0 +1,20 @@
+﻿// Copyright Epic Games, Inc. All Rights Reserved.
+
+#include "SkeletalMeshModifiersModule.h"
+
+#define LOCTEXT_NAMESPACE "FSkeletalMeshModifiersModule"
+
+void FSkeletalMeshModifiersModule::StartupModule()
+{
+	// This code will execute after your module is loaded into memory; the exact timing is specified in the .uplugin file per-module
+}
+
+void FSkeletalMeshModifiersModule::ShutdownModule()
+{
+	// This function may be called during shutdown to clean up your module.  For modules that support dynamic reloading,
+	// we call this function before unloading the module.
+}
+
+#undef LOCTEXT_NAMESPACE
+	
+IMPLEMENT_MODULE(FSkeletalMeshModifiersModule, SkeletalMeshModifiers)
\ No newline at end of file
diff --git a/Engine/Plugins/Experimental/MeshModelingToolsetExp/Source/SkeletalMeshModifiers/Private/SkeletonModifier.cpp b/Engine/Plugins/Experimental/MeshModelingToolsetExp/Source/SkeletalMeshModifiers/Private/SkeletonModifier.cpp
new file mode 100644
--- /dev/null
+++ b/Engine/Plugins/Experimental/MeshModelingToolsetExp/Source/SkeletalMeshModifiers/Private/SkeletonModifier.cpp
@@ -0,0 +1,1260 @@
+﻿// Copyright Epic Games, Inc. All Rights Reserved.
+
+
+#include "SkeletonModifier.h"
+
+#include "BoneWeights.h"
+#include "LODUtilities.h"
+#include "MeshDescription.h"
+#include "Animation/Skeleton.h"
+#include "ReferenceSkeleton.h"
+#include "SkeletalMeshAttributes.h"
+#include "Animation/AnimSequence.h"
+#include "Animation/PoseAsset.h"
+#include "Animation/AnimMontage.h"
+#include "AssetRegistry/AssetRegistryModule.h"
+#include "AssetRegistry/AssetData.h"
+#include "Rendering/SkeletalMeshModel.h"
+
+
+namespace USkeletonModifierLocals
+{
+	static constexpr int32 LODIndex = 0;
+}
+
+FTransform FMirrorOptions::MirrorTransform(const FTransform& InTransform) const
+{
+	FTransform Transform = InTransform;
+	Transform.SetLocation(MirrorVector(Transform.GetLocation()));
+
+	if (bMirrorRotation)
+	{
+		FRotator Rotator(ForceInitToZero);
+		switch (MirrorAxis)
+		{
+		case EAxis::X:
+			{
+				Rotator.Roll = 180;
+				break;
+			}
+		case EAxis::Y:
+			{
+				Rotator.Pitch = 180;
+				break;
+			}
+		case EAxis::Z:
+			{
+				Rotator.Yaw = 180;
+				break;
+			}
+		default:
+			break;
+		}	
+		Transform.SetRotation(FQuat::MakeFromRotator(Rotator) * Transform.GetRotation());
+	}
+
+	return Transform;
+}
+
+FVector FMirrorOptions::MirrorVector(const FVector& InVector) const
+{
+	FVector Axis(ForceInitToZero); Axis.SetComponentForAxis(MirrorAxis, 1.f);
+	return InVector.MirrorByVector(Axis);
+}
+
+FTransform FOrientOptions::OrientTransform(const FVector& InPrimaryTarget, const FTransform& InTransform) const
+{
+	if (Primary == EAxis::None || InPrimaryTarget.IsNearlyZero())
+	{
+		return InTransform;
+	}
+
+	FTransform Transform = InTransform;
+
+	FVector PrimaryAxis(ForceInitToZero); PrimaryAxis.SetComponentForAxis(Primary, 1.f);
+	const FVector PrimaryTarget = InPrimaryTarget.GetSafeNormal();
+	
+	// orient primary axis towards InPrimaryTarget
+	{
+		const FVector Axis = Transform.TransformVectorNoScale(PrimaryAxis).GetSafeNormal();
+		const FQuat Rotation = FQuat::FindBetweenNormals(Axis, PrimaryTarget);
+		const FQuat NewRotation = (Rotation * Transform.GetRotation()).GetNormalized();
+		Transform.SetRotation(NewRotation);
+	}
+
+	if (Secondary == EAxis::None || SecondaryTarget.IsNearlyZero())
+	{ // no need to use secondary axis
+		return Transform;
+	}
+
+	FVector Target = SecondaryTarget.GetSafeNormal();
+	if (FMath::IsNearlyEqual(FMath::Abs(FVector::DotProduct(PrimaryTarget, Target)), 1.0f))
+	{ // both targets are parallel
+		return Transform;
+	}
+
+	// orient secondary axis towards SecondaryDirection
+	{
+		FVector SecondaryAxis(ForceInitToZero); SecondaryAxis.SetComponentForAxis(Secondary, 1.f);
+		
+		{
+			const FVector Axis = Transform.TransformVectorNoScale(PrimaryAxis).GetSafeNormal();
+			Target = Target - FVector::DotProduct(Target, Axis) * Axis;
+		}
+
+		if (!Target.IsNearlyZero())
+		{
+			Target = Target.GetSafeNormal();
+			const FVector Axis = Transform.TransformVectorNoScale(SecondaryAxis).GetSafeNormal();
+			const FQuat Rotation = FQuat::FindBetweenNormals(Axis, Target);
+			const FQuat NewRotation = (Rotation * Transform.GetRotation()).GetNormalized();
+			Transform.SetRotation(NewRotation);
+		}
+	}
+	
+	return Transform;
+}
+
+void USkeletonModifier::ExternalUpdate(const FReferenceSkeleton& InRefSkeleton, const TArray<int32>& InIndexTracker)
+{
+	if (!ReferenceSkeleton)
+	{
+		return;
+	}
+	
+	*ReferenceSkeleton = InRefSkeleton;
+	TransformComposer.Reset(new FTransformComposer(*ReferenceSkeleton));
+	BoneIndexTracker = InIndexTracker;
+}
+
+bool USkeletonModifier::Init(USkeletalMesh* InSkeletalMesh)
+{
+	SkeletalMesh = nullptr;
+	MeshDescription.Reset();
+	ReferenceSkeleton.Reset();
+	TransformComposer.Reset();
+	BoneIndexTracker.Reset();
+
+#if WITH_EDITORONLY_DATA
+	// validate supplied skeletal mesh exists
+	if (!InSkeletalMesh)
+	{
+		UE_LOG(LogAnimation, Error, TEXT("Skeleton Modifier: No skeletal mesh supplied to load."));
+		return false;
+	}
+
+	const USkeleton* Skeleton = InSkeletalMesh->GetSkeleton();
+	if (!Skeleton)
+	{
+		UE_LOG(LogAnimation, Error, TEXT("Skeleton Modifier: Skeletal Mesh supplied has no skeleton."));
+		return false;
+	}
+
+	// verify user is not trying to modify one of the core engine assets
+	if (InSkeletalMesh->GetPathName().StartsWith(TEXT("/Engine/")))
+	{
+		UE_LOG(LogAnimation, Error, TEXT("Skeleton Modifier: Cannot modify built-in engine asset."));
+		return false;
+	}
+	
+	// check assets using this skeleton ?
+	{
+		// TODO avoid certain changes when the skeleton is referenced by other assets (i.e. changing the skeletal mesh's
+		// reference skeleton poses is fine, re-parenting/removing bones, etc. is not)
+		static const TArray<FTopLevelAssetPath> AssetPaths({
+		   UAnimSequence::StaticClass()->GetClassPathName(),
+		   UAnimMontage::StaticClass()->GetClassPathName(),
+		   UPoseAsset::StaticClass()->GetClassPathName(),
+		   USkeletalMesh::StaticClass()->GetClassPathName()});
+
+		FARFilter Filter;
+		Filter.ClassPaths = AssetPaths;
+
+		const IAssetRegistry& AssetRegistry = FModuleManager::LoadModuleChecked<FAssetRegistryModule>("AssetRegistry").Get();
+		TArray<FAssetData> Assets; AssetRegistry.GetAssets(Filter, Assets);
+
+		const FAssetData SkeletonAssetData(Skeleton);
+		const FString SkeletonPath = SkeletonAssetData.GetExportTextName();
+	
+		static const FName Tag("Skeleton");
+		for (const FAssetData& AssetData: Assets)
+		{
+			const FString TagValue = AssetData.GetTagValueRef<FString>(Tag);
+			if (TagValue == SkeletonPath)
+			{
+				UE_LOG(LogAnimation, Warning, TEXT("%s references that skeleton."), *AssetData.GetExportTextName());
+			}
+		}
+	}
+
+	// store pointer to mesh and instantiate a mesh description for commiting changes
+	SkeletalMesh = InSkeletalMesh;
+
+	// store mesh description to edit
+	MeshDescription = MakeUnique<FMeshDescription>();
+	SkeletalMesh->GetMeshDescription(USkeletonModifierLocals::LODIndex, *MeshDescription);
+	if (MeshDescription->IsEmpty())
+	{
+		UE_LOG(LogAnimation, Error, TEXT("Skeleton Modifier: mesh description is emtpy."));
+		return false;
+	}
+
+	// store reference skeleton to edit
+	ReferenceSkeleton = MakeUnique<FReferenceSkeleton>();
+	*ReferenceSkeleton = SkeletalMesh->GetRefSkeleton();
+
+	TransformComposer.Reset(new FTransformComposer(*ReferenceSkeleton));
+
+	// store initial bones indices to track for changes
+	const int32 NumBones = ReferenceSkeleton->GetRawBoneNum();
+	BoneIndexTracker.Reserve(NumBones);
+	for (int32 Index = 0; Index < NumBones; Index++)
+	{
+		BoneIndexTracker.Add(Index);
+	}
+	return true;
+#else
+	ensureMsgf(false, TEXT("Skeleton Modifier is an editor only feature."));
+#endif
+	return false;
+}
+
+bool USkeletonModifier::IsReferenceSkeletonValid(const bool bLog) const
+{
+	if (!ReferenceSkeleton.IsValid())
+	{
+		if(bLog)
+		{
+			UE_LOG(LogAnimation, Error, TEXT("Skeleton Modifier: No valid reference skeleton provided."));
+		}
+		return false;		
+	}
+	return true;
+}
+
+bool USkeletonModifier::CommitSkeletonToSkeletalMesh()
+{
+	if (!SkeletalMesh || !ReferenceSkeleton || !MeshDescription)
+	{
+		UE_LOG(LogAnimation, Error, TEXT("Skeleton Modifier: No mesh loaded. Cannot apply skeleton edits."));
+		return false;
+	}
+
+#if WITH_EDITORONLY_DATA
+	// before commiting, we have to reparent non-root bones with no parent as the animation pipeline
+	// doesn't support multi-roots
+	const TArray<FMeshBoneInfo>& BoneInfos = ReferenceSkeleton->GetRawRefBoneInfo();
+	if (!BoneInfos.IsEmpty())
+	{
+		TArray<FName> BonesToParent;
+		for (int32 Index = 1; Index < BoneInfos.Num(); Index++)
+		{
+			const FMeshBoneInfo& BoneInfo = BoneInfos[Index];
+			if (BoneInfo.ParentIndex == INDEX_NONE)
+			{
+				BonesToParent.Add(BoneInfo.Name);
+			}
+		}
+		
+		if (!BonesToParent.IsEmpty())
+		{
+			for (const FName& BoneName: BonesToParent)
+			{
+				UE_LOG(LogAnimation, Warning, TEXT("Skeleton Modifier: %s will be parented to the root bone before commiting."), *BoneName.ToString());
+			}
+			ParentBones(BonesToParent, {BoneInfos[0].Name});
+		}
+	}
+	
+	// update mesh description
+	{
+		FSkeletalMeshAttributes MeshAttributes(*MeshDescription);
+
+		// update bone data
+		if (!MeshAttributes.HasBones())
+		{
+			MeshAttributes.Register(true);
+		}
+
+		MeshAttributes.Bones().Reset(BoneInfos.Num());
+
+		FSkeletalMeshAttributes::FBoneNameAttributesRef BoneNames = MeshAttributes.GetBoneNames();
+		FSkeletalMeshAttributes::FBoneParentIndexAttributesRef BoneParentIndices = MeshAttributes.GetBoneParentIndices();
+		FSkeletalMeshAttributes::FBonePoseAttributesRef BonePoses = MeshAttributes.GetBonePoses();
+
+		const TArray<FTransform> Transforms = ReferenceSkeleton->GetRawRefBonePose();
+		for (int Index = 0; Index < BoneInfos.Num(); ++Index)
+		{
+			const FMeshBoneInfo& Info = BoneInfos[Index];
+			const FBoneID BoneID = MeshAttributes.CreateBone();
+			BoneNames.Set(BoneID, Info.Name);
+			BoneParentIndices.Set(BoneID, Info.ParentIndex);
+			BonePoses.Set(BoneID, Transforms[Index]);
+		}
+		
+		// update skin data if needed
+		auto HasBoneIndexesChanged = [&]()
+		{
+			for (int32 Index = 0; Index < BoneIndexTracker.Num(); Index++)
+			{
+				if (BoneIndexTracker[Index] != Index)
+				{
+					return true;
+				}
+			}
+			return false;
+		};
+		if (HasBoneIndexesChanged())
+		{
+			using namespace UE::AnimationCore;
+			FBoneWeightsSettings BoneSettings; BoneSettings.SetNormalizeType(EBoneWeightNormalizeType::None);
+		
+			FSkinWeightsVertexAttributesRef SkinWeights = MeshAttributes.GetVertexSkinWeights();
+			for (const FVertexID& VertexID: MeshDescription->Vertices().GetElementIDs())
+			{
+				FVertexBoneWeights BoneWeights = SkinWeights.Get(VertexID);
+				if (const int32 NumBoneWeights = BoneWeights.Num())
+				{
+					TArray<FBoneWeight> NewWeights;
+					for (int32 Idx = 0; Idx < NumBoneWeights; ++Idx)
+					{
+						const FBoneWeight& OldBoneWeight = BoneWeights[Idx];
+						const int32 BoneIndex = OldBoneWeight.GetBoneIndex();
+					
+						check(BoneIndexTracker.IsValidIndex(BoneIndex));
+					
+						const int32 NewBoneIndex = BoneIndexTracker[BoneIndex];
+						if (NewBoneIndex != INDEX_NONE)
+						{
+							NewWeights.Add(FBoneWeight(NewBoneIndex, OldBoneWeight.GetRawWeight()));
+						}
+					}
+					SkinWeights.Set(VertexID, FBoneWeights::Create(NewWeights, BoneSettings));
+				}
+			}
+		}
+	}
+
+	// update skeletal mesh
+	FlushRenderingCommands();
+	
+	SkeletalMesh->SetFlags(RF_Transactional);
+	SkeletalMesh->Modify();
+
+	// update the ref skeleton
+	SkeletalMesh->SetRefSkeleton(*ReferenceSkeleton);
+	SkeletalMesh->GetRefBasesInvMatrix().Reset();
+	SkeletalMesh->CalculateInvRefMatrices();
+	
+	// update skeletal mesh LOD (cf. USkeletalMesh::CommitMeshDescription)
+	SkeletalMesh->CommitMeshDescription(USkeletonModifierLocals::LODIndex, *MeshDescription);
+
+	SkeletalMesh->PostEditChange();
+
+	// update skeleton
+	USkeleton* Skeleton = SkeletalMesh->GetSkeleton();
+	Skeleton->Modify();
+	if (Skeleton->RecreateBoneTree(SkeletalMesh))
+	{
+		Skeleton->MarkPackageDirty();	
+	}
+
+	return true;
+#else
+	ensureMsgf(false, TEXT("Skeleton Modifier is an editor only feature."));
+#endif
+	return false;
+}
+
+bool USkeletonModifier::AddBone(const FName InBoneName, const FName InParentName, const FTransform& InTransform)
+{
+	if (InBoneName == NAME_None)
+	{
+		UE_LOG(LogAnimation, Error, TEXT("Skeleton Modifier - Add: Cannot add bone with no name."));
+		return false;
+	}
+	
+	return AddBones({InBoneName}, {InParentName}, {InTransform});
+}
+
+bool USkeletonModifier::AddBones(
+	const TArray<FName>& InBoneNames, const TArray<FName>& InParentNames, const TArray<FTransform>& InTransforms)
+{
+	if (!IsReferenceSkeletonValid())
+	{
+		return false;
+	}
+	
+	const int32 NumBonesToAdd = InBoneNames.Num();
+	if (NumBonesToAdd == 0)
+	{
+		UE_LOG(LogAnimation, Error, TEXT("Skeleton Modifier - Add: The provided bone names array is empty."));
+		return false;
+	}
+	
+	struct FBoneData
+	{
+		FMeshBoneInfo BoneInfo;
+		int32 TransformOffset;
+	};
+
+	TArray<FBoneData> BonesToAdd;
+	BonesToAdd.Reserve(NumBonesToAdd);
+
+	auto GetParentName = [&](int32 Index)
+	{
+		if (InBoneNames.Num() == InParentNames.Num())
+		{
+			return InParentNames[Index];
+		}
+		return InParentNames.IsEmpty() ? NAME_None : InParentNames[0];  
+	};
+
+	const int NumBonesBefore = ReferenceSkeleton->GetRawBoneNum();
+
+	for (int32 Index = 0; Index < NumBonesToAdd; Index++)
+	{
+		const FName& BoneName = InBoneNames[Index];		
+		if (ReferenceSkeleton->FindBoneIndex(InBoneNames[Index]) == INDEX_NONE)
+		{
+			const FName ParentName = GetParentName(Index);
+
+			// look for parent index in the ref skeleton
+			int32 ParentIndex = ReferenceSkeleton->FindBoneIndex(ParentName);
+			if (ParentIndex == INDEX_NONE && Index > 0)
+			{
+				// otherwise, check if one of the new bone is going to be the parent
+				ParentIndex = InBoneNames.IndexOfByKey(ParentName);
+
+				if (ParentIndex > INDEX_NONE && ParentIndex < Index)
+				{
+					ParentIndex += NumBonesBefore;
+				}
+			}
+			const FMeshBoneInfo NewBoneInfo(BoneName, BoneName.ToString(), ParentIndex);
+			BonesToAdd.Add({NewBoneInfo, Index});
+		}
+	}
+
+	if (BonesToAdd.IsEmpty())
+	{
+		UE_LOG(LogAnimation, Error, TEXT("Skeleton Modifier - Add: None of the provided names is avalable to be added."));
+		return false;
+	}
+
+	auto GetTransform = [&](int32 Index)
+	{
+		if (InBoneNames.Num() == InTransforms.Num())
+		{
+			return InTransforms[Index];
+		}
+		return InTransforms.IsEmpty() ? FTransform::Identity : InTransforms[0];  
+	};
+	
+	//update reference skeleton
+	{
+		static constexpr bool bAllowMultipleRoots = true;
+		FReferenceSkeletonModifier Modifier(*ReferenceSkeleton, nullptr);
+		for (FBoneData& BoneData: BonesToAdd)
+		{
+			Modifier.Add(BoneData.BoneInfo, GetTransform(BoneData.TransformOffset), bAllowMultipleRoots);
+		}
+	}
+
+	// invalidate composer
+	TransformComposer->Invalidate(INDEX_NONE);
+	
+	// update index tracker: nothing to do as those new indices do not represent any bone in the initial skinning data
+	
+	return true;
+}
+
+bool USkeletonModifier::MirrorBone(const FName InBoneName, const FMirrorOptions& InOptions)
+{
+	if (InBoneName == NAME_None)
+	{
+		UE_LOG(LogAnimation, Error, TEXT("Skeleton Modifier - Mirror: Cannot mirror bone with no name."));
+		return false;
+	}
+	
+	return MirrorBones({InBoneName}, InOptions);
+}
+
+bool USkeletonModifier::MirrorBones(const TArray<FName>& InBonesName, const FMirrorOptions& InOptions)
+{
+	if (!IsReferenceSkeletonValid())
+	{
+		return false;
+	}
+
+	// get bones to mirror
+	TArray<int32> BonesToMirror; GetBonesToMirror(InBonesName, InOptions, BonesToMirror);
+	
+	const int32 NumBonesToMirror = BonesToMirror.Num();
+	if (NumBonesToMirror == 0)
+	{
+		UE_LOG(LogAnimation, Error, TEXT("Skeleton Modifier - Mirror: None of the provided names has been found."));
+		return false;
+	}
+
+	// get mirrored names
+	TArray<FName> MirroredNames; GetMirroredNames(BonesToMirror, InOptions, MirroredNames);
+
+	// add bones first if they are missing
+	TArray<int32> MirroredBones; GetMirroredBones(BonesToMirror, MirroredNames, MirroredBones);
+	if (MirroredBones.Num() != NumBonesToMirror)
+	{
+		UE_LOG(LogAnimation, Error, TEXT("Skeleton Modifier - Mirror: Couldn't find mirrored bones."));
+		return false;
+	}
+
+	// compute mirrored transforms
+	TArray<FTransform> MirroredTransforms; GetMirroredTransforms(BonesToMirror, MirroredBones, InOptions, MirroredTransforms);
+
+	// update reference skeleton
+	{
+		FReferenceSkeletonModifier Modifier(*ReferenceSkeleton, nullptr);
+		for (int32 Index = 0; Index < NumBonesToMirror; Index++)
+		{
+			Modifier.UpdateRefPoseTransform(MirroredBones[Index], MirroredTransforms[Index]);
+		}
+	}
+
+	// invalidate composer
+	TransformComposer->Invalidate(INDEX_NONE);
+	
+	return true;
+}
+
+void USkeletonModifier::GetBonesToMirror(
+	const TArray<FName>& InBonesName, const FMirrorOptions& InOptions, TArray<int32>& OutBonesToMirror) const
+{
+	OutBonesToMirror.Reset();
+	
+	TSet<int32> IndicesToMirror;
+	auto GetBonesToMirror = [&](const int32 BoneIndex, auto&& GetBonesToMirror2) -> void
+	{
+		if (BoneIndex == INDEX_NONE)
+		{
+			return;
+		}
+		
+		IndicesToMirror.Add(BoneIndex);
+
+		if (InOptions.bMirrorChildren)
+		{
+			TArray<int32> Children; ReferenceSkeleton->GetDirectChildBones(BoneIndex, Children);
+			for (int32 ChildIndex: Children)
+			{
+				GetBonesToMirror2(ChildIndex, GetBonesToMirror2);
+			}
+		}
+	};
+	
+	for (const FName& BoneName: InBonesName)
+	{
+		GetBonesToMirror(ReferenceSkeleton->FindRawBoneIndex(BoneName), GetBonesToMirror);
+	}
+
+	const int32 NumBonesToMirror = IndicesToMirror.Num();
+	if (NumBonesToMirror == 0)
+	{
+		return;
+	}
+
+	IndicesToMirror.Sort([](const int32 Index0, const int32 Index1) {return Index0 < Index1;});
+	OutBonesToMirror = IndicesToMirror.Array();
+}
+
+void USkeletonModifier::GetMirroredNames(
+	const TArray<int32>& InBonesToMirror, const FMirrorOptions& InOptions, TArray<FName>& OutBonesName) const
+{
+	OutBonesName.Reset();
+	if (InBonesToMirror.IsEmpty())
+	{
+		return;
+	}
+	
+	const TArray<FMeshBoneInfo>& BoneInfos = ReferenceSkeleton->GetRawRefBoneInfo();
+	
+	OutBonesName.Reserve(InBonesToMirror.Num());
+	Algo::Transform(InBonesToMirror, OutBonesName, [&](const int32 BoneIndex)
+	{
+		FString BoneStr = BoneInfos[BoneIndex].Name.ToString();
+		if (BoneStr.Contains(InOptions.LeftString, ESearchCase::CaseSensitive, ESearchDir::FromEnd))
+		{
+			BoneStr = BoneStr.Replace(*InOptions.LeftString, *InOptions.RightString, ESearchCase::CaseSensitive);
+		}
+		else if (BoneStr.Contains(InOptions.RightString, ESearchCase::CaseSensitive, ESearchDir::FromEnd))
+		{
+			BoneStr = BoneStr.Replace(*InOptions.RightString, *InOptions.LeftString, ESearchCase::CaseSensitive);
+		}
+		else
+		{
+			return GetUniqueName(BoneInfos[BoneIndex].Name);
+		}
+		return FName(*BoneStr);
+	});
+}
+
+void USkeletonModifier::GetMirroredBones(
+	const TArray<int32>& InBonesToMirror, const TArray<FName>& InMirroredNames, TArray<int32>& OutMirroredBones)
+{
+	OutMirroredBones.Reset();
+	if (InBonesToMirror.IsEmpty() || InBonesToMirror.Num() != InMirroredNames.Num())
+	{
+		return;
+	}
+	
+	const TArray<FMeshBoneInfo>& BoneInfos = ReferenceSkeleton->GetRawRefBoneInfo();
+	const TArray<FTransform>& BoneTransforms = ReferenceSkeleton->GetRawRefBonePose();
+	
+	TArray<FName> BonesToAdd, ParentNames;
+	TArray<FTransform> Transforms;
+	const int32 NumBones = InBonesToMirror.Num();
+	for (int32 Index = 0; Index < NumBones; Index++)
+	{
+		const FName& MirroredName = InMirroredNames[Index];
+		const int32 MirroredIndex = ReferenceSkeleton->FindRawBoneIndex(MirroredName);
+		if (MirroredIndex == INDEX_NONE)
+		{
+			const int32 RefBoneIndex = InBonesToMirror[Index];
+				
+			// name
+			BonesToAdd.Add(MirroredName);
+
+			// parent
+			int32 ParentIndex = BoneInfos[RefBoneIndex].ParentIndex;
+			FName ParentName = NAME_None;
+			if (ParentIndex != INDEX_NONE)
+			{
+				ParentName = BoneInfos[ParentIndex].Name;
+				// is that parent being mirrored?
+				const int32 ParentIndexInMirrored = InBonesToMirror.IndexOfByKey(ParentIndex);
+				if (ParentIndexInMirrored != INDEX_NONE)
+				{
+					ParentName = InMirroredNames[ParentIndexInMirrored];
+				}
+			}
+			ParentNames.Add(ParentName);
+
+			// transform
+			Transforms.Add(BoneTransforms[RefBoneIndex]);
+		}
+		else
+		{
+			OutMirroredBones.Add(MirroredIndex);
+		}
+	}
+
+	if (BonesToAdd.IsEmpty())
+	{
+		return;
+	}
+
+	// add missing bones and get their index
+	OutMirroredBones.Reset();
+	if (!AddBones(BonesToAdd, ParentNames, Transforms))
+	{
+		return;
+	}
+
+	OutMirroredBones.Reserve(NumBones);
+	for (int32 Index = 0; Index < NumBones; Index++)
+	{
+		const FName& MirroredName = InMirroredNames[Index];
+		const int32 MirroredIndex = ReferenceSkeleton->FindRawBoneIndex(MirroredName);
+		if (MirroredIndex == INDEX_NONE)
+		{
+			OutMirroredBones.Reset();
+			return;
+		}
+		OutMirroredBones.Add(MirroredIndex);
+	}
+}
+
+void USkeletonModifier::GetMirroredTransforms(
+	const TArray<int32>& InBonesToMirror, const TArray<int32>& InMirroredBones,
+	const FMirrorOptions& InOptions, TArray<FTransform>& OutMirroredTransforms) const
+{
+	OutMirroredTransforms.Reset();
+	
+	const int32 NumBonesToMirror = InBonesToMirror.Num();
+	if (NumBonesToMirror == 0)
+	{
+		return;
+	}
+	
+	const TArray<FMeshBoneInfo>& BoneInfos = ReferenceSkeleton->GetRawRefBoneInfo();
+	auto FindFirstNotMirroredParent = [&](const int32 RefBoneIndex) -> int32
+	{
+		if (RefBoneIndex == INDEX_NONE)
+		{
+			return INDEX_NONE;
+		}
+		
+		int32 RefParentIndex = BoneInfos[RefBoneIndex].ParentIndex;
+		if (RefParentIndex == INDEX_NONE)
+		{
+			return INDEX_NONE;
+		}
+
+		int32 ParentMirroredIndex = InBonesToMirror.IndexOfByKey(RefParentIndex);
+		while (ParentMirroredIndex != INDEX_NONE)
+		{
+			RefParentIndex = BoneInfos[RefParentIndex].ParentIndex;
+			ParentMirroredIndex = InBonesToMirror.IndexOfByKey(RefParentIndex);
+		}
+		return RefParentIndex;
+	};
+
+	// compute global mirrored transforms
+	TArray<FTransform> MirroredGlobal; MirroredGlobal.Reserve(NumBonesToMirror);
+	for (int32 Index = 0; Index < NumBonesToMirror; Index++)
+	{
+		// bone global
+		const int32 RefBoneIndex = InBonesToMirror[Index];
+		FTransform Global = TransformComposer->GetGlobalTransform(RefBoneIndex);
+	
+		// first parent not mirrored global
+		const int32 FirstNotMirroredParent = FindFirstNotMirroredParent(RefBoneIndex);
+		const FTransform& FirstNotMirroredGlobal = TransformComposer->GetGlobalTransform(FirstNotMirroredParent);
+
+		// switch to first parent not mirrored (translation only)
+		Global.AddToTranslation(-FirstNotMirroredGlobal.GetTranslation());
+
+		// mirror
+		Global = InOptions.MirrorTransform(Global);
+
+		// switch back to global (translation only)
+		Global.AddToTranslation(FirstNotMirroredGlobal.GetTranslation());
+
+		MirroredGlobal.Add(Global);
+	}
+
+	// switch back to local
+	OutMirroredTransforms.Reserve(NumBonesToMirror);
+	for (int32 Index = 0; Index < NumBonesToMirror; Index++)
+	{
+		const int32 RefBoneIndex = InBonesToMirror[Index];		
+		const int32 RefParentIndex = BoneInfos[RefBoneIndex].ParentIndex;
+		const int32 ParentMirroredIndex = InBonesToMirror.IndexOfByKey(RefParentIndex);
+		const int32 ParentIndex = BoneInfos[InMirroredBones[Index]].ParentIndex;
+		const FTransform& ParentGlobal = ParentMirroredIndex != INDEX_NONE ? MirroredGlobal[ParentMirroredIndex] :
+							TransformComposer->GetGlobalTransform(ParentIndex);
+		OutMirroredTransforms.Add(MirroredGlobal[Index].GetRelativeTransform(ParentGlobal));
+	}
+}
+
+// NOTE : that function might take a bUpdateChildren to decide whether we want to compensate the children transforms
+// atm, we update the bone's local ref transform so children's global transforms are changed (we just need to cache the
+// global transforms then restore them back)
+// orienting the bone for example should change the children's global transform
+bool USkeletonModifier::SetBoneTransform( const FName InBoneName, const FTransform& InNewTransform, const bool bMoveChildren)
+{
+	if (InBoneName == NAME_None)
+	{
+		UE_LOG(LogAnimation, Error, TEXT("Skeleton Modifier - Move: Cannot move bone with no name."));
+		return false;
+	}
+	
+	return SetBonesTransforms({InBoneName}, {InNewTransform}, bMoveChildren);
+}
+
+bool USkeletonModifier::SetBonesTransforms(
+	const TArray<FName>& InBoneNames, const TArray<FTransform>& InNewTransforms, const bool bMoveChildren)
+{
+	if (!IsReferenceSkeletonValid())
+	{
+		return false;
+	}
+	
+	const int32 NumBonesToMove = InBoneNames.Num();
+	if (NumBonesToMove == 0 || NumBonesToMove != InNewTransforms.Num())
+	{
+		UE_LOG(LogAnimation, Error, TEXT("Skeleton Modifier - Move: Discrepancy between bones and transforms (%d / %d)."), NumBonesToMove, InNewTransforms.Num());
+		return false;
+	}
+	
+	TArray<int32> BoneIndices, Offsets;
+	BoneIndices.Reserve(NumBonesToMove);
+	Offsets.Reserve(NumBonesToMove);
+
+	for (int32 Index = 0; Index < NumBonesToMove; Index++)
+	{
+		const int32 BoneIndex = ReferenceSkeleton->FindRawBoneIndex(InBoneNames[Index]);
+		if (BoneIndex != INDEX_NONE)
+		{
+			BoneIndices.Add(BoneIndex);
+			Offsets.Add(Index);
+		}
+	}
+
+	if (BoneIndices.IsEmpty())
+	{
+		UE_LOG(LogAnimation, Error, TEXT("Skeleton Modifier - Move: None of the provided bones has been found."));
+		return false;
+	}
+
+	// compute global transforms if needed
+	TArray<int32> ChildrenToFix;
+	TArray<FTransform> GlobalTransforms;
+	
+	if (!bMoveChildren)
+	{
+		// get children
+		for (int32 Index = 0; Index < BoneIndices.Num(); Index++)
+		{
+			TArray<int32> Children;
+			ReferenceSkeleton->GetDirectChildBones(BoneIndices[Index], Children);
+			for (int32 ChildIndex: Children)
+			{
+				if (!BoneIndices.Contains(ChildIndex))
+				{
+					ChildrenToFix.Add(ChildIndex);
+				}
+			}
+		}
+
+		// sort them from highest index to lowest
+		ChildrenToFix.Sort([](const int32 Index0, const int32 Index1) {return Index0 > Index1;} );
+		const int32 NumChildren = ChildrenToFix.Num();
+
+		// compute global transforms (note that we could cache them for faster implementation) 
+		GlobalTransforms.AddUninitialized(NumChildren);
+
+		for (int32 Index = 0; Index < NumChildren; Index++)
+		{
+			GlobalTransforms[Index] = TransformComposer->GetGlobalTransform(ChildrenToFix[Index]);
+		}
+	}
+	
+	// update reference skeleton
+	{
+		FReferenceSkeletonModifier Modifier(*ReferenceSkeleton, nullptr);
+		for (int32 Index = 0; Index < BoneIndices.Num(); Index++)
+		{
+			Modifier.UpdateRefPoseTransform(BoneIndices[Index], InNewTransforms[Offsets[Index]]);
+
+			// invalidate cached global transform
+			TransformComposer->Invalidate(BoneIndices[Index]);
+		}
+
+		if (!bMoveChildren)
+		{
+			const int32 NumChildren = ChildrenToFix.Num();
+			for (int32 Index = 0; Index < NumChildren; Index++)
+			{
+				const int32 ChildrenIndex = ChildrenToFix[Index];
+				const int32 ParentIndex = ReferenceSkeleton->GetRawParentIndex(ChildrenIndex);
+				const FTransform& NewParentGlobal = TransformComposer->GetGlobalTransform(ParentIndex);
+				const FTransform& NewLocal = GlobalTransforms[Index].GetRelativeTransform(NewParentGlobal);
+				Modifier.UpdateRefPoseTransform(ChildrenIndex, NewLocal);
+				TransformComposer->Invalidate(ChildrenIndex);
+			}
+		}
+	}
+
+	// update index tracker: no modification on bone indices to track when changing transforms
+	
+	return true;
+}
+
+bool USkeletonModifier::RemoveBone(const FName InBoneName, const bool bRemoveChildren)
+{
+	if (InBoneName == NAME_None)
+	{
+		UE_LOG(LogAnimation, Error, TEXT("Skeleton Modifier - Remove: Cannot remove bone with no name."));
+		return false;
+	}
+	
+	return RemoveBones({InBoneName}, bRemoveChildren);
+}
+
+bool USkeletonModifier::RemoveBones(const TArray<FName>& InBoneNames, const bool bRemoveChildren)
+{
+	if (!IsReferenceSkeletonValid())
+	{
+		return false;
+	}
+
+	if (InBoneNames.IsEmpty())
+	{
+		UE_LOG(LogAnimation, Error, TEXT("Skeleton Modifier - Remove: No bone provided."));
+		return false;
+	}
+	
+	// store initial data
+	const TArray<FMeshBoneInfo> InfosBeforeRemoval = ReferenceSkeleton->GetRawRefBoneInfo();
+
+	// update reference skeleton
+	{
+		FReferenceSkeletonModifier Modifier(*ReferenceSkeleton, nullptr);
+		for (const FName& BoneName: InBoneNames)
+		{
+			Modifier.Remove(BoneName, bRemoveChildren);
+		}
+	}
+
+	if (InfosBeforeRemoval.Num() == ReferenceSkeleton->GetRawRefBoneInfo().Num())
+	{ // no bone has been removed
+		UE_LOG(LogAnimation, Error, TEXT("Skeleton Modifier - Remove: No bone has been removed."));
+		return false;
+	}
+	
+	// invalidate composer
+	TransformComposer->Invalidate(INDEX_NONE);
+
+	// update index tracker
+	UpdateBoneTracker(InfosBeforeRemoval);
+
+	return true;		
+}
+
+bool USkeletonModifier::RenameBone(const FName InOldBoneName, const FName InNewBoneName)
+{
+	if (InOldBoneName == NAME_None || InOldBoneName == NAME_None || InNewBoneName == InOldBoneName)
+	{
+		UE_LOG(LogAnimation, Error, TEXT("Skeleton Modifier - Rename: cannot rename %s with %s."), *InOldBoneName.ToString(), *InNewBoneName.ToString());
+		return false;
+	}
+	
+	return RenameBones({InOldBoneName}, {InNewBoneName});
+}
+
+bool USkeletonModifier::RenameBones(const TArray<FName>& InOldBoneNames, const TArray<FName>& InNewBoneNames)
+{
+	if (!IsReferenceSkeletonValid())
+	{
+		return false;
+	}
+	
+	if (InOldBoneNames.IsEmpty() || InNewBoneNames.Num() != InOldBoneNames.Num())
+	{
+		UE_LOG(LogAnimation, Error, TEXT("Skeleton Modifier - Rename: Discrepancy between old and new names (%d / %d)."), InOldBoneNames.Num(), InNewBoneNames.Num());
+		return false;
+	}
+
+	// update reference skeleton
+	{
+		FReferenceSkeletonModifier Modifier(*ReferenceSkeleton, nullptr);
+
+		const int32 NumBonesToRename = InOldBoneNames.Num();
+		for (int32 Index = 0; Index < NumBonesToRename; ++Index)
+		{
+			const FName& OldName = InOldBoneNames[Index];
+			const FName& NewName = InNewBoneNames[Index];
+			if (OldName != NAME_None && OldName != NewName)
+			{
+				Modifier.Rename(OldName, NewName);
+			}
+		}
+	}
+
+	// update index tracker: no modification on bone indices to track when renaming
+
+	return true;
+}
+
+bool USkeletonModifier::ParentBone(const FName InBoneName, const FName InParentName)
+{
+	if (InBoneName == NAME_None)
+	{
+		UE_LOG(LogAnimation, Error, TEXT("Skeleton Modifier - Parent: Cannot parent a bone with no name."));
+		return false;
+	}
+	
+	return ParentBones({InBoneName}, {InParentName});
+}
+
+bool USkeletonModifier::ParentBones(const TArray<FName>& InBoneNames, const TArray<FName>& InParentNames)
+{
+	if (!IsReferenceSkeletonValid())
+	{
+		return false;
+	}
+	
+	if (InBoneNames.IsEmpty())
+	{
+		UE_LOG(LogAnimation, Error, TEXT("Skeleton Modifier - Parent: No bone provided."));
+		return false;
+	}
+
+	// store initial data
+	const TArray<FMeshBoneInfo> InfosBeforeParenting = ReferenceSkeleton->GetRawRefBoneInfo();
+	
+	// update reference skeleton
+	{
+		auto GetParentName = [&](int32 Index)
+		{
+			if (InBoneNames.Num() == InParentNames.Num())
+			{
+				return InParentNames[Index];
+			}
+			return InParentNames.IsEmpty() ? NAME_None : InParentNames[0];  
+		};
+
+		static constexpr bool bAllowMultipleRoots = true;
+		
+		FReferenceSkeletonModifier Modifier(*ReferenceSkeleton, nullptr);
+		for (int32 Index = 0; Index < InBoneNames.Num(); ++Index)
+		{
+			const int32 BoneIndex = ReferenceSkeleton->FindRawBoneIndex(InBoneNames[Index]);
+			if (BoneIndex != INDEX_NONE)
+			{
+				const FName NewParentName = GetParentName(Index);
+
+				// store global transforms
+				const FTransform GlobalTransform = TransformComposer->GetGlobalTransform(BoneIndex);
+
+				const int32 ParentIndex = ReferenceSkeleton->FindRawBoneIndex(NewParentName);
+				const FTransform ParentTransform = TransformComposer->GetGlobalTransform(ParentIndex);
+				
+				// change parent
+				const int32 NewIndex = Modifier.SetParent(InBoneNames[Index], NewParentName, bAllowMultipleRoots);
+				if (NewIndex > INDEX_NONE)
+				{
+					// update local transform
+					Modifier.UpdateRefPoseTransform(NewIndex, GlobalTransform.GetRelativeTransform(ParentTransform));
+				}
+
+				// invalidate composer
+				TransformComposer->Invalidate(INDEX_NONE);
+			}
+		}
+	}
+
+	// update index tracker
+	UpdateBoneTracker(InfosBeforeParenting);
+	
+	return true;
+}
+
+bool USkeletonModifier::OrientBone(const FName InBoneName, const FOrientOptions& InOptions)
+{
+	if (InBoneName == NAME_None)
+	{
+		UE_LOG(LogAnimation, Error, TEXT("Skeleton Modifier - Orient: Cannot orient a bone with no name."));
+		return false;
+	}
+	
+	return OrientBones({InBoneName}, InOptions);
+}
+
+bool USkeletonModifier::OrientBones(const TArray<FName>& InBoneNames, const FOrientOptions& InOptions)
+{
+	if (!IsReferenceSkeletonValid())
+	{
+		return false;
+	}
+	
+	if (InBoneNames.IsEmpty())
+	{
+		UE_LOG(LogAnimation, Error, TEXT("Skeleton Modifier - Orient: No bone provided."));
+		return false;
+	}
+
+	// get bones to mirror
+	TArray<int32> BonesToOrient; GetBonesToOrient(InBoneNames, InOptions, BonesToOrient);
+	const int32 NumBonesToOrient = BonesToOrient.Num();
+	if (NumBonesToOrient == 0)
+	{
+		UE_LOG(LogAnimation, Error, TEXT("Skeleton Modifier - Orient: None of the provided names has been found."));
+		return false;
+	}
+	
+	auto GetAlignedTransform = [&](const int32 BoneIndex) -> FTransform
+	{
+		const FTransform& BoneGlobal = TransformComposer->GetGlobalTransform(BoneIndex);
+		
+		const int32 ParentIndex = BoneIndex != INDEX_NONE ? ReferenceSkeleton->GetRawParentIndex(BoneIndex) : INDEX_NONE;
+
+		TArray<int32> Children; ReferenceSkeleton->GetDirectChildBones(BoneIndex, Children);
+		const int32 NumChildren = Children.Num();
+		if (NumChildren > 1)
+		{ // we can't align if there are more than one children
+			return BoneGlobal;
+		}
+
+		const FTransform& ParentGlobal = TransformComposer->GetGlobalTransform(ParentIndex);
+		FVector Direction = (BoneGlobal.GetLocation() - ParentGlobal.GetLocation()).GetSafeNormal();
+
+		if (NumChildren > 0)
+		{
+			const FTransform& ChildGlobal = TransformComposer->GetGlobalTransform(Children[0]);
+			Direction = (ChildGlobal.GetLocation() - BoneGlobal.GetLocation()).GetSafeNormal();
+		}
+		
+		if (Direction.IsNearlyZero())
+		{
+			return BoneGlobal;
+		}
+
+		return InOptions.OrientTransform(Direction, BoneGlobal);
+	};
+
+	const TArray<FMeshBoneInfo>& BoneInfos = ReferenceSkeleton->GetRawRefBoneInfo();
+	
+	TArray<FName> BonesToAlign; BonesToAlign.Reserve(NumBonesToOrient);
+	TArray<FTransform> OrientedGlobal; OrientedGlobal.Reserve(NumBonesToOrient);
+	for (int32 Index = 0; Index < NumBonesToOrient; ++Index)	
+	{
+		const int32 BoneIndex = BonesToOrient[Index];
+		BonesToAlign.Add(BoneInfos[BoneIndex].Name);
+		OrientedGlobal.Add(GetAlignedTransform(BoneIndex));
+	}
+
+	// switch back to local
+	TArray<FTransform> Transforms; Transforms.Reserve(NumBonesToOrient);
+	for (int32 Index = 0; Index < BonesToAlign.Num(); Index++)
+	{
+		const FName& BoneName = BonesToAlign[Index];
+		const int32 BoneIndex = ReferenceSkeleton->FindRawBoneIndex(BoneName);
+		const int32 ParentIndex = BoneInfos[BoneIndex].ParentIndex;
+		const int32 ParentOrientedIndex = ParentIndex != INDEX_NONE ? BonesToAlign.IndexOfByKey(BoneInfos[ParentIndex].Name) : INDEX_NONE;
+		const FTransform& ParentGlobal = ParentOrientedIndex != INDEX_NONE ? OrientedGlobal[ParentOrientedIndex] :
+					TransformComposer->GetGlobalTransform(ParentIndex);
+		Transforms.Add(OrientedGlobal[Index].GetRelativeTransform(ParentGlobal));
+	}
+
+	if (BonesToAlign.IsEmpty())
+	{
+		UE_LOG(LogAnimation, Error, TEXT("Skeleton Modifier - Orient: No bone to orient."));
+		return false;
+	}
+
+	static constexpr bool bMoveChildren = false;
+	return SetBonesTransforms(BonesToAlign, Transforms, bMoveChildren);
+}
+
+void USkeletonModifier::GetBonesToOrient(
+	const TArray<FName>& InBonesName, const FOrientOptions& InOptions, TArray<int32>& OutBonesToOrient) const
+{
+	OutBonesToOrient.Reset();
+	
+	TSet<int32> IndicesToOrient;
+	auto GetBonesToOrient = [&](const int32 BoneIndex, auto&& GetBonesToOrient2) -> void
+	{
+		if (BoneIndex == INDEX_NONE)
+		{
+			return;
+		}
+		
+		IndicesToOrient.Add(BoneIndex);
+
+		if (InOptions.bOrientChildren)
+		{
+			TArray<int32> Children; ReferenceSkeleton->GetDirectChildBones(BoneIndex, Children);
+			for (int32 ChildIndex: Children)
+			{
+				GetBonesToOrient2(ChildIndex, GetBonesToOrient2);
+			}
+		}
+	};
+	
+	for (const FName& BoneName: InBonesName)
+	{
+		GetBonesToOrient(ReferenceSkeleton->FindRawBoneIndex(BoneName), GetBonesToOrient);
+	}
+
+	const int32 NumBonesToOrient = IndicesToOrient.Num();
+	if (NumBonesToOrient == 0)
+	{
+		return;
+	}
+
+	IndicesToOrient.Sort([](const int32 Index0, const int32 Index1) {return Index0 < Index1;});
+	OutBonesToOrient = IndicesToOrient.Array();
+}
+
+void USkeletonModifier::UpdateBoneTracker(const TArray<FMeshBoneInfo>& InOtherInfos)
+{
+	for (int32 Index = 0; Index < BoneIndexTracker.Num(); ++Index)
+	{
+		const int32 IndexToTrack = BoneIndexTracker[Index];
+		if (IndexToTrack > INDEX_NONE)
+		{
+			check(InOtherInfos.IsValidIndex(IndexToTrack));
+			const int32 NewIndex = ReferenceSkeleton->FindRawBoneIndex(InOtherInfos[IndexToTrack].Name);
+			BoneIndexTracker[Index] = NewIndex;
+		}
+	}
+}
+
+FName USkeletonModifier::GetUniqueName(const FName InBoneName) const
+{
+	if (!ReferenceSkeleton)
+	{
+		return NAME_None;
+	}
+	
+	auto IsNameAvailable = [&](const FName Name)
+	{
+		const int32 Index = ReferenceSkeleton->FindRawBoneIndex(Name);
+		if (Index != INDEX_NONE)
+		{
+			return false;
+		}
+		return true;
+	};
+
+	FName OutName = InBoneName;
+	int32 CurrentIndex = 0;
+
+	while (!IsNameAvailable(OutName))
+	{
+		OutName = FName(*FString::Printf(TEXT("%s%d"), *InBoneName.ToString(), CurrentIndex++));
+	}
+
+	return OutName;
+}
+
+const FReferenceSkeleton& USkeletonModifier::GetReferenceSkeleton() const
+{
+	static const FReferenceSkeleton Dummy;
+	return ReferenceSkeleton ? *ReferenceSkeleton.Get() : Dummy;
+}
+
+const TArray<int32>& USkeletonModifier::GetBoneIndexTracker() const
+{
+	return BoneIndexTracker;	
+}
+
+FTransform USkeletonModifier::GetBoneTransform(const FName InBoneName, const bool bGlobal) const
+{
+	if (!ReferenceSkeleton.IsValid())
+	{
+		return FTransform::Identity;
+	}
+	return GetTransform(ReferenceSkeleton->FindRawBoneIndex(InBoneName), bGlobal);
+}
+
+FName USkeletonModifier::GetParentName(const FName InBoneName) const
+{
+	const int32 BoneIndex = ReferenceSkeleton->FindBoneIndex(InBoneName);
+	if (BoneIndex == INDEX_NONE)
+	{
+		return NAME_None; 
+	}
+	const TArray<FMeshBoneInfo>& BoneInfos = ReferenceSkeleton->GetRawRefBoneInfo();
+	const int32 ParentIndex = BoneInfos[BoneIndex].ParentIndex;
+	return (ParentIndex > INDEX_NONE) ? BoneInfos[ParentIndex].Name : NAME_None;
+}
+
+const FTransform& USkeletonModifier::GetTransform(const int32 InBoneIndex, const bool bGlobal) const
+{
+	if (!ReferenceSkeleton.IsValid())
+	{
+		return FTransform::Identity;
+	}
+	
+	if (bGlobal)
+	{
+		return TransformComposer.IsValid() ? TransformComposer->GetGlobalTransform(InBoneIndex) : FTransform::Identity;  
+	}
+
+	const TArray<FTransform>& LocalTransforms = ReferenceSkeleton->GetRawRefBonePose();
+	return LocalTransforms.IsValidIndex(InBoneIndex) ? LocalTransforms[InBoneIndex] : FTransform::Identity;
+}
+
+
diff --git a/Engine/Plugins/Experimental/MeshModelingToolsetExp/Source/SkeletalMeshModifiers/Private/SkinWeightModifier.cpp b/Engine/Plugins/Experimental/MeshModelingToolsetExp/Source/SkeletalMeshModifiers/Private/SkinWeightModifier.cpp
new file mode 100644
--- /dev/null
+++ b/Engine/Plugins/Experimental/MeshModelingToolsetExp/Source/SkeletalMeshModifiers/Private/SkinWeightModifier.cpp
@@ -0,0 +1,412 @@
+﻿// Copyright Epic Games, Inc. All Rights Reserved.
+
+#include "SkinWeightModifier.h"
+
+#include "MeshDescription.h"
+#include "SkeletalMeshAttributes.h"
+#include "Rendering/SkeletalMeshModel.h"
+
+#define LOCTEXT_NAMESPACE "SkinWeightModifier"
+
+bool USkinWeightModifier::SetSkeletalMesh(USkeletalMesh* InMesh)
+{
+	// prepare to load mesh data
+	Reset();
+	
+#if WITH_EDITORONLY_DATA
+	
+	// validate supplied skeletal mesh exists
+	if (!InMesh)
+	{
+		UE_LOG(LogAnimation, Error, TEXT("Skin Weight Modifier: No skeletal mesh supplied to load."));
+		return false;
+	}
+
+	// verify user is not trying to modify one of the core engine assets
+	if (InMesh->GetPathName().StartsWith(TEXT("/Engine/")))
+	{
+		UE_LOG(LogAnimation, Error, TEXT("Skin Weight Modifier: Cannot modify built-in engine asset."));
+		return false;
+	}
+
+	// store pointer to mesh 
+	Mesh = InMesh;
+	
+	// store mesh description to edit
+	MeshDescription = MakeUnique<FMeshDescription>();
+	InMesh->GetMeshDescription(LODIndex, *MeshDescription);
+	if (MeshDescription->IsEmpty())
+	{
+		UE_LOG(LogAnimation, Error, TEXT("Skin Weight Modifier: mesh description is empty."));
+		return false;
+	}
+
+	// load weights into memory for subsequent editing
+	const FSkeletalMeshConstAttributes MeshAttribs(*MeshDescription);
+	const FSkinWeightsVertexAttributesConstRef VertexSkinWeights = MeshAttribs.GetVertexSkinWeights();
+	const int32 NumVertices = MeshDescription->Vertices().Num();
+	Weights.SetNum(NumVertices);
+	const FReferenceSkeleton& RefSkeleton = Mesh->GetRefSkeleton();
+	for (int32 VertexIndex = 0; VertexIndex < NumVertices; VertexIndex++)
+	{
+		const FVertexID VertexID(VertexIndex);
+		int32 InfluenceIndex = 0;
+		for (UE::AnimationCore::FBoneWeight BoneWeight: VertexSkinWeights.Get(VertexID))
+		{
+			const int32 BoneIndex = BoneWeight.GetBoneIndex();
+			const FName BoneName = RefSkeleton.GetBoneName(BoneIndex);
+			const float Weight = BoneWeight.GetWeight();
+			Weights[VertexIndex].Add(BoneName, Weight);
+			++InfluenceIndex;
+		}
+
+		if (InfluenceIndex > MAX_TOTAL_INFLUENCES)
+		{
+			UE_LOG(LogAnimation, Warning, TEXT("Skin Weight Modifier: Supplied skeletal mesh has a vertex with too many influences."));
+		}
+	}
+	return true;
+	
+#else
+	ensureMsgf(false, TEXT("Skin Weight Modifier: is an editor only feature."));
+#endif
+	return false;
+}
+
+bool USkinWeightModifier::CommitWeightsToSkeletalMesh()
+{
+	if (!Mesh)
+	{
+		UE_LOG(LogAnimation, Warning, TEXT("Skin Weight Modifier: No mesh loaded. Cannot apply weights to skeletal mesh."));
+		return false;
+	}
+
+#if WITH_EDITORONLY_DATA
+	// update weights in the mesh description
+	{
+		FSkeletalMeshAttributes MeshAttribs(*MeshDescription);
+		FSkinWeightsVertexAttributesRef VertexSkinWeights = MeshAttribs.GetVertexSkinWeights();
+	
+		UE::AnimationCore::FBoneWeightsSettings Settings;
+		Settings.SetNormalizeType(UE::AnimationCore::EBoneWeightNormalizeType::Always);
+
+		TArray<UE::AnimationCore::FBoneWeight> SourceBoneWeights;
+		SourceBoneWeights.Reserve(UE::AnimationCore::MaxInlineBoneWeightCount);
+
+		const int32 NumVertices = MeshDescription->Vertices().Num();
+		const TArray<FName> AllBoneNames = GetAllBoneNames();
+		for (int32 VertexIndex = 0; VertexIndex < NumVertices; VertexIndex++)
+		{
+			SourceBoneWeights.Reset();
+
+			const TMap<FName, float>& VertexWeights = Weights[VertexIndex];
+			if (VertexWeights.IsEmpty())
+			{
+				// in case where user has removed ALL influences from a vertex, revert back to the root bone
+				const FName RootName = Mesh->GetRefSkeleton().GetBoneName(0);
+				const int32 BoneIndex = AllBoneNames.IndexOfByKey(RootName);
+				constexpr float FullWeight = 1.0f;
+				SourceBoneWeights.Add(UE::AnimationCore::FBoneWeight(BoneIndex, FullWeight));
+			}
+			else
+			{
+				for (const TTuple<FName, float>& SingleBoneWeight : VertexWeights)
+				{
+					const FName BoneName = SingleBoneWeight.Key;
+					const int32 BoneIndex = AllBoneNames.IndexOfByKey(BoneName);
+					const float Weight = SingleBoneWeight.Value;
+					SourceBoneWeights.Add(UE::AnimationCore::FBoneWeight(BoneIndex, Weight));
+				}
+			}
+
+			VertexSkinWeights.Set(FVertexID(VertexIndex), UE::AnimationCore::FBoneWeights::Create(SourceBoneWeights, Settings));
+		}
+	}
+	
+	// push changes to the skeletal mesh asset
+	{
+		// flush any pending rendering commands, which might touch a component while we are rebuilding it's mesh
+		FlushRenderingCommands();
+		// update skeletal mesh LOD (cf. USkeletalMesh::CommitMeshDescription)
+		Mesh->CommitMeshDescription(LODIndex, *MeshDescription);
+		// rebuilds mesh
+		Mesh->PostEditChange();
+	}
+
+	return true;
+
+#else
+	ensureMsgf(false, TEXT("Skin Weight Modifier: is an editor only feature."));
+#endif
+		return false;
+}
+
+TMap<FName, float> USkinWeightModifier::GetVertexWeights(int32 VertexID)
+{
+	if (!Weights.IsValidIndex(VertexID))
+	{
+		UE_LOG(LogAnimation, Warning, TEXT("Skin Weight Modifier: Supplied vertex index is not valid for this mesh."));
+		return TMap<FName, float>();
+	}
+
+	return Weights[VertexID];
+}
+
+bool USkinWeightModifier::SetVertexWeights(
+	const int32 VertexID,
+	const TMap<FName, float>& InWeights,
+	const bool bReplaceAll)
+{
+	if (!Mesh)
+	{
+		UE_LOG(LogAnimation, Warning, TEXT("Skin Weight Modifier: No mesh loaded. Cannot set weights."));
+		return false;
+	}
+
+	if (!Weights.IsValidIndex(VertexID))
+	{
+		UE_LOG(LogAnimation, Warning, TEXT("Skin Weight Modifier: Supplied vertex index is not valid on this mesh. Cannot set weights."));
+		return false;
+	}
+
+	TMap<FName, float>& VertexWeights = Weights[VertexID];
+	if (bReplaceAll)
+	{
+		// replace weights entirely
+		VertexWeights = InWeights;
+	}
+	else
+	{
+		// replace only the input weights
+		for (const TTuple<FName, float>& InWeight : InWeights)
+		{
+			VertexWeights.FindOrAdd(InWeight.Key) = InWeight.Value;
+		}
+	}
+	
+	return true;
+}
+
+bool USkinWeightModifier::NormalizeVertexWeights(const int32 VertexID)
+{
+	if (!Mesh)
+	{
+		UE_LOG(LogAnimation, Warning, TEXT("Skin Weight Modifier: No mesh loaded. Cannot normalize weights."));
+		return false;
+	}
+
+	if (!Weights.IsValidIndex(VertexID))
+	{
+		UE_LOG(LogAnimation, Warning, TEXT("Skin Weight Modifier: Supplied vertex ID was invalid. Cannot normalize vertex weights."));
+		return false;
+	}
+
+	return NormalizeVertexInternal(VertexID);
+}
+
+bool USkinWeightModifier::NormalizeAllWeights()
+{
+	if (!(Mesh && MeshDescription))
+	{
+		UE_LOG(LogAnimation, Warning, TEXT("Skin Weight Modifier: No mesh loaded. Cannot normalize weights."));
+		return false;
+	}
+	
+	const int32 NumVertices = MeshDescription->Vertices().Num();
+	for (int32 VertexID=0; VertexID<NumVertices; ++VertexID)
+	{
+		NormalizeVertexInternal(VertexID);
+	}
+
+	return true;
+}
+
+bool USkinWeightModifier::EnforceMaxInfluences(const int32 MaxInfluences)
+{
+	if (!MeshDescription)
+	{
+		UE_LOG(LogAnimation, Warning, TEXT("Skin Weight Modifier: No mesh loaded. Cannot enforce max influences."));
+		return false;
+	}
+	
+	int32 Max = MaxInfluences;
+	if (Max < 0)
+	{
+		Max = UE::AnimationCore::MaxInlineBoneWeightCount;
+	}
+
+	const int32 NumVertices = MeshDescription->Vertices().Num();
+	for (int32 VertexID=0; VertexID<NumVertices; ++VertexID)
+	{
+		EnforceMaxInfluenceInternal(VertexID, Max);
+	}
+
+	return true;
+}
+
+bool USkinWeightModifier::PruneVertexWeights(const int32 VertexID, const float WeightThreshold)
+{
+	if (!Mesh)
+	{
+		UE_LOG(LogAnimation, Warning, TEXT("Skin Weight Modifier: No mesh loaded. Cannot prune weights."));
+		return false;
+	}
+
+	if (!Weights.IsValidIndex(VertexID))
+	{
+		UE_LOG(LogAnimation, Warning, TEXT("Skin Weight Modifier: Supplied vertex ID was invalid. Cannot prune weights."));
+		return false;
+	}
+
+	PruneVertexWeightInternal(VertexID, WeightThreshold);
+	
+	return true;
+}
+
+bool USkinWeightModifier::PruneAllWeights(const float WeightThreshold)
+{
+	if (!Mesh)
+	{
+		UE_LOG(LogAnimation, Warning, TEXT("Skin Weight Modifier: No mesh loaded. Cannot prune weights."));
+		return false;
+	}
+
+	const int32 NumVertices = MeshDescription->Vertices().Num();
+	for (int32 i=0; i<NumVertices; ++i)
+	{
+		PruneVertexWeightInternal(i, WeightThreshold);
+	}
+	
+	return true;
+}
+
+int32 USkinWeightModifier::GetNumVertices()
+{
+	if (!MeshDescription)
+	{
+		UE_LOG(LogAnimation, Warning, TEXT("Skin Weight Modifier: No mesh loaded. Cannot get vertex count."));
+		return INDEX_NONE;
+	}
+
+	return MeshDescription->Vertices().Num();
+}
+
+TArray<FName> USkinWeightModifier::GetAllBoneNames()
+{
+	if (!Mesh)
+	{
+		UE_LOG(LogAnimation, Warning, TEXT("Skin Weight Modifier: No mesh loaded. Cannot get bones."));
+		return TArray<FName>();
+	}
+
+	const FReferenceSkeleton& RefSkeleton = Mesh->GetRefSkeleton();
+	const int32 NumBones = RefSkeleton.GetNum();
+	TArray<FName> AllBoneNames;
+	for (int32 i=0; i<NumBones; ++i)
+	{
+		AllBoneNames.Add(RefSkeleton.GetBoneName(i));
+	}
+	return AllBoneNames;
+}
+
+bool USkinWeightModifier::NormalizeVertexInternal(const int32 VertexID)
+{
+	// handle case where user has removed all weights on a vertex (probably a mistake so warn them)
+	TMap<FName, float>& VertexWeights = Weights[VertexID];
+	if (VertexWeights.IsEmpty())
+	{
+		const FName RootBone = Mesh->GetRefSkeleton().GetBoneName(0);
+		VertexWeights.Add(RootBone, 1.f);
+		UE_LOG(LogAnimation, Warning, TEXT("Skin Weight Modifier: Supplied vertex did not have any bone influences. Applied all weight to root bone."));
+		return true;
+	}
+
+	// calculate total weight on all vertices
+	float TotalWeight = 0.f;
+	for (const TTuple<FName, float>& Weight : VertexWeights)
+	{
+		TotalWeight += Weight.Value;
+	}
+
+	// handle case where total weight is near zero
+	if (TotalWeight <= KINDA_SMALL_NUMBER)
+	{
+		// in this case, we just split the weight evenly across all available influences
+		const float SplitWeight = 1.f / VertexWeights.Num();
+		for (TTuple<FName, float>& Weight : VertexWeights)
+		{
+			Weight.Value = SplitWeight;
+		}
+		UE_LOG(LogAnimation, Warning, TEXT("Skin Weight Modifier: Supplied vertex had near zero total weight."));
+		return true;
+	}
+
+	// scale weights to achieve total value of 1.0
+	for (TTuple<FName, float>& Weight : VertexWeights)
+	{
+		Weight.Value /= TotalWeight;
+	}
+
+	return true;
+}
+
+void USkinWeightModifier::EnforceMaxInfluenceInternal(const int32 VertexID, const int32 Max)
+{
+	// get weights for this vertex
+	TMap<FName, float>& VertexWeights = Weights[VertexID];
+
+	// already under max influences?
+	if (VertexWeights.Num() <= Max)
+	{
+		return;
+	}
+
+	// remove smallest influences until we reach max influences
+	while(VertexWeights.Num() > Max)
+	{
+		FName SmallestInfName = NAME_None;
+		float SmallestInf = TNumericLimits<float>::Max();
+		for (const TTuple<FName, float> Weight : VertexWeights)
+		{
+			if (Weight.Value <= SmallestInf)
+			{
+				SmallestInf = Weight.Value;
+				SmallestInfName = Weight.Key;
+			}
+		}
+
+		VertexWeights.Remove(SmallestInfName);
+	}
+
+	NormalizeVertexInternal(VertexID);
+}
+
+bool USkinWeightModifier::PruneVertexWeightInternal(const int32 VertexID, const float WeightThreshold)
+{
+	// get weights for this vertex
+	TMap<FName, float>& VertexWeights = Weights[VertexID];
+	TArray<FName> InfluencesToPrune;
+	for (const TTuple<FName, float>& Weight : VertexWeights)
+	{
+		if (Weight.Value < WeightThreshold)
+		{
+			InfluencesToPrune.Add(Weight.Key);
+		}
+	}
+	
+	for (const FName InfluenceToPrune : InfluencesToPrune)
+	{
+		VertexWeights.Remove(InfluenceToPrune);
+	}
+	
+	return true;
+}
+
+void USkinWeightModifier::Reset()
+{
+	Mesh = nullptr;
+	MeshDescription.Reset();
+	Weights.Reset();
+}
+
+#undef LOCTEXT_NAMESPACE
\ No newline at end of file
diff --git a/Engine/Plugins/Experimental/MeshModelingToolsetExp/Source/SkeletalMeshModifiers/Public/SkeletalMeshModifiersModule.h b/Engine/Plugins/Experimental/MeshModelingToolsetExp/Source/SkeletalMeshModifiers/Public/SkeletalMeshModifiersModule.h
new file mode 100644
--- /dev/null
+++ b/Engine/Plugins/Experimental/MeshModelingToolsetExp/Source/SkeletalMeshModifiers/Public/SkeletalMeshModifiersModule.h
@@ -0,0 +1,14 @@
+﻿// Copyright Epic Games, Inc. All Rights Reserved.
+
+#pragma once
+
+#include "Modules/ModuleManager.h"
+
+class FSkeletalMeshModifiersModule : public IModuleInterface
+{
+public:
+
+	/** IModuleInterface implementation */
+	virtual void StartupModule() override;
+	virtual void ShutdownModule() override;
+};
diff --git a/Engine/Plugins/Experimental/MeshModelingToolsetExp/Source/SkeletalMeshModifiers/Public/SkeletonModifier.h b/Engine/Plugins/Experimental/MeshModelingToolsetExp/Source/SkeletalMeshModifiers/Public/SkeletonModifier.h
new file mode 100644
--- /dev/null
+++ b/Engine/Plugins/Experimental/MeshModelingToolsetExp/Source/SkeletalMeshModifiers/Public/SkeletonModifier.h
@@ -0,0 +1,265 @@
+﻿// Copyright Epic Games, Inc. All Rights Reserved.
+
+#pragma once
+
+#include "Math/MathFwd.h"
+#include "Containers/Array.h"
+#include "Templates/UniquePtr.h"
+#include "MeshDescription.h"
+#include "ReferenceSkeleton.h"
+
+#include "SkeletonModifier.generated.h"
+
+class USkeletalMesh;
+class FName;
+struct FMeshDescription;
+struct FReferenceSkeleton;
+struct FTransformComposer;
+
+/**
+ * FMirrorOptions
+ */
+
+USTRUCT(BlueprintType)
+struct SKELETALMESHMODIFIERS_API FMirrorOptions
+{
+	GENERATED_USTRUCT_BODY()
+
+	UPROPERTY(BlueprintReadWrite, EditAnywhere, Category = "Mirror")
+	TEnumAsByte<EAxis::Type> MirrorAxis = EAxis::X;
+
+	UPROPERTY(BlueprintReadWrite, EditAnywhere, Category = "Mirror")
+	bool bMirrorRotation = true;
+
+	UPROPERTY(BlueprintReadWrite, EditAnywhere, Category = "Mirror")
+	FString LeftString = TEXT("_l");
+
+	UPROPERTY(BlueprintReadWrite, EditAnywhere, Category = "Mirror")
+	FString RightString = TEXT("_r");
+
+	UPROPERTY(BlueprintReadWrite, EditAnywhere, Category = "Mirror")
+	bool bMirrorChildren = true;
+
+	FTransform MirrorTransform(const FTransform& InGlobalTransform) const;
+	FVector MirrorVector(const FVector& InVector) const;
+};
+
+/**
+ * FOrientOptions
+ */
+
+USTRUCT(BlueprintType)
+struct SKELETALMESHMODIFIERS_API FOrientOptions
+{
+	GENERATED_USTRUCT_BODY()
+
+	UPROPERTY(BlueprintReadWrite, EditAnywhere, Category = "Orient")
+	TEnumAsByte<EAxis::Type> Primary = EAxis::X;
+
+	UPROPERTY(BlueprintReadWrite, EditAnywhere, Category = "Orient")
+	TEnumAsByte<EAxis::Type> Secondary = EAxis::Y;
+
+	UPROPERTY(BlueprintReadWrite, EditAnywhere, Category = "Orient")
+	FVector SecondaryTarget = FVector::YAxisVector;
+
+	UPROPERTY(BlueprintReadWrite, EditAnywhere, Category = "Orient")
+	bool bOrientChildren = true;
+
+	FTransform OrientTransform(const FVector& InPrimaryTarget, const FTransform& InTransform) const;
+};
+
+/**
+ * FSkeletalMeshSkeletonModifier
+ */
+
+UCLASS(BlueprintType)
+class SKELETALMESHMODIFIERS_API USkeletonModifier : public UObject
+{
+	GENERATED_BODY()
+	
+public:
+
+	UFUNCTION(BlueprintCallable, Category = "Mesh")
+	bool Init(USkeletalMesh* InSkeletalMesh);
+
+	/** Creates a new bone in the skeleton hierarchy at desired transform
+	 *  @param InBoneName The new bone's name.
+	 *  @param InParentName The new bone parent's name. 
+	 *  @param InTransform The default local transform in the parent space.
+	 *  @return \c true if the operation succeeded, false otherwise. 
+	 */
+	UFUNCTION(BlueprintCallable, Category = "Skeleton")
+	bool AddBone(const FName InBoneName, const FName InParentName, const FTransform& InTransform);
+	UFUNCTION(BlueprintCallable, Category = "Skeleton")
+	bool AddBones(const TArray<FName>& InBonesName, const TArray<FName>& InParentsName, const TArray<FTransform>& InTransforms);
+
+	/** Mirror bones
+	 *  @param InBoneName The new bone's name.
+	 *  @param InOptions The mirroring options
+	 *  @return \c true if the operation succeeded, false otherwise. 
+	 */
+	UFUNCTION(BlueprintCallable, Category = "Skeleton")
+	bool MirrorBone(const FName InBoneName, const FMirrorOptions& InOptions = FMirrorOptions());
+	UFUNCTION(BlueprintCallable, Category = "Skeleton")
+	bool MirrorBones(const TArray<FName>& InBonesName, const FMirrorOptions& InOptions = FMirrorOptions());
+
+	/** Sets the bone the desired local transform
+	 *  @param InBoneName The new bone's name that needs to be moved.
+	 *  @param InNewTransform The new local transform in the bone's parent space.
+	 *  @param bMoveChildren Propagate new transform to children
+	 *  @return \c true if the operation succeeded, false otherwise. 
+	 */
+	UFUNCTION(BlueprintCallable, Category = "Skeleton")
+	bool SetBoneTransform(const FName InBoneName, const FTransform& InNewTransform, const bool bMoveChildren);
+	UFUNCTION(BlueprintCallable, Category = "Skeleton")
+	bool SetBonesTransforms(const TArray<FName>& InBoneNames, const TArray<FTransform>& InNewTransforms, const bool bMoveChildren);
+
+	/** Remove a bone in the skeleton hierarchy
+	 *  @param InBoneName The new bone's name.
+	 *  @param bRemoveChildren Remove children recursively.
+	 *  @return \c true if the operation succeeded, false otherwise. 
+	 */
+	UFUNCTION(BlueprintCallable, Category = "Skeleton")
+	bool RemoveBone(const FName InBoneName, const bool bRemoveChildren);
+	UFUNCTION(BlueprintCallable, Category = "Skeleton")
+	bool RemoveBones(const TArray<FName>& InBoneNames, const bool bRemoveChildren);
+
+	/** Rename bones
+	 *  @param InOldBoneName The current bone's name.
+	 *  @param InNewBoneName The new bone's name.
+	 *  @return \c true if the operation succeeded, false otherwise. 
+	 */
+	UFUNCTION(BlueprintCallable, Category = "Skeleton")
+	bool RenameBone(const FName InOldBoneName, const FName InNewBoneName);
+	UFUNCTION(BlueprintCallable, Category = "Skeleton")
+	bool RenameBones(const TArray<FName>& InOldBoneNames, const TArray<FName>& InNewBoneNames);
+	
+	/** Parent bones
+	 *  @param InBoneName The current bone's name.
+	 *  @param InParentName The new parent's name (Name_NONE to unparent).
+	 *  @return \c true if the operation succeeded, false otherwise. 
+	 */
+	UFUNCTION(BlueprintCallable, Category = "Skeleton")
+	bool ParentBone(const FName InBoneName, const FName InParentName);
+	UFUNCTION(BlueprintCallable, Category = "Skeleton")
+	bool ParentBones(const TArray<FName>& InBoneNames, const TArray<FName>& InParentNames);
+
+	/** Align bones
+	 *  @param InBoneName The current bone's name.
+	 *  @param InOptions The orienting options
+	 *  @return \c true if the operation succeeded, false otherwise. 
+	 */
+	 UFUNCTION(BlueprintCallable, Category = "Skeleton")
+	bool OrientBone(const FName InBoneName, const FOrientOptions& InOptions = FOrientOptions());
+	UFUNCTION(BlueprintCallable, Category = "Skeleton")
+	bool OrientBones(const TArray<FName>& InBoneNames, const FOrientOptions& InOptions = FOrientOptions());
+
+	/** Get Bone Transform
+	 *  @param InBoneName The current bone's name.
+	 *  @param bGlobal Whether it should return the parent space or global transform
+	 *  @return \c The current bone's transform 
+	 */
+	UFUNCTION(BlueprintCallable, Category = "Skeleton")
+	FTransform GetBoneTransform(const FName InBoneName, const bool bGlobal = false) const;
+
+	/** Get Parent Name
+	 *  @param InBoneName The current bone's name.
+	 *  @return \c The current bone's parent name 
+	 */
+	UFUNCTION(BlueprintCallable, Category = "Skeleton")
+	FName GetParentName(const FName InBoneName) const;
+	
+	/**
+	 * Actually applies the skeleton modifications to the skeletal mesh.
+	 * @return true if commit succeeded.
+	 */
+	UFUNCTION(BlueprintCallable, Category = "Mesh")
+	bool CommitSkeletonToSkeletalMesh();
+
+	void ExternalUpdate(const FReferenceSkeleton& InRefSkeleton, const TArray<int32>& InIndexTracker);
+	
+	FName GetUniqueName(const FName InBoneName) const;
+	
+	const FReferenceSkeleton& GetReferenceSkeleton() const;
+	const TArray<int32>& GetBoneIndexTracker() const;
+	
+	const FTransform& GetTransform(const int32 InBoneIndex, const bool bGlobal = false) const;
+	
+private:
+
+	bool IsReferenceSkeletonValid(const bool bLog = true) const;
+
+	void UpdateBoneTracker(const TArray<FMeshBoneInfo>& InOtherInfos);
+	
+	// mirroring functions
+	void GetBonesToMirror(const TArray<FName>& InBonesName, const FMirrorOptions& InOptions, TArray<int32>& OutBonesToMirror) const;
+	void GetMirroredNames(const TArray<int32>& InBonesToMirror, const FMirrorOptions& InOptions, TArray<FName>& OutBonesName) const;
+	void GetMirroredBones(const TArray<int32>& InBonesToMirror, const TArray<FName>& InMirroredNames, TArray<int32>& OutMirroredBones);
+	void GetMirroredTransforms(	const TArray<int32>& InBonesToMirror, const TArray<int32>& InMirroredBones,
+								const FMirrorOptions& InOptions, TArray<FTransform>& OutMirroredTransforms) const;
+
+	// orient function
+	void GetBonesToOrient(const TArray<FName>& InBonesName, const FOrientOptions& InOptions, TArray<int32>& OutBonesToOrient) const;
+	
+	TObjectPtr<USkeletalMesh> SkeletalMesh = nullptr;
+	TUniquePtr<FMeshDescription> MeshDescription;
+	TUniquePtr<FReferenceSkeleton> ReferenceSkeleton;
+	TUniquePtr<FTransformComposer> TransformComposer;
+	TArray<int32> BoneIndexTracker;
+};
+
+struct FTransformComposer
+{
+	FTransformComposer(const FReferenceSkeleton& InRefSkeleton)
+		: RefSkeleton(InRefSkeleton)
+		, Transforms(InRefSkeleton.GetRawRefBonePose())
+	{
+		TransformCached.Init(false, Transforms.Num());
+	}
+
+	const FTransform& GetGlobalTransform(const uint32 Index)
+	{
+		if (!Transforms.IsValidIndex(Index))
+		{
+			return FTransform::Identity;
+		}
+
+		if (TransformCached[Index])
+		{
+			return Transforms[Index];
+		}
+
+		const int32 ParentIndex = RefSkeleton.GetRawParentIndex(Index);
+		if (ParentIndex > INDEX_NONE)
+		{
+			Transforms[Index] *= GetGlobalTransform(ParentIndex);
+		}
+
+		TransformCached[Index] = true;
+		return Transforms[Index];
+	}
+
+	void Invalidate(const uint32 Index = INDEX_NONE)
+	{
+		if (!TransformCached.IsValidIndex(Index))
+		{
+			Transforms = RefSkeleton.GetRawRefBonePose();
+			TransformCached.Init(false, Transforms.Num());
+			return;
+		}
+		
+		Transforms[Index] = RefSkeleton.GetRawRefBonePose()[Index];
+		TransformCached[Index] = false;
+
+		TArray<int32> Children; RefSkeleton.GetDirectChildBones(Index, Children);
+		for (const int32 ChildIndex: Children)
+		{
+			Invalidate(ChildIndex);
+		}
+	}
+
+	const FReferenceSkeleton& RefSkeleton;
+	TArray<FTransform> Transforms;
+	TBitArray<> TransformCached;
+};
+
diff --git a/Engine/Plugins/Experimental/MeshModelingToolsetExp/Source/SkeletalMeshModifiers/Public/SkinWeightModifier.h b/Engine/Plugins/Experimental/MeshModelingToolsetExp/Source/SkeletalMeshModifiers/Public/SkinWeightModifier.h
new file mode 100644
--- /dev/null
+++ b/Engine/Plugins/Experimental/MeshModelingToolsetExp/Source/SkeletalMeshModifiers/Public/SkinWeightModifier.h
@@ -0,0 +1,187 @@
+﻿// Copyright Epic Games, Inc. All Rights Reserved.
+
+#pragma once
+
+#include "Templates/UniquePtr.h"
+#include "MeshDescription.h"
+#include "Engine/SkeletalMesh.h"
+
+#include "SkinWeightModifier.generated.h"
+
+class USkeletalMesh;
+
+/**
+ * API used to modify skin weights on a Skeletal Mesh asset.
+ * 
+ * To use:
+ * 1. Instantiate an instance of USkinWeightModifier
+ * 2. Call "SetSkeletalMesh(MyMeshAsset)", passing in the skeletal mesh you want to edit weights on
+ * 3. Use Get/Set weights functions, and Normalize/Prune to edit the weights as desired
+ * 4. When ready to commit to the asset, call CommitWeightsToSkeletalMesh() and save the asset if desired
+ *
+ * This API can be used from C++, Blueprint or Python. Here is a sample usage of the API in Python:
+
+import unreal
+
+# create a weight modifier for a given skeletal mesh
+skel_mesh = unreal.EditorAssetLibrary().load_asset("/Game/Characters/Wolf/Meshes/SK_Wolf")
+weight_modifier = unreal.SkinWeightModifier()
+weight_modifier.set_skeletal_mesh(skel_mesh)
+
+# get weight of vertex 1234
+vertex_weights = weight_modifier.get_vertex_weights(1234)
+print(vertex_weights)
+
+# remove neck2 as an influence on this vertex
+vertex_weights.pop("neck2")
+weight_modifier.set_vertex_weights(vertex_weights, True)
+print(vertex_weights)
+
+# commit change to the skeletal mesh
+weight_modifier.commit_weights_to_skeletal_mesh()
+
+ * OUTPUT:
+ * {"head": 0.6, "neck1": 0.3, "neck2": 0.1}
+ * {"head": 0.6, "neck1": 0.3}
+ * 
+ * In Python, the per-vertex weights are stored as a dictionary mapping Bone Names to float weight values.
+ *
+ * The "SetVertexWeights()" function expects the same data structure. You can add/remove/edit influences as needed.
+ * The SetVertexWeights() function does not normalize the weights. So you can make multiple modifications
+ * and call NormalizeVertexWeights() or NormalizeAllWeights() as desired.
+ *
+ * The PruneVertexWeights() and EnforceMaxInfluences() functions can be used to trim small influences
+ * and clamp the total number of influences per vertex as needed.
+ * 
+ * Note that it is not required to normalize the weights by calling any of the normalize functions. Or manually before
+ * calling SetVertexWeights(). Committing the weights to the skeletal mesh will always enforce normalization.
+ *
+ * Though it may be useful to normalize while editing.
+ */
+UCLASS(BlueprintType)
+class SKELETALMESHMODIFIERS_API USkinWeightModifier : public UObject
+{
+	GENERATED_BODY()
+
+public:
+	
+	/**
+	 * Call this first to load the weights for a skeletal mesh for fast editing.
+	 * @param InMesh - The skeletal mesh asset to load for weight editing
+	 * @return bool - true if the mesh weights were successfully loaded.
+	 */
+	UFUNCTION(BlueprintCallable, Category = "Mesh")
+	bool SetSkeletalMesh(USkeletalMesh* InMesh);
+
+	/**
+	 * Actually applies the weight modifications to the skeletal mesh. This action creates an undo transaction.
+	 * The skeletal mesh asset will be dirtied, but it is up to the user to save the asset if required.
+	 * @return true if weights were applied to a skeletal mesh, false otherwise
+	 */
+	UFUNCTION(BlueprintCallable, Category = "Mesh")
+	bool CommitWeightsToSkeletalMesh();
+
+	/**
+	 * Get a reference to the skeletal mesh that was loaded
+	 * @return USkeletalMesh* - the mesh that was loaded, or null
+	 */
+	UFUNCTION(BlueprintCallable, Category = "Mesh")
+	USkeletalMesh* GetSkeletalMesh() { return Mesh; }
+
+	/**
+	 * Get the total number of vertices in the skeletal mesh.
+	 * @return int, number of vertices
+	 */
+	UFUNCTION(BlueprintCallable, Category = "Mesh Query")
+	int32 GetNumVertices();
+
+	/**
+	 * Get an array of all bone names in the skeletal mesh.
+	 * @return array of bone names
+	 */
+	UFUNCTION(BlueprintCallable, Category = "Mesh Query")
+	TArray<FName> GetAllBoneNames();
+
+	/**
+	 * Get all bone weights for a single vertex.
+	 * @param VertexID the index of the vertex
+	 * @return a map of Bone Name to weight values for all bones that influence the specified vertex.
+	 */
+	UFUNCTION(BlueprintCallable, Category = "Weights")
+	TMap<FName, float> GetVertexWeights(const int32 VertexID);
+
+	/**
+	 * Set bone weights for a single vertex. The weights are stored as supplied and not normalized until
+	 * either "CommitWeightsToSkeletalMesh()" or "NormalizeVertexWeights()" is called.
+	 * @param VertexID the index of the vertex
+	 * @param InWeights a map of Bone-Name to Weight for all bones that influence the specified vertex, ie {"Head": 0.6, "Neck": 0.4}
+	 * @param bReplaceAll if true, all weights on this vertex will be replaced with the input weights. Default is false.
+	 */
+	UFUNCTION(BlueprintCallable, Category = "Weights")
+	bool SetVertexWeights(
+		const int32 VertexID,
+		const TMap<FName, float>& InWeights,
+		const bool bReplaceAll=false);
+
+	/**
+	 * Normalize weights on the specified vertex.
+	 * @param VertexID the index of the vertex to normalize weights on
+	 * @return true if normalization was performed, false otherwise
+	 */
+	UFUNCTION(BlueprintCallable, Category = "Weights")
+	bool NormalizeVertexWeights(const int32 VertexID);
+
+	/**
+	 * Normalize weights on all vertices in the mesh.
+	 * @return true if normalization was performed, false otherwise
+	 */
+	UFUNCTION(BlueprintCallable, Category = "Weights")
+	bool NormalizeAllWeights();
+
+	/**
+	 * Strips out smallest influences to ensure each vertex does not have weight on more influences than MaxInfluences.
+	 * Influences with the smallest weight are culled first.
+	 * @param MaxInfluences The maximum number of influences to allow for each vertex in the mesh. If -1 is passed, will use the project-wide MaxInfluences amount.
+	 * @return true if influences were removed, false otherwise
+	 */
+	UFUNCTION(BlueprintCallable, Category = "Weights")
+	bool EnforceMaxInfluences(const int32 MaxInfluences=-1);
+
+	/**
+	 * Remove all weights below the given threshold value, on the given vertex.
+	 * Influences that are pruned will no longer receive weight from normalization.
+	 * @param VertexID the index of the vertex to prune weights on
+	 * @param WeightThreshold vertex weights below this value will be removed. 
+	 * @return true if influences were removed, false otherwise
+	 */
+	UFUNCTION(BlueprintCallable, Category = "Weights")
+	bool PruneVertexWeights(const int32 VertexID, const float WeightThreshold);
+
+	/**
+	 * Remove all weights below the given threshold value, on all vertices.
+	 * @param WeightThreshold vertex weights below this value will be removed. 
+	 * @return true if influences were removed, false otherwise
+	 */
+	UFUNCTION(BlueprintCallable, Category = "Weights")
+	bool PruneAllWeights(const float WeightThreshold);
+
+private:
+
+	// the skeletal mesh that was loaded by SetSkeletalMesh
+	UPROPERTY()
+	TObjectPtr<USkeletalMesh> Mesh;
+	TUniquePtr<FMeshDescription> MeshDescription;
+
+	// weights stored in-memory for editing before being committed to the skeletal mesh
+	TArray<TMap<FName, float>> Weights;
+
+	// fast internal procedures (do not validate inputs)
+	bool NormalizeVertexInternal(const int32 VertexID);
+	void EnforceMaxInfluenceInternal(const int32 VertexID, const int32 Max);
+	bool PruneVertexWeightInternal(const int32 VertexID, const float WeightThreshold);
+
+	// reset all internal data when loading a new mesh
+	void Reset();
+
+	static constexpr int32 LODIndex = 0;
+};
diff --git a/Engine/Plugins/Experimental/MeshModelingToolsetExp/Source/SkeletalMeshModifiers/SkeletalMeshModifiers.build.cs b/Engine/Plugins/Experimental/MeshModelingToolsetExp/Source/SkeletalMeshModifiers/SkeletalMeshModifiers.build.cs
new file mode 100644
--- /dev/null
+++ b/Engine/Plugins/Experimental/MeshModelingToolsetExp/Source/SkeletalMeshModifiers/SkeletalMeshModifiers.build.cs
@@ -0,0 +1,40 @@
+// Copyright Epic Games, Inc. All Rights Reserved.
+
+using UnrealBuildTool;
+
+public class SkeletalMeshModifiers : ModuleRules
+{
+	public SkeletalMeshModifiers(ReadOnlyTargetRules Target) : base(Target)
+	{
+		PCHUsage = ModuleRules.PCHUsageMode.UseExplicitOrSharedPCHs;
+		
+		PublicDependencyModuleNames.AddRange(
+			new string[]
+			{
+				"Core",
+				// ... add other public dependencies that you statically link with here ...
+			}
+            );
+			
+		
+		PrivateDependencyModuleNames.AddRange(
+			new string[]
+			{
+				"MeshDescription",
+				"SkeletalMeshDescription",
+				"SkeletalMeshUtilitiesCommon",
+				"RenderCore",
+				"CoreUObject",
+				"Engine"
+			}
+			);
+
+
+		DynamicallyLoadedModuleNames.AddRange(
+			new string[]
+			{
+				// ... add any modules that your module loads dynamically here ...
+			}
+			);
+	}
+}
diff --git a/Engine/Source/Developer/SkeletalMeshUtilitiesCommon/Private/SkeletonModifier.cpp b/Engine/Source/Developer/SkeletalMeshUtilitiesCommon/Private/SkeletonModifier.cpp
deleted file mode 100644
diff --git a/Engine/Source/Developer/SkeletalMeshUtilitiesCommon/Public/SkeletonModifier.h b/Engine/Source/Developer/SkeletalMeshUtilitiesCommon/Public/SkeletonModifier.h
deleted file mode 100644
diff --git a/Engine/Source/Editor/SkeletalMeshEditor/Private/SkinWeightModifier.cpp b/Engine/Source/Editor/SkeletalMeshEditor/Private/SkinWeightModifier.cpp
deleted file mode 100644
diff --git a/Engine/Source/Editor/SkeletalMeshEditor/Public/SkinWeightModifier.h b/Engine/Source/Editor/SkeletalMeshEditor/Public/SkinWeightModifier.h
deleted file mode 100644
diff --git a/Engine/Source/Editor/SkeletalMeshEditor/SkeletalMeshEditor.Build.cs b/Engine/Source/Editor/SkeletalMeshEditor/SkeletalMeshEditor.Build.cs
--- a/Engine/Source/Editor/SkeletalMeshEditor/SkeletalMeshEditor.Build.cs
+++ b/Engine/Source/Editor/SkeletalMeshEditor/SkeletalMeshEditor.Build.cs
@@ -40,10 +40,7 @@ public class SkeletalMeshEditor : ModuleRules
 				"StatusBar",
 				"PhysicsUtilities",
 				"InterchangeCore",
-				"InterchangeEngine",
-				"MeshDescription",
-				"SkeletalMeshDescription",
-				"RenderCore"
+				"InterchangeEngine"
 			}
 		);
 
