Fixing nonunity compile error in SkinWeightModifier.cpp

diff --git a/Engine/Plugins/Experimental/MeshModelingToolsetExp/Source/SkeletalMeshModifiers/Private/SkinWeightModifier.cpp b/Engine/Plugins/Experimental/MeshModelingToolsetExp/Source/SkeletalMeshModifiers/Private/SkinWeightModifier.cpp
--- a/Engine/Plugins/Experimental/MeshModelingToolsetExp/Source/SkeletalMeshModifiers/Private/SkinWeightModifier.cpp
+++ b/Engine/Plugins/Experimental/MeshModelingToolsetExp/Source/SkeletalMeshModifiers/Private/SkinWeightModifier.cpp
@@ -5,6 +5,7 @@
 #include "MeshDescription.h"
 #include "SkeletalMeshAttributes.h"
 #include "Rendering/SkeletalMeshModel.h"
+#include "EngineLogs.h"
 
 #define LOCTEXT_NAMESPACE "SkinWeightModifier"
 
