Mesh Modeling Tools:  Improve Weld Mesh Edges tool to optionally merge attributes.
Geometry Core: Adding Split Attribute Welder.

diff --git a/Engine/Plugins/Experimental/MeshModelingToolsetExp/Source/MeshModelingToolsExp/Private/WeldMeshEdgesTool.cpp b/Engine/Plugins/Experimental/MeshModelingToolsetExp/Source/MeshModelingToolsExp/Private/WeldMeshEdgesTool.cpp
--- a/Engine/Plugins/Experimental/MeshModelingToolsetExp/Source/MeshModelingToolsExp/Private/WeldMeshEdgesTool.cpp
+++ b/Engine/Plugins/Experimental/MeshModelingToolsetExp/Source/MeshModelingToolsExp/Private/WeldMeshEdgesTool.cpp
@@ -12,6 +12,8 @@
 #include "DynamicMesh/Operations/MergeCoincidentMeshEdges.h"
 #include "Operations/MeshResolveTJunctions.h"
 
+#include "Math/UnrealMathUtility.h"
+
 #include UE_INLINE_GENERATED_CPP_BY_NAME(WeldMeshEdgesTool)
 
 using namespace UE::Geometry;
@@ -32,13 +34,28 @@ USingleSelectionMeshEditingTool* UWeldMeshEdgesToolBuilder::CreateNewTool(const
 
 class FWeldMeshEdgesOp : public  FDynamicMeshOperator
 {
+
 public:
+	enum class EWeldAttributeMode : uint8
+	{
+		/** Do not weld attributes*/
+		None,
+		/** Apply attribute merging on the mesh weld */
+		OnWeldedMeshEdgesOnly,
+		/** Apply attribute merging to all split attributes */
+		OnFullMesh
+	};
 
 	// parameters set by the tool
 	TSharedPtr<FDynamicMesh3, ESPMode::ThreadSafe> SourceMesh;
 	double Tolerance;
 	bool bOnlyUnique;
 	bool bResolveTJunctions;
+	EWeldAttributeMode WeldAttributeMode;
+	float SplitNormalThreshold;
+	float SplitTangentsThreshold;
+	float SplitUVThreshold;
+	float SplitColorThreshold; 
 
 	int32 InitialNumBoundaryEdges = 0;
 	int32 FinalNumBoundaryEdges = 0;
@@ -48,6 +65,17 @@ public:
 
 	virtual void CalculateResult(FProgressCancel* Progress) override
 	{
+
+		auto SetAttributeWeldOptions = [this](FSplitAttributeWelder& SplitAttributeWelder)
+		{
+			constexpr float DegToRad = static_cast<float>(UE_PI / 180.f);
+
+			SplitAttributeWelder.UVDistSqrdThreshold = SplitUVThreshold * SplitUVThreshold;
+			SplitAttributeWelder.NormalVecDotThreshold = FMath::Abs(1.f - FMath::Cos( DegToRad * SplitNormalThreshold));
+			SplitAttributeWelder.TangentVecDotThreshold = FMath::Abs(1.f - FMath::Cos( DegToRad * SplitTangentsThreshold));
+			SplitAttributeWelder.ColorDistSqrdThreshold = SplitColorThreshold * SplitColorThreshold;
+		};
+
 		if ((Progress && Progress->Cancelled()) || !SourceMesh)
 		{
 			return;
@@ -74,6 +102,8 @@ public:
 		Merger.MergeVertexTolerance = Tolerance;
 		Merger.MergeSearchTolerance = 2 * Merger.MergeVertexTolerance;
 		Merger.OnlyUniquePairs = bOnlyUnique;
+		Merger.bWeldAttrsOnMergedEdges = (WeldAttributeMode == EWeldAttributeMode::OnWeldedMeshEdgesOnly);
+		SetAttributeWeldOptions(Merger.SplitAttributeWelder);
 
 		bool bOK = Merger.Apply();
 		FinalNumBoundaryEdges = Merger.FinalNumBoundaryEdges;
@@ -90,11 +120,21 @@ public:
 				SecondPassMerger.MergeVertexTolerance = Tolerance;
 				SecondPassMerger.MergeSearchTolerance = 2 * SecondPassMerger.MergeVertexTolerance;
 				SecondPassMerger.OnlyUniquePairs = bOnlyUnique;
+				SecondPassMerger.bWeldAttrsOnMergedEdges = (WeldAttributeMode == EWeldAttributeMode::OnWeldedMeshEdgesOnly);
+				SetAttributeWeldOptions(SecondPassMerger.SplitAttributeWelder);
+
 				bOK = SecondPassMerger.Apply();
 				FinalNumBoundaryEdges = SecondPassMerger.FinalNumBoundaryEdges;
 			}
 		}
 
+		if (WeldAttributeMode == EWeldAttributeMode::OnFullMesh)
+		{
+			FSplitAttributeWelder AttributeWelder;
+			SetAttributeWeldOptions(AttributeWelder);
+			AttributeWelder.WeldSplitElements(*ResultMesh.Get());
+		}
+
 		FDynamicMesh3::FValidityOptions UseValidityCheck = (bWasCleanMesh) ?
 			FDynamicMesh3::FValidityOptions() : FDynamicMesh3::FValidityOptions::Permissive();
 		if (bOK == false)
@@ -175,6 +215,11 @@ void UWeldMeshEdgesTool::Setup()
 	Settings->WatchProperty(Settings->Tolerance, [this](float) { PreviewCompute->InvalidateResult(); });
 	Settings->WatchProperty(Settings->bOnlyUnique, [this](bool) { PreviewCompute->InvalidateResult(); });
 	Settings->WatchProperty(Settings->bResolveTJunctions, [this](bool) { PreviewCompute->InvalidateResult(); });
+	Settings->WatchProperty(Settings->AttrWeldingMode, [this](EWeldMeshEdgesAttributeUIMode) { PreviewCompute->InvalidateResult(); });
+	Settings->WatchProperty(Settings->SplitUVThreshold, [this](float) { PreviewCompute->InvalidateResult(); });
+	Settings->WatchProperty(Settings->SplitColorThreshold, [this](float) { PreviewCompute->InvalidateResult(); });
+	Settings->WatchProperty(Settings->SplitNormalThreshold, [this](float) { PreviewCompute->InvalidateResult(); });
+	Settings->WatchProperty(Settings->SplitTangentsThreshold, [this](float) { PreviewCompute->InvalidateResult(); });
 
 	// create mesh display
 	MeshElementsDisplay = NewObject<UMeshElementsVisualizer>(this);
@@ -264,9 +309,27 @@ void UWeldMeshEdgesTool::OnTick(float DeltaTime)
 
 void UWeldMeshEdgesTool::UpdateOpParameters(FWeldMeshEdgesOp& Op) const
 {
+	auto GetAttrMergeMode = [](const EWeldMeshEdgesAttributeUIMode& m)
+	{
+		switch (m)
+		{
+			case EWeldMeshEdgesAttributeUIMode::None :  return FWeldMeshEdgesOp::EWeldAttributeMode::None;
+			case EWeldMeshEdgesAttributeUIMode::OnWeldedMeshEdgesOnly: return FWeldMeshEdgesOp::EWeldAttributeMode::OnWeldedMeshEdgesOnly;
+			case EWeldMeshEdgesAttributeUIMode::OnFullMesh : return FWeldMeshEdgesOp::EWeldAttributeMode::OnFullMesh;
+			default : return FWeldMeshEdgesOp::EWeldAttributeMode::None;
+		}
+	};
+
 	Op.bOnlyUnique = Settings->bOnlyUnique;
 	Op.bResolveTJunctions = Settings->bResolveTJunctions;
 	Op.Tolerance = Settings->Tolerance;
+
+	Op.WeldAttributeMode = GetAttrMergeMode(Settings->AttrWeldingMode);
+	Op.SplitUVThreshold = Settings->SplitUVThreshold;
+	Op.SplitColorThreshold = Settings->SplitColorThreshold;
+	Op.SplitNormalThreshold = Settings->SplitNormalThreshold;
+	Op.SplitTangentsThreshold = Settings->SplitTangentsThreshold;
+	
 	Op.SourceMesh = SourceMesh;
 
 	FTransform LocalToWorld = (FTransform)UE::ToolTarget::GetLocalToWorldTransform(Target);
diff --git a/Engine/Plugins/Experimental/MeshModelingToolsetExp/Source/MeshModelingToolsExp/Public/WeldMeshEdgesTool.h b/Engine/Plugins/Experimental/MeshModelingToolsetExp/Source/MeshModelingToolsExp/Public/WeldMeshEdgesTool.h
--- a/Engine/Plugins/Experimental/MeshModelingToolsetExp/Source/MeshModelingToolsExp/Public/WeldMeshEdgesTool.h
+++ b/Engine/Plugins/Experimental/MeshModelingToolsetExp/Source/MeshModelingToolsExp/Public/WeldMeshEdgesTool.h
@@ -28,6 +28,16 @@ public:
 	virtual USingleSelectionMeshEditingTool* CreateNewTool(const FToolBuilderState& SceneState) const override;
 };
 
+UENUM()
+enum class EWeldMeshEdgesAttributeUIMode : uint8
+{
+	/** Do not weld split-attributes*/
+	None,
+	/** Apply attribute welding only along the current mesh welds */
+	OnWeldedMeshEdgesOnly,
+	/** Apply attribute welding to all split-attributes */
+	OnFullMesh
+};
 
 
 UCLASS()
@@ -55,6 +65,28 @@ public:
 	/** Number of remaining open boundary edges */
 	UPROPERTY(VisibleAnywhere, Category = Statistics)
 	int32 RemainingEdges;
+
+	/** Controls split-attribute welding performed after the Mesh weld.  Applies to normals, tangents, UVs and colors. */
+	UPROPERTY(EditAnywhere, Category = Options, meta = (DisplayName = "Split-Attribute Welding"))
+	EWeldMeshEdgesAttributeUIMode AttrWeldingMode = EWeldMeshEdgesAttributeUIMode::OnWeldedMeshEdgesOnly;
+
+
+	/** Threshold on the angle between normals used to determine if split normals should be merged*/
+	UPROPERTY(EditAnywhere, Category = Options, AdvancedDisplay, meta = (EditCondition = "AttrWeldingMode != EWeldMeshEdgesAttributeUIMode::None", UIMin = "0.0", UIMax = "180.0", ClampMin = "0.0", ClampMax = "180.0", DisplayName="Normal Angle Threshold"))
+	float SplitNormalThreshold = 0.1f;
+
+	/** Threshold on the angle between tangent used to determine if split tangents should be merged*/
+	UPROPERTY(EditAnywhere, Category = Options, AdvancedDisplay, meta = (EditCondition = "AttrWeldingMode != EWeldMeshEdgesAttributeUIMode::None", UIMin = "0.0", UIMax = "180.0", ClampMin = "0.0", ClampMax = "180.0", DisplayName="Tangent Angle Threshold"))
+	float SplitTangentsThreshold = 0.1f;
+
+	/** Threshold uv-distance used to determine if split UVs should be merged */
+	UPROPERTY(EditAnywhere, Category = Options, AdvancedDisplay, meta = (EditCondition = "AttrWeldingMode != EWeldMeshEdgesAttributeUIMode::None",  UIMin = "0" , DisplayName = "UV Distance Threshold"))
+	float SplitUVThreshold = 0.01f;
+
+	/** Threshold color-distance used to determine if split colors should be merged */
+	UPROPERTY(EditAnywhere, Category = Options, AdvancedDisplay, meta = (EditCondition = "AttrWeldingMode != EWeldMeshEdgesAttributeUIMode::None",  UIMin = "0", DisplayName = "Color Distance Threshold"))
+	float SplitColorThreshold = 0.01f;
+
 };
 
 
diff --git a/Engine/Source/Runtime/GeometryCore/Private/DynamicMesh/Operations/MergeCoincidentMeshEdges.cpp b/Engine/Source/Runtime/GeometryCore/Private/DynamicMesh/Operations/MergeCoincidentMeshEdges.cpp
--- a/Engine/Source/Runtime/GeometryCore/Private/DynamicMesh/Operations/MergeCoincidentMeshEdges.cpp
+++ b/Engine/Source/Runtime/GeometryCore/Private/DynamicMesh/Operations/MergeCoincidentMeshEdges.cpp
@@ -2,6 +2,7 @@
 
 
 #include "DynamicMesh/Operations/MergeCoincidentMeshEdges.h"
+#include "DynamicMesh/Operations/SplitAttributeWelder.h"
 #include "DynamicMesh/DynamicMesh3.h"
 #include "DynamicMesh/MeshAdapterUtil.h"
 #include "Spatial/PointSetHashTable.h"
@@ -177,6 +178,13 @@ bool FMergeCoincidentMeshEdges::Apply()
 				delete EquivalenceSets[other_eid];
 				EquivalenceSets[other_eid] = nullptr;
 				RemainingEdges.Remove(other_eid);
+
+				// weld attributes 
+				if (bWeldAttrsOnMergedEdges)
+				{ 
+					SplitAttributeWelder.WeldSplitElements(*Mesh, mergeInfo.KeptVerts[0]);
+					SplitAttributeWelder.WeldSplitElements(*Mesh, mergeInfo.KeptVerts[1]);
+				}
 			}
 		}
 
diff --git a/Engine/Source/Runtime/GeometryCore/Private/DynamicMesh/Operations/SplitAttributeWelder.cpp b/Engine/Source/Runtime/GeometryCore/Private/DynamicMesh/Operations/SplitAttributeWelder.cpp
new file mode 100644
--- /dev/null
+++ b/Engine/Source/Runtime/GeometryCore/Private/DynamicMesh/Operations/SplitAttributeWelder.cpp
@@ -0,0 +1,197 @@
+// Copyright Epic Games, Inc. All Rights Reserved.
+
+
+#include "DynamicMesh/Operations/SplitAttributeWelder.h"
+#include "DynamicMesh/DynamicMesh3.h"
+#include "DynamicMesh/DynamicMeshAttributeSet.h"
+
+
+using namespace UE::Geometry;
+
+namespace
+{
+
+
+template <typename OverlayType, typename ShouldWeldFunctorType>
+void WeldSplits(const UE::Geometry::FDynamicMesh3* ParentMesh, const int32 ParentVID, OverlayType& Overlay,  ShouldWeldFunctorType& ShouldWeld)
+{
+	if (!ParentMesh || !ParentMesh->IsVertex(ParentVID))
+	{
+		return;
+	}
+
+	TArray<int> ElementIDs;
+	Overlay.GetVertexElements(ParentVID, ElementIDs);
+
+	// assume the number of elements at given vertex is small and do simple O(n^2) 
+	const int32 NumElements = ElementIDs.Num();
+	TArray<int> ConsumedMask;  ConsumedMask.SetNumZeroed(NumElements);
+	for (int32 i = 0; i < NumElements; ++i)
+	{
+		if (ConsumedMask[i] == 1) 
+		{
+			continue;
+		}
+
+		const int32 eid = ElementIDs[i];
+
+		for (int32 j = i + 1; j < NumElements; ++j)
+		{
+			if(ConsumedMask[j] == 1)
+			{
+				continue;
+			}
+			const int32 oeid = ElementIDs[j];
+
+			if (ShouldWeld(eid, oeid))
+			{
+				// consumed by weld.
+				ConsumedMask[j] = 1;
+
+				for (int TID : ParentMesh->VtxTrianglesItr(ParentVID)) // only care about triangles connected to the *new* vertex; these are updated
+				{
+					if (!Overlay.IsSetTriangle(TID))
+					{
+						continue;
+					}
+
+					FIndex3i TriElements = Overlay.GetTriangle(TID);
+					bool bUpdateTriangle = false;
+					for (int c = 0; c < 3; ++c)
+					{
+						if (TriElements[c] == oeid)
+						{
+							TriElements[c] = eid;
+							bUpdateTriangle = true;
+						}
+					}
+					if (bUpdateTriangle)
+					{
+						Overlay.SetTriangle(TID, TriElements, true /* allow element freeing*/);
+					}
+				}
+			}
+		}
+
+	}
+
+}
+
+}
+
+void  FSplitAttributeWelder::WeldSplitElements(FDynamicMesh3& ParentMesh, const int32 ParentVID)
+{
+
+	
+
+	FDynamicMeshAttributeSet* Attributes = ParentMesh.Attributes();
+
+	if (!Attributes)
+	{
+		return;
+	}
+
+	if (!ParentMesh.IsVertex(ParentVID))
+	{
+		return;
+	}
+
+	// uvs 
+	for (int32 i = 0, I = Attributes->NumUVLayers(); i < I; ++i)
+	{
+		FDynamicMeshUVOverlay* Overlay = Attributes->GetUVLayer(i);
+		WeldSplitUVs(ParentVID, *Overlay, UVDistSqrdThreshold);
+	}
+
+	// tangent space
+	for (int32 i = 0, I = Attributes->NumNormalLayers(); i < I; ++i )
+	{
+		const float DotThreshold = (i == 0) ? NormalVecDotThreshold : TangentVecDotThreshold;
+		FDynamicMeshNormalOverlay* Overlay = Attributes->GetNormalLayer(i);
+		WeldSplitUnitVectors(ParentVID, *Overlay, DotThreshold);
+	}
+
+	// colors
+	if (FDynamicMeshColorOverlay* Overlay = Attributes->PrimaryColors())
+	{
+		WeldSplitColors(ParentVID, *Overlay, ColorDistSqrdThreshold);
+	}
+}
+
+
+void FSplitAttributeWelder::WeldSplitElements(FDynamicMesh3& ParentMesh)
+{
+	for (int vid : ParentMesh.VertexIndicesItr())
+	{
+		WeldSplitElements(ParentMesh, vid);
+	}
+}
+
+void FSplitAttributeWelder::WeldSplitUVs(const int32 ParentVID, FDynamicMeshUVOverlay& Overlay, float UVDistSqrdThreshold)
+{
+
+	const FDynamicMesh3* ParentMesh = Overlay.GetParentMesh();
+
+	const float Threshold = FMath::Max(UVDistSqrdThreshold, 0.f);
+
+	auto ShouldWeld = [&Overlay, Threshold](const int32 eid, const int32 oeid)->bool
+	{
+		const FVector2f UV            = Overlay.GetElement(eid);
+		const FVector2f otherUV       = Overlay.GetElement(oeid);
+		const float UVDistanceSquared = FVector2f::DistSquared(UV, otherUV);
+
+		return (UVDistanceSquared <= Threshold);
+	};
+
+	WeldSplits(ParentMesh, ParentVID, Overlay, ShouldWeld);
+}
+
+// note, this only compares vector orientation - not length.
+void FSplitAttributeWelder::WeldSplitUnitVectors(const int32 ParentVID, FDynamicMeshNormalOverlay& Overlay, float DotThreshold, bool bMergeZeroVectors)
+{
+	const FDynamicMesh3* ParentMesh = Overlay.GetParentMesh();
+
+	auto ShouldWeld = [&Overlay, DotThreshold, bMergeZeroVectors](const int32 eid, const int32 oeid)->bool
+	{
+		FVector3f Vec      = Overlay.GetElement(eid);
+		FVector3f otherVec = Overlay.GetElement(oeid);
+
+		const bool bVecNormalized = Vec.Normalize();
+		const bool bOtherVecNormalized = otherVec.Normalize();
+	
+		if (bVecNormalized && bOtherVecNormalized)
+		{
+			const float  CosAngle = FVector3f::DotProduct(Vec, otherVec);
+			return ( FMath::Abs(1.f - CosAngle) <= DotThreshold);
+		}
+
+		const bool bBothZero = (!bVecNormalized && !bOtherVecNormalized);
+		if (bMergeZeroVectors && bBothZero)
+		{
+			return true;
+		}
+
+		return false;
+	};
+
+	WeldSplits(ParentMesh, ParentVID, Overlay, ShouldWeld);
+}
+
+void FSplitAttributeWelder::WeldSplitColors(const int32 ParentVID, FDynamicMeshColorOverlay& Overlay, float ColorDistSqrdThreshold)
+{
+	const FDynamicMesh3* ParentMesh = Overlay.GetParentMesh();
+
+	const float Threshold = FMath::Max(ColorDistSqrdThreshold, 0.f);
+
+	auto ShouldWeld = [&Overlay, Threshold](const int32 eid, const int32 oeid)->bool
+	{
+		const FVector4f Color = Overlay.GetElement(eid);
+		const FVector4f otherColor = Overlay.GetElement(oeid);
+
+		const float ColorDistSqrd = (Color - otherColor).SizeSquared();
+
+		return (ColorDistSqrd <= Threshold);
+	};
+
+	WeldSplits(ParentMesh, ParentVID, Overlay, ShouldWeld);
+}
diff --git a/Engine/Source/Runtime/GeometryCore/Public/DynamicMesh/Operations/MergeCoincidentMeshEdges.h b/Engine/Source/Runtime/GeometryCore/Public/DynamicMesh/Operations/MergeCoincidentMeshEdges.h
--- a/Engine/Source/Runtime/GeometryCore/Public/DynamicMesh/Operations/MergeCoincidentMeshEdges.h
+++ b/Engine/Source/Runtime/GeometryCore/Public/DynamicMesh/Operations/MergeCoincidentMeshEdges.h
@@ -7,6 +7,7 @@
 #include "Math/UnrealMathSSE.h"
 #include "MathUtil.h"
 #include "VectorTypes.h"
+#include "SplitAttributeWelder.h"
 
 namespace UE
 {
@@ -52,6 +53,13 @@ public:
 	/** Number of mesh boundary edges after merging */
 	int32 FinalNumBoundaryEdges = 0;
 
+	/** Enable / Disable attribute welding along merged mesh edges */
+	bool bWeldAttrsOnMergedEdges = false;
+
+	/** Used to weld attributes at the merged edges */
+	FSplitAttributeWelder  SplitAttributeWelder;
+
+
 public:
 	FMergeCoincidentMeshEdges(FDynamicMesh3* mesh) : Mesh(mesh)
 	{
diff --git a/Engine/Source/Runtime/GeometryCore/Public/DynamicMesh/Operations/SplitAttributeWelder.h b/Engine/Source/Runtime/GeometryCore/Public/DynamicMesh/Operations/SplitAttributeWelder.h
new file mode 100644
--- /dev/null
+++ b/Engine/Source/Runtime/GeometryCore/Public/DynamicMesh/Operations/SplitAttributeWelder.h
@@ -0,0 +1,104 @@
+// Copyright Epic Games, Inc. All Rights Reserved.
+
+
+#pragma once
+
+#include "Math/MathFwd.h"
+
+namespace UE
+{
+namespace Geometry
+{
+
+class FDynamicMesh3;
+
+template<typename RealType, int ElementSize, typename VectorType>
+class TDynamicMeshVectorOverlay; 
+
+
+/**
+* FSplitAttributeWelder can be used to weld split-attributes that share the same vertex in the parent mesh (e.g. split normals)
+*
+*/
+class GEOMETRYCORE_API FSplitAttributeWelder
+{
+
+	typedef TDynamicMeshVectorOverlay<float, 2, FVector2f> FDynamicMeshUVOverlay;
+	typedef TDynamicMeshVectorOverlay<float, 3, FVector3f> FDynamicMeshNormalOverlay;
+	typedef TDynamicMeshVectorOverlay<float, 4, FVector4f> FDynamicMeshColorOverlay;
+
+public:
+
+	FSplitAttributeWelder() {}
+
+
+	/**
+	* Merge threshold used to compare UV Overlay elements.
+	* Applied as FVector2f::DistSquared(UVA , UVB) <= UVDisSqrdThreshold
+	*/ 
+	float UVDistSqrdThreshold = 0.f;
+
+	/**
+	* Merge threshold used to compare Color Overlay elements
+	* Applied as (ColorA - ColorB).SizeSquared() <= ColorDistSqrdThreshold
+	*/ 
+	float ColorDistSqrdThreshold = 0.f;
+	
+	/**
+	* Merge threshold used to compare Normal Overlay vectors. 
+	* Applied as Abs(1 - VecA.dot.VecB)  <= NormalVecDotThreshold.
+	*/
+	float NormalVecDotThreshold = 0.f;
+
+
+	/**
+	* Merge threshold used to compare Tangent (and BiTangent) Overlay vectors. 
+	* Applied as Abs(1 - VecA.dot.VecB)  <= TangentVecDotThreshold.
+	*/
+	float TangentVecDotThreshold = 0.f;
+
+	/**
+	* Weld split-elements at the ParentVID in each overlay that are "close enough" as defined by the appropriate threshold. 
+	*/
+	void WeldSplitElements(FDynamicMesh3& ParentMesh, const int32 ParentVID);
+	
+
+	/**
+	*  Weld split-elements across the entire mesh. 
+	*/
+	void WeldSplitElements(FDynamicMesh3& ParentMesh);
+	
+
+	/**
+	* Welds split UVs shared by the ParentVID vertex in the ParentMesh.  
+	* 
+	* @param ParentVID  defines a vertex in the parent mesh. 
+	* @param UVOverlay contains potentially split UVs associated with the ParentVID. Assumed to be an overlay of the ParentMesh.
+	* @param UVDistSqrdThreshold  provides a threshold in distance squared UV space between split UVs that should be welded. 
+	*/
+	static void WeldSplitUVs(const int32 ParentVID, FDynamicMeshUVOverlay& UVOverlay, float UVDistSqrdThreshold);
+	
+
+	/**
+	* Welds split Normals / Tangents shared by the ParentVID vertex in the ParentMesh based on the angled between the vectors.
+	* 
+	* @param ParentVID  defines a vertex in the Parent Mesh 
+	* @param NormalOverlay contains potentially split vectors associated with the ParentVID. Assumed to be an overlay of the ParentMesh.
+	* @param DotThreshold  provides a threshold for the dot product between two vectors that should be welded.
+	* @param bMergeZeroVectors will weld vectors that are too small to normalize.
+	*/
+	static void WeldSplitUnitVectors(const int32 ParentVID, FDynamicMeshNormalOverlay& NormalOverlay, float DotThreshold, bool bMergeZeroVectors = true);
+	
+	/**
+	* Welds split colors shared by the ParentVID vertex in the ParentMesh.  
+	* 
+	* @param ParentVID  defines a vertex in the Parent Mesh 
+	* @param ColorOverlay contains potentially split colors associated with the ParentVID. Assumed to be an overlay of the ParentMesh.
+	* @param ColorDistSqrdThreshold  provides a threshold applied to the sum of the square differences between to colors.  Used to determine if a weld should be preformed.
+	*/
+	static void WeldSplitColors(const int32 ParentVID, FDynamicMeshColorOverlay& ColorOverlay, float ColorDistSqrdThreshold);
+
+};
+
+} // end Geometry
+} // end UE
