Skeleton Mesh Editing tool: multi selection + transform gizmo
    - multiselection: bones can be multi-selected (using shift/ctrl) so that they can be bulk edited
    - bones can be transformed individually (unlike in the level editor where the last selected element's transform is used as pivot)
    - added USkeletalMeshGizmoContextObject to manage gizmo wrappers for skeletal mesh tools
    - gizmo: re-used the transform gizmo work done with several fixes/improvements so that they can be used in other contexts
        - UTransformGizmo can be used without TransformGizmoSource, most of the infos can be retrieved from the IToolsContextQueriesAPI
        - UTransformGizmo can be passed an external StateTarget to handle transaction externally
        - UTransformGizmo fixed rotations going in the wrong direction when dragging

diff --git a/Engine/Plugins/Experimental/Animation/SkeletalMeshModelingTools/Source/SkeletalMeshModelingTools/Private/SkeletalMeshGizmoUtils.cpp b/Engine/Plugins/Experimental/Animation/SkeletalMeshModelingTools/Source/SkeletalMeshModelingTools/Private/SkeletalMeshGizmoUtils.cpp
new file mode 100644
--- /dev/null
+++ b/Engine/Plugins/Experimental/Animation/SkeletalMeshModelingTools/Source/SkeletalMeshModelingTools/Private/SkeletalMeshGizmoUtils.cpp
@@ -0,0 +1,207 @@
+﻿// Copyright Epic Games, Inc. All Rights Reserved.
+
+#include "SkeletalMeshGizmoUtils.h"
+
+#include "ContextObjectStore.h"
+#include "InteractiveToolsContext.h"
+#include "BaseGizmos/GizmoElementHitTargets.h"
+#include "SkeletalMesh/SkeletonTransformProxy.h"
+
+#include "EditorGizmos/TransformGizmo.h"
+#include "Editor/Experimental/EditorInteractiveToolsFramework/Public/EditorInteractiveGizmoManager.h"
+// #include "EditorGizmos/EditorTransformGizmoBuilder.h"
+
+#include UE_INLINE_GENERATED_CPP_BY_NAME(SkeletalMeshGizmoUtils)
+
+bool UE::SkeletalMeshGizmoUtils::RegisterTransformGizmoContextObject(UInteractiveToolsContext* ToolsContext)
+{
+	if (ensure(ToolsContext))
+	{
+		const USkeletalMeshGizmoContextObject* Found = ToolsContext->ContextObjectStore->FindContext<USkeletalMeshGizmoContextObject>();
+		if (Found)
+		{
+			return true;
+		}
+		
+		USkeletalMeshGizmoContextObject* GizmoContextObject = NewObject<USkeletalMeshGizmoContextObject>(ToolsContext->ToolManager);
+		if (ensure(GizmoContextObject))
+		{
+			GizmoContextObject->RegisterGizmosWithManager(ToolsContext->ToolManager);
+			return true;
+		}
+	}
+	return false;
+}
+
+bool UE::SkeletalMeshGizmoUtils::UnregisterTransformGizmoContextObject(UInteractiveToolsContext* ToolsContext)
+{
+	if (ensure(ToolsContext))
+	{
+		USkeletalMeshGizmoContextObject* Found = ToolsContext->ContextObjectStore->FindContext<USkeletalMeshGizmoContextObject>();
+		if (Found != nullptr)
+		{
+			Found->UnregisterGizmosWithManager(ToolsContext->ToolManager);
+			ToolsContext->ContextObjectStore->RemoveContextObject(Found);
+		}
+		return true;
+	}
+	return false;
+}
+
+UTransformGizmo* UE::SkeletalMeshGizmoUtils::CreateTransformGizmo(UInteractiveToolManager* ToolManager, void* Owner)
+{
+	if (!ToolManager)
+	{
+		return nullptr;
+	}
+	
+	UEditorInteractiveGizmoManager* GizmoManager = Cast<UEditorInteractiveGizmoManager>(ToolManager->GetPairedGizmoManager());
+	if (!GizmoManager)
+	{
+		return nullptr;
+	}
+
+	const USkeletalMeshGizmoContextObject* Found = GizmoManager->GetContextObjectStore()->FindContext<USkeletalMeshGizmoContextObject>();
+	if (!Found)
+	{
+		return nullptr;
+	}
+
+	static const FString InstanceIdentifier;
+	UInteractiveGizmo* NewGizmo = GizmoManager->CreateGizmo(
+		USkeletalMeshGizmoContextObject::TransformBuilderIdentifier(), InstanceIdentifier, Owner);
+
+	ensure(NewGizmo);
+	
+	return Cast<UTransformGizmo>(NewGizmo);
+}
+
+const FString& USkeletalMeshGizmoContextObject::TransformBuilderIdentifier()
+{
+	static const FString Identifier(TEXT("SkeletalMeshTransformGizmo"));
+	return Identifier;
+}
+
+void USkeletalMeshGizmoContextObject::RegisterGizmosWithManager(UInteractiveToolManager* InToolManager)
+{
+	if (ensure(!bGizmosRegistered) == false)
+	{
+		return;
+	}
+
+	InToolManager->GetContextObjectStore()->AddContextObject(this);
+	
+	UEditorInteractiveGizmoManager* GizmoManager = Cast<UEditorInteractiveGizmoManager>(InToolManager->GetPairedGizmoManager());
+	if (ensure(GizmoManager))
+	{
+		GizmoManager->RegisterGizmoType(TransformBuilderIdentifier(), GizmoManager->GetTransformGizmoBuilder());
+	}
+	
+	bGizmosRegistered = true;
+}
+
+void USkeletalMeshGizmoContextObject::UnregisterGizmosWithManager(UInteractiveToolManager* InToolManager)
+{
+	ensure(bGizmosRegistered);
+	
+	InToolManager->GetContextObjectStore()->RemoveContextObject(this);
+
+	UInteractiveGizmoManager* GizmoManager = InToolManager->GetPairedGizmoManager();
+	GizmoManager->DeregisterGizmoType(TransformBuilderIdentifier());
+	
+	bGizmosRegistered = false;
+}
+
+USkeletalMeshGizmoWrapperBase* USkeletalMeshGizmoContextObject::GetNewWrapper(UInteractiveToolManager* InToolManager, UObject* Outer, IGizmoStateTarget* InStateTarget)
+{
+	if (!bGizmosRegistered)
+	{
+		return nullptr;
+	}
+	
+	UTransformGizmo* Gizmo = UE::SkeletalMeshGizmoUtils::CreateTransformGizmo(InToolManager, Outer);
+	if (!Gizmo)
+	{
+		return nullptr;
+	}
+
+	Gizmo->TransformGizmoSource = nullptr;
+	Gizmo->SetVisibility(true);
+
+	USkeletonTransformProxy* SkeletonTransformProxy = NewObject<USkeletonTransformProxy>(Outer);
+	SkeletonTransformProxy->bRotatePerObject = true;
+	
+	Gizmo->SetActiveTarget(SkeletonTransformProxy, nullptr, InStateTarget);
+
+	if (UGizmoElementHitMultiTarget* HitMultiTarget = Cast< UGizmoElementHitMultiTarget>(Gizmo->HitTarget))
+	{
+		HitMultiTarget->GizmoTransformProxy = SkeletonTransformProxy;
+	}
+
+	USkeletalMeshGizmoWrapper* GizmoWrapper = NewObject<USkeletalMeshGizmoWrapper>(Outer);
+	GizmoWrapper->TransformGizmo = Gizmo;
+	GizmoWrapper->TransformProxy = SkeletonTransformProxy;
+	
+	return GizmoWrapper;
+}
+
+bool USkeletalMeshGizmoWrapper::IsGizmoHit(const FInputDeviceRay& PressPos) const
+{
+	if (!TransformGizmo || !TransformGizmo->bVisible)
+	{
+		return false;
+	}
+	return TransformGizmo->HitTarget ? TransformGizmo->HitTarget->IsHit(PressPos).bHit : false;
+}
+
+void USkeletalMeshGizmoWrapper::Clear()
+{
+	if (TransformProxy)
+	{
+		TransformProxy->Initialize(FTransform::Identity, EToolContextCoordinateSystem::Local);
+	}
+	
+	if (TransformGizmo)
+	{
+		TransformGizmo->SetVisibility(false);
+		TransformGizmo->ClearActiveTarget();
+
+		if (UInteractiveGizmoManager* GizmoManager = TransformGizmo->GetGizmoManager())
+		{
+			GizmoManager->DestroyGizmo(TransformGizmo);
+		}
+	}
+}
+
+void USkeletalMeshGizmoWrapper::Initialize(const FTransform& InTransform, const EToolContextCoordinateSystem& InTransformMode)
+{
+	if (TransformProxy)
+	{
+		TransformProxy->Initialize(InTransform, InTransformMode);
+	}
+	
+	if (TransformGizmo)
+	{
+		TransformGizmo->SetVisibility(TransformProxy ? TransformProxy->IsValid() : false);
+	}
+}
+
+void USkeletalMeshGizmoWrapper::HandleBoneTransform(FGetTransform GetTransformFunc, FSetTransform SetTransformFunc)
+{
+	if (!TransformProxy || !Component.IsValid())
+	{
+		return;
+	}
+
+	static constexpr bool bModifyComponentOnTransform = false;
+	TransformProxy->AddComponentCustom(Component.Get(),
+		MoveTemp(GetTransformFunc),
+		MoveTemp(SetTransformFunc),
+		INDEX_NONE,
+		bModifyComponentOnTransform);
+
+	if (TransformGizmo)
+	{
+		TransformGizmo->SetVisibility(TransformProxy->IsValid());
+	}
+}
\ No newline at end of file
diff --git a/Engine/Plugins/Experimental/Animation/SkeletalMeshModelingTools/Source/SkeletalMeshModelingTools/Private/SkeletalMeshGizmoUtils.h b/Engine/Plugins/Experimental/Animation/SkeletalMeshModelingTools/Source/SkeletalMeshModelingTools/Private/SkeletalMeshGizmoUtils.h
new file mode 100644
--- /dev/null
+++ b/Engine/Plugins/Experimental/Animation/SkeletalMeshModelingTools/Source/SkeletalMeshModelingTools/Private/SkeletalMeshGizmoUtils.h
@@ -0,0 +1,91 @@
+﻿// Copyright Epic Games, Inc. All Rights Reserved.
+
+#pragma once
+
+#include "UObject/Object.h"
+#include "Containers/UnrealString.h"
+#include "SkeletalMesh/SkeletalMeshEditionInterface.h"
+
+#include "SkeletalMeshGizmoUtils.generated.h"
+
+class UInteractiveToolsContext;
+class UInteractiveToolManager;
+class UTransformGizmo;
+class USkeletonTransformProxy;
+
+namespace UE
+{
+	
+namespace SkeletalMeshGizmoUtils
+{
+	/**
+	 * The functions below are helper functions that simplify usage of a USkeletalMeshGizmoContextObject
+	 * that is registered as a ContextStoreObject in an InteractiveToolsContext
+	 */
+
+	/**
+	 * If one does not already exist, create a new instance of USkeletalMeshGizmoContextObject and add it to the
+	 * ToolsContext's ContextObjectStore
+	 * @return true if the ContextObjectStore now has a USkeletalMeshGizmoContextObject (whether it already existed, or was created)
+	 */
+	SKELETALMESHMODELINGTOOLS_API bool RegisterTransformGizmoContextObject(UInteractiveToolsContext* InToolsContext);
+
+	/**
+	 * Remove any existing USkeletalMeshGizmoContextObject from the ToolsContext's ContextObjectStore
+	 * @return true if the ContextObjectStore no longer has a USkeletalMeshGizmoContextObject (whether it was removed, or did not exist)
+	 */
+	SKELETALMESHMODELINGTOOLS_API bool UnregisterTransformGizmoContextObject(UInteractiveToolsContext* InToolsContext);
+
+	/**
+	 * Spawn an editor like Transform Gizmo. ToolManager's ToolsContext must have a USkeletalMeshGizmoContextObject registered 
+	 */
+	SKELETALMESHMODELINGTOOLS_API UTransformGizmo* CreateTransformGizmo(UInteractiveToolManager* ToolManager, void* Owner = nullptr);
+}
+	
+}
+
+/**
+ * USkeletalMeshGizmoWrapper is a wrapper class to handle a Transform Gizmo and it's Transform proxy so that it can
+ * be used to update skeletal mesh infos using a Gizmo.
+ */
+
+UCLASS()
+class SKELETALMESHMODELINGTOOLS_API USkeletalMeshGizmoWrapper : public USkeletalMeshGizmoWrapperBase
+{
+	GENERATED_BODY()
+public:
+	virtual void Initialize(const FTransform& InTransform = FTransform::Identity, const EToolContextCoordinateSystem& InTransformMode = EToolContextCoordinateSystem::Local) override;
+	virtual void HandleBoneTransform(FGetTransform GetTransformFunc, FSetTransform SetTransformFunc) override;
+	virtual void Clear() override;
+
+	virtual bool IsGizmoHit(const FInputDeviceRay& PressPos) const override;
+
+	UPROPERTY()
+	TObjectPtr<UTransformGizmo> TransformGizmo;
+	UPROPERTY()
+	TObjectPtr<USkeletonTransformProxy> TransformProxy;
+};
+
+/**
+ * USkeletalMeshGizmoContextObject is a utility object that registers a Gizmo Builder for UTransformGizmo.
+ * (see UCombinedTransformGizmoContextObject for more details)
+ */
+
+UCLASS()
+class SKELETALMESHMODELINGTOOLS_API USkeletalMeshGizmoContextObject : public USkeletalMeshGizmoContextObjectBase
+{
+	GENERATED_BODY()
+
+public:
+
+	// builder identifiers for transform gizmo
+	static const FString& TransformBuilderIdentifier();
+	
+	void RegisterGizmosWithManager(UInteractiveToolManager* InToolManager);
+	void UnregisterGizmosWithManager(UInteractiveToolManager* InToolManager);
+
+	virtual USkeletalMeshGizmoWrapperBase* GetNewWrapper(UInteractiveToolManager* InToolManager, UObject* Outer = nullptr, IGizmoStateTarget* InStateTarget = nullptr) override;
+
+private:
+	bool bGizmosRegistered = false;
+};
diff --git a/Engine/Plugins/Experimental/Animation/SkeletalMeshModelingTools/Source/SkeletalMeshModelingTools/Private/SkeletalMeshModelingToolsEditorMode.cpp b/Engine/Plugins/Experimental/Animation/SkeletalMeshModelingTools/Source/SkeletalMeshModelingTools/Private/SkeletalMeshModelingToolsEditorMode.cpp
--- a/Engine/Plugins/Experimental/Animation/SkeletalMeshModelingTools/Source/SkeletalMeshModelingTools/Private/SkeletalMeshModelingToolsEditorMode.cpp
+++ b/Engine/Plugins/Experimental/Animation/SkeletalMeshModelingTools/Source/SkeletalMeshModelingTools/Private/SkeletalMeshModelingToolsEditorMode.cpp
@@ -4,6 +4,7 @@
 
 #include "SkeletalMeshModelingToolsEditorModeToolkit.h"
 #include "SkeletalMeshModelingToolsCommands.h"
+#include "SkeletalMeshGizmoUtils.h"
 
 #include "BaseGizmos/TransformGizmoUtil.h"
 #include "EdMode.h"
@@ -19,6 +20,7 @@
 #include "DisplaceMeshTool.h"
 #include "DynamicMeshSculptTool.h"
 #include "EditMeshPolygonsTool.h"
+#include "EditorInteractiveGizmoManager.h"
 #include "HoleFillTool.h"
 #include "LatticeDeformerTool.h"
 #include "MeshAttributePaintTool.h"
@@ -38,6 +40,7 @@
 #include "WeldMeshEdgesTool.h"
 
 #include "SkeletalMeshNotifier.h"
+#include "Animation/DebugSkelMeshComponent.h"
 #include "SkeletalMesh/SkeletalMeshEditionInterface.h"
 
 #include "SkeletalMesh/SkeletonEditingTool.h"
@@ -167,11 +170,14 @@ void USkeletalMeshModelingToolsEditorMode::Enter()
 {
 	UEdMode::Enter();
 
-	GetInteractiveToolsContext()->TargetManager->AddTargetFactory(NewObject<USkeletalMeshComponentToolTargetFactory>(GetInteractiveToolsContext()->TargetManager));
+	UEditorInteractiveToolsContext* InteractiveToolsContext = GetInteractiveToolsContext();
+	
+	InteractiveToolsContext->TargetManager->AddTargetFactory(NewObject<USkeletalMeshComponentToolTargetFactory>(InteractiveToolsContext->TargetManager));
 
 	StylusStateTracker = MakeUnique<FStylusStateTracker>();
 	// register gizmo helper
-	UE::TransformGizmoUtil::RegisterTransformGizmoContextObject(GetInteractiveToolsContext());
+	UE::TransformGizmoUtil::RegisterTransformGizmoContextObject(InteractiveToolsContext);
+	UE::SkeletalMeshGizmoUtils::RegisterTransformGizmoContextObject(InteractiveToolsContext);
 
 	const FModelingToolsManagerCommands& ToolManagerCommands = FModelingToolsManagerCommands::Get();
 
@@ -237,9 +243,23 @@ void USkeletalMeshModelingToolsEditorMode::Enter()
 	GetInteractiveToolsContext()->ToolManager->SelectActiveToolType(EToolSide::Left, TEXT("BeginSkinWeightsPaintTool"));
 }
 
+UDebugSkelMeshComponent* USkeletalMeshModelingToolsEditorMode::GetSkelMeshComponent() const
+{
+	FToolBuilderState State; GetToolManager()->GetContextQueriesAPI()->GetCurrentSelectionState(State);
+	UActorComponent* SkeletalMeshComponent = ToolBuilderUtil::FindFirstComponent(State, [&](UActorComponent* Component)
+	{
+		return IsValid(Component) && Component->IsA<UDebugSkelMeshComponent>();
+	});
+
+	return Cast<UDebugSkelMeshComponent>(SkeletalMeshComponent);
+}
 
 void USkeletalMeshModelingToolsEditorMode::Exit()
 {
+	UEditorInteractiveToolsContext* InteractiveToolsContext = GetInteractiveToolsContext();
+	UE::TransformGizmoUtil::DeregisterTransformGizmoContextObject(InteractiveToolsContext);
+	UE::SkeletalMeshGizmoUtils::UnregisterTransformGizmoContextObject(InteractiveToolsContext);
+	
 	StylusStateTracker = nullptr;
 
 	UEdMode::Exit();
@@ -357,7 +377,7 @@ void USkeletalMeshModelingToolsEditorMode::SetEditorBinding(TSharedPtr<ISkeletal
 
 ISkeletalMeshEditionInterface* USkeletalMeshModelingToolsEditorMode::GetSkeletonInterface(UInteractiveTool* InTool)
 {
-	if (!IsValid(InTool) && !InTool->Implements<USkeletalMeshEditionInterface>())
+	if (!IsValid(InTool) || !InTool->Implements<USkeletalMeshEditionInterface>())
 	{
 		return nullptr;
 	}
@@ -396,6 +416,8 @@ void USkeletalMeshModelingToolsEditorMode::ConnectTool(UInteractiveTool* InTool)
 				BindingPtr->GetNotifier().HandleNotification(BoneNames, InNotifyType);
 			});
 		}
+
+		SkeletonInterface->GetNotifier().HandleNotification(BindingPtr->GetSelectedBones(), ESkeletalMeshNotifyType::BonesSelected);
 	}
 }
 
@@ -421,4 +443,20 @@ void USkeletalMeshModelingToolsEditorMode::DisconnectTool(UInteractiveTool* InTo
 	}
 }
 
+bool USkeletalMeshModelingToolsEditorMode::NeedsTransformGizmo() const
+{
+	UInteractiveTool* Tool = GetToolManager()->GetActiveTool(EToolSide::Mouse);
+	if (const ISkeletalMeshEditionInterface* SkeletonInterface = GetSkeletonInterface(Tool))
+	{
+		return !SkeletonInterface->GetSelectedBones().IsEmpty();
+	}
+
+	if (Binding.IsValid())
+	{
+		return !Binding.Pin()->GetSelectedBones().IsEmpty();
+	}
+	
+	return false;
+}
+
 #undef LOCTEXT_NAMESPACE
diff --git a/Engine/Plugins/Experimental/Animation/SkeletalMeshModelingTools/Source/SkeletalMeshModelingTools/Private/SkeletalMeshModelingToolsEditorMode.h b/Engine/Plugins/Experimental/Animation/SkeletalMeshModelingTools/Source/SkeletalMeshModelingTools/Private/SkeletalMeshModelingToolsEditorMode.h
--- a/Engine/Plugins/Experimental/Animation/SkeletalMeshModelingTools/Source/SkeletalMeshModelingTools/Private/SkeletalMeshModelingToolsEditorMode.h
+++ b/Engine/Plugins/Experimental/Animation/SkeletalMeshModelingTools/Source/SkeletalMeshModelingTools/Private/SkeletalMeshModelingToolsEditorMode.h
@@ -14,6 +14,7 @@ class ISkeletalMeshNotifier;
 class ISkeletalMeshEditorBinding;
 class ISkeletalMeshEditionInterface;
 class HHitProxy;
+class UDebugSkelMeshComponent;
 
 UCLASS()
 class USkeletalMeshModelingToolsEditorMode : 
@@ -55,9 +56,13 @@ private:
 
 	static ISkeletalMeshEditionInterface* GetSkeletonInterface(UInteractiveTool* InTool);
 
+	UDebugSkelMeshComponent* GetSkelMeshComponent() const;
+
 	void ConnectTool(UInteractiveTool* InTool);
 	void DisconnectTool(UInteractiveTool* InTool);
 
+	bool NeedsTransformGizmo() const;
+
 	FDelegateHandle ToToolNotifierHandle;
 	FDelegateHandle FromToolNotifierHandle;
 	
diff --git a/Engine/Plugins/Experimental/MeshModelingToolsetExp/Source/MeshModelingToolsEditorOnlyExp/Private/ISMEditorTool.cpp b/Engine/Plugins/Experimental/MeshModelingToolsetExp/Source/MeshModelingToolsEditorOnlyExp/Private/ISMEditorTool.cpp
--- a/Engine/Plugins/Experimental/MeshModelingToolsetExp/Source/MeshModelingToolsEditorOnlyExp/Private/ISMEditorTool.cpp
+++ b/Engine/Plugins/Experimental/MeshModelingToolsetExp/Source/MeshModelingToolsEditorOnlyExp/Private/ISMEditorTool.cpp
@@ -555,7 +555,7 @@ static void AddInstancedComponentInstance(UInstancedStaticMeshComponent* ISMC, i
 			ISMC->GetInstanceTransform(Index, Tmp, true);
 			return Tmp;
 		},
-		[ISMC, Index](FTransform NewTransform) {
+		[ISMC, Index](const FTransform& NewTransform) {
 			ISMC->UpdateInstanceTransform(Index, NewTransform, true, true, true);
 		},
 		Index, bModifyOnTransform
diff --git a/Engine/Plugins/Experimental/MeshModelingToolsetExp/Source/MeshModelingToolsExp/Private/SkeletalMesh/SkeletalMeshEditionInterface.cpp b/Engine/Plugins/Experimental/MeshModelingToolsetExp/Source/MeshModelingToolsExp/Private/SkeletalMesh/SkeletalMeshEditionInterface.cpp
--- a/Engine/Plugins/Experimental/MeshModelingToolsetExp/Source/MeshModelingToolsExp/Private/SkeletalMesh/SkeletalMeshEditionInterface.cpp
+++ b/Engine/Plugins/Experimental/MeshModelingToolsetExp/Source/MeshModelingToolsExp/Private/SkeletalMesh/SkeletalMeshEditionInterface.cpp
@@ -39,6 +39,12 @@ TOptional<FName> ISkeletalMeshEditionInterface::GetBoneName(HHitProxy* InHitProx
 	return Binding.IsValid() && Binding.Pin()->GetNameFunction() ? Binding.Pin()->GetNameFunction()(InHitProxy) : TOptional<FName>();
 }
 
+TArray<FName> ISkeletalMeshEditionInterface::GetSelectedBones() const
+{
+	static const TArray<FName> Dummy;
+	return Binding.IsValid() ? Binding.Pin()->GetSelectedBones() : Dummy;
+}
+
 FSkeletalMeshToolNotifier::FSkeletalMeshToolNotifier(TWeakInterfacePtr<ISkeletalMeshEditionInterface> InInterface)
 	: ISkeletalMeshNotifier()
 	, Interface(InInterface)
diff --git a/Engine/Plugins/Experimental/MeshModelingToolsetExp/Source/MeshModelingToolsExp/Private/SkeletalMesh/SkeletonEditingTool.cpp b/Engine/Plugins/Experimental/MeshModelingToolsetExp/Source/MeshModelingToolsExp/Private/SkeletalMesh/SkeletonEditingTool.cpp
--- a/Engine/Plugins/Experimental/MeshModelingToolsetExp/Source/MeshModelingToolsExp/Private/SkeletalMesh/SkeletonEditingTool.cpp
+++ b/Engine/Plugins/Experimental/MeshModelingToolsetExp/Source/MeshModelingToolsExp/Private/SkeletalMesh/SkeletonEditingTool.cpp
@@ -17,6 +17,8 @@
 #include "BaseGizmos/GizmoViewContext.h"
 
 #include "Algo/Count.h"
+#include "BaseGizmos/TransformGizmoUtil.h"
+#include "SkeletalMesh/SkeletonTransformProxy.h"
 #include "TargetInterfaces/PrimitiveComponentBackedTarget.h"
 
 #define LOCTEXT_NAMESPACE "USkeletonEditingTool"
@@ -42,12 +44,14 @@ void FRefSkeletonChange::Apply(UObject* Object)
 { // redo
 	USkeletonEditingTool* Tool = CastChecked<USkeletonEditingTool>(Object);
 	Tool->Modifier->ExternalUpdate(PostChangeSkeleton, PostBoneTracker);
+	Tool->UpdateGizmo();
 }
 
 void FRefSkeletonChange::Revert(UObject* Object)
 { // undo
 	USkeletonEditingTool* Tool = CastChecked<USkeletonEditingTool>(Object);
 	Tool->Modifier->ExternalUpdate(PreChangeSkeleton, PreBoneTracker);
+	Tool->UpdateGizmo();
 }
 
 }
@@ -85,7 +89,10 @@ UInteractiveTool* USkeletonEditingToolBuilder::BuildTool(const FToolBuilderState
 void USkeletonEditingTool::Init(const FToolBuilderState& InSceneState)
 {
 	TargetWorld = InSceneState.World;
-	ViewContext = InSceneState.ToolManager->GetContextObjectStore()->FindContext<UGizmoViewContext>();
+
+	const UContextObjectStore* ContextObjectStore = InSceneState.ToolManager->GetContextObjectStore();
+	ViewContext = ContextObjectStore->FindContext<UGizmoViewContext>();
+	GizmoContext = ContextObjectStore->FindContext<USkeletalMeshGizmoContextObjectBase>();
 }
 
 void USkeletonEditingTool::Setup()
@@ -93,7 +100,7 @@ void USkeletonEditingTool::Setup()
 	Super::Setup();
 
 	const IPrimitiveComponentBackedTarget* TargetComponent = Cast<IPrimitiveComponentBackedTarget>(Target);
-	const USkeletalMeshComponent* Component = Cast<USkeletalMeshComponent>(TargetComponent->GetOwnerComponent());
+	USkeletalMeshComponent* Component = Cast<USkeletalMeshComponent>(TargetComponent->GetOwnerComponent());
 
 	USkeletalMesh* SkeletalMesh = Component ? Component->GetSkeletalMeshAsset() : nullptr;
 	if (!SkeletalMesh)
@@ -112,7 +119,7 @@ void USkeletonEditingTool::Setup()
 
 	if (NumBones)
 	{
-		CurrentBone = RootBoneName;
+		Selection = {RootBoneName};
 	}
 
 	// setup preview
@@ -136,6 +143,7 @@ void USkeletonEditingTool::Setup()
 	{
 		Properties = NewObject<USkeletonEditingProperties>();
 		Properties->Initialize(this);
+		Properties->Name = GetCurrentBone();
 		Properties->RestoreProperties(this);
 		AddToolPropertySource(Properties);
 
@@ -159,10 +167,123 @@ void USkeletonEditingTool::Setup()
 	{
 		UClickDragInputBehavior* ClickDragBehavior = NewObject<UClickDragInputBehavior>(this);
 		ClickDragBehavior->Initialize(this);
+		ClickDragBehavior->Modifiers.RegisterModifier(AddToSelectionModifier, FInputDeviceState::IsShiftKeyDown);
+		ClickDragBehavior->Modifiers.RegisterModifier(ToggleSelectionModifier, FInputDeviceState::IsCtrlKeyDown);
 		AddInputBehavior(ClickDragBehavior);
 	}
+
+	// setup gizmo
+	if (GizmoContext.IsValid())
+	{
+		UGizmoLambdaStateTarget* NewTarget = NewObject<UGizmoLambdaStateTarget>(this);
+		NewTarget->BeginUpdateFunction = [this]()
+		{
+			TGuardValue OperationGuard(Operation, EEditingOperation::Transform);
+			BeginChange();
+		};
+		NewTarget->EndUpdateFunction = [&]()
+		{
+			TGuardValue OperationGuard(Operation, EEditingOperation::Transform);
+			EndChange();
+
+			const IToolsContextQueriesAPI* ToolsContextQueries = GetToolManager()->GetPairedGizmoManager()->GetContextQueriesAPI();
+			check(ToolsContextQueries);
+			if (ToolsContextQueries->GetCurrentCoordinateSystem() == EToolContextCoordinateSystem::World)
+			{
+				UpdateGizmo();
+			}
+		};
+		
+		GizmoWrapper = GizmoContext->GetNewWrapper(GetToolManager(), this, NewTarget);
+		if (GizmoWrapper)
+		{
+			GizmoWrapper->Component = Component;
+		}
+	}
+
+	// setup watchers
+	SelectionWatcher.Initialize(
+	[this]()
+	{
+		return this->Selection;
+	},
+	[this](TArray<FName> InBoneNames)
+	{
+		UpdateGizmo();
+		if (NeedsNotification())
+		{
+			GetNotifier().Notify(InBoneNames, ESkeletalMeshNotifyType::BonesSelected);
+		}
+	},
+	this->Selection);
+
+	IToolsContextQueriesAPI* ToolsContextQueries = GetToolManager()->GetPairedGizmoManager()->GetContextQueriesAPI();
+	CoordinateSystemWatcher.Initialize(
+	[ToolsContextQueries]()
+	{
+		check(ToolsContextQueries);
+		return ToolsContextQueries->GetCurrentCoordinateSystem();
+	},
+	[this](EToolContextCoordinateSystem)
+	{
+		UpdateGizmo();
+	},
+	ToolsContextQueries->GetCurrentCoordinateSystem());
+	
+}
+
+void USkeletonEditingTool::UpdateGizmo() const
+{	
+	if (!GizmoWrapper)
+	{
+		return;
+	}
+
+	GizmoWrapper->Initialize();
+
+	const bool bUseGizmo = Operation == EEditingOperation::Select || Operation == EEditingOperation::Transform; 
+	if (Selection.IsEmpty() || !bUseGizmo)
+	{
+		return;
+	}
+
+	const IToolsContextQueriesAPI* ToolsContextQueries = GetToolManager()->GetPairedGizmoManager()->GetContextQueriesAPI();
+	const bool bWorld = ToolsContextQueries->GetCurrentCoordinateSystem() == EToolContextCoordinateSystem::World;
+	
+	const TArray<int32> BoneIndexes = GetSelectedBoneIndexes();
+	FTransform InitTransform = Modifier->GetTransform(BoneIndexes.Last(), true);
+	if (bWorld)
+	{
+		InitTransform = FTransform(FQuat::Identity, InitTransform.GetTranslation(), InitTransform.GetScale3D());
+	}
+	GizmoWrapper->Initialize(InitTransform, ToolsContextQueries->GetCurrentCoordinateSystem());
+
+	const FReferenceSkeleton& ReferenceSkeleton = Modifier->GetReferenceSkeleton();
+	const TArray<FMeshBoneInfo>& BoneInfos = ReferenceSkeleton.GetRawRefBoneInfo();
+	for (const int32 BoneIndex: BoneIndexes)
+	{
+		const FName BoneName = BoneInfos[BoneIndex].Name;
+		const int32 ParentBoneIndex = BoneInfos[BoneIndex].ParentIndex;
+		const FTransform& ParentGlobal = Modifier->GetTransform(ParentBoneIndex, true);
+		
+		GizmoWrapper->HandleBoneTransform(
+			[this, BoneIndex, bWorld]()
+			{
+				return Modifier->GetTransform(BoneIndex, bWorld);
+			},
+			[this, BoneName, ParentGlobal, bWorld](const FTransform& InNewTransform)
+			{
+				if (bWorld)
+				{
+					const FTransform NewLocal = InNewTransform.GetRelativeTransform(ParentGlobal);
+					return Modifier->SetBoneTransform(BoneName, NewLocal, true);
+				}
+				return Modifier->SetBoneTransform(BoneName, InNewTransform, true);
+			});
+	}
 }
 
+
 void USkeletonEditingTool::Shutdown(EToolShutdownType ShutdownType)
 {
 	if (ShutdownType == EToolShutdownType::Accept)
@@ -191,11 +312,17 @@ void USkeletonEditingTool::Shutdown(EToolShutdownType ShutdownType)
 	// show the skeletal mesh component
 	UE::ToolTarget::ShowSourceObject(Target);
 
-	//save properties	
+	// save properties	
 	Properties->SaveProperties(this);
 	ProjectionProperties->SaveProperties(this);
 	MirroringProperties->SaveProperties(this);
 	OrientingProperties->SaveProperties(this);
+
+	// clear gizmo
+	if (GizmoWrapper)
+	{
+		GizmoWrapper->Clear();
+	}
 }
 
 void USkeletonEditingTool::RegisterActions(FInteractiveToolActionSet& ActionSet)
@@ -216,6 +343,7 @@ void USkeletonEditingTool::RegisterActions(FInteractiveToolActionSet& ActionSet)
 		[this]()
 		{
 			Operation = EEditingOperation::Create;
+			UpdateGizmo();
 			GetToolManager()->DisplayMessage(LOCTEXT("Create", "Click & Drag to place a new bone."), EToolMessageLevel::UserNotification);
 		});
 	
@@ -258,6 +386,7 @@ void USkeletonEditingTool::RegisterActions(FInteractiveToolActionSet& ActionSet)
 		[this]()
 		{
 			Operation = EEditingOperation::Parent;
+			UpdateGizmo();
 			GetToolManager()->DisplayMessage(LOCTEXT("Parent", "Click on a bone to be set as the new parent."), EToolMessageLevel::UserNotification);
 		});
 }
@@ -272,7 +401,8 @@ void USkeletonEditingTool::CreateNewBone()
 	BeginChange();
 
 	const FName BoneName = Modifier->GetUniqueName(Properties->DefaultName);
-	const bool bBoneAdded = Modifier->AddBone(BoneName, CurrentBone, Properties->Transform);
+	const FName ParentName = Selection.IsEmpty() ? NAME_None : Selection.Last(); 
+	const bool bBoneAdded = Modifier->AddBone(BoneName, ParentName, Properties->Transform);
 	if (bBoneAdded)
 	{
 		if (NeedsNotification())
@@ -280,8 +410,8 @@ void USkeletonEditingTool::CreateNewBone()
 			GetNotifier().Notify({BoneName}, ESkeletalMeshNotifyType::BonesAdded);
 		}
 	
-		CurrentBone = BoneName;
-		Properties->Name = CurrentBone;
+		Selection = {BoneName};
+		Properties->Name = BoneName;
 
 		EndChange();
 		return;
@@ -315,10 +445,15 @@ void USkeletonEditingTool::RemoveBones()
 	const bool bBonesRemoved = Modifier->RemoveBones(BonesToRemove, true);
 	if (bBonesRemoved)
 	{
-		// if (NeedsNotification())
-		// {
-		// 	GetNotifier().Notify(BonesToRemove, ESkeletalMeshNotifyType::BonesRemoved);
-		// }
+		Selection.RemoveAll([&](const FName& BoneName)
+		{
+			return BonesToRemove.Contains(BoneName);
+		});
+		
+		if (NeedsNotification())
+		{
+			GetNotifier().Notify(BonesToRemove, ESkeletalMeshNotifyType::BonesRemoved);
+		}
 		
 		EndChange();
 		return;
@@ -401,6 +536,7 @@ void USkeletonEditingTool::MoveBones()
 
 void USkeletonEditingTool::RenameBones()
 {
+	const FName CurrentBone = GetCurrentBone();
 	if (CurrentBone == Properties->Name || Properties->Name == NAME_None)
 	{
 		return;
@@ -418,12 +554,12 @@ void USkeletonEditingTool::RenameBones()
 	const bool bBoneRenamed = Modifier->RenameBone(CurrentBone, Properties->Name);
 	if (bBoneRenamed)
 	{
-		CurrentBone = Properties->Name;
+		Selection.Last() = Properties->Name;
 
-		// if (NeedsNotification())
-		// {
-		// 	GetNotifier().Notify({CurrentBone}, ESkeletalMeshNotifyType::BonesRenamed);
-		// }
+		if (NeedsNotification())
+		{
+			GetNotifier().Notify({CurrentBone}, ESkeletalMeshNotifyType::BonesRenamed);
+		}
 		
 		EndChange();
 		return;
@@ -432,15 +568,25 @@ void USkeletonEditingTool::RenameBones()
 	CancelChange();
 }
 
-void USkeletonEditingTool::OnClickPress(const FInputDeviceRay& PressPos)
+void USkeletonEditingTool::OnClickPress(const FInputDeviceRay& InPressPos)
 {
+	if (PendingFunction)
+	{
+		PendingFunction();
+		PendingFunction.Reset();
+	}
 	BeginChange();
 }
 
-void USkeletonEditingTool::OnClickDrag(const FInputDeviceRay& DragPos)
+void USkeletonEditingTool::OnClickDrag(const FInputDeviceRay& InDragPos)
 {
+	if (Operation != EEditingOperation::Create)
+	{
+		return;
+	}
+	
 	FVector HitPoint;
-	if (ProjectionProperties->GetProjectionPoint(DragPos, HitPoint))
+	if (ProjectionProperties->GetProjectionPoint(InDragPos, HitPoint))
 	{
 		const FTransform& ParentGlobal = Modifier->GetTransform(ParentIndex, true);
 		Properties->Transform.SetLocation(ParentGlobal.InverseTransformPosition(HitPoint));
@@ -451,7 +597,7 @@ void USkeletonEditingTool::OnClickDrag(const FInputDeviceRay& DragPos)
 			BeginChange();
 		}
 
-		const bool bBoneMoved = Modifier->SetBoneTransform(CurrentBone, Properties->Transform, Properties->bUpdateChildren);
+		const bool bBoneMoved = Modifier->SetBoneTransform(GetCurrentBone(), Properties->Transform, Properties->bUpdateChildren);
 		if (!bBoneMoved)
 		{
 			CancelChange();
@@ -501,7 +647,121 @@ void USkeletonEditingTool::OrientBones()
 	CancelChange();
 }
 
-void USkeletonEditingTool::OnClickRelease(const FInputDeviceRay& ReleasePos)
+void USkeletonEditingTool::OnUpdateModifierState(int InModifierID, bool bIsOn)
+{
+	auto UpdateFlag = [this, bIsOn](EBoneSelectionMode InFlag)
+	{
+		if (bIsOn)
+		{
+			EnumAddFlags(SelectionMode, InFlag);
+		}
+		else
+		{
+			EnumRemoveFlags(SelectionMode, InFlag);
+		}		
+	};
+	
+	if (InModifierID == AddToSelectionModifier)
+	{
+		UpdateFlag(EBoneSelectionMode::Additive);
+	}
+	else if (InModifierID == ToggleSelectionModifier)
+	{
+		UpdateFlag(EBoneSelectionMode::Toggle);
+	}
+}
+
+void USkeletonEditingTool::SelectBone(const FName& InBoneName)
+{
+	TArray<FName> NewSelection = Selection;
+
+	if (EnumHasAnyFlags(SelectionMode, EBoneSelectionMode::Additive))
+	{
+		if (InBoneName != NAME_None)
+		{
+			NewSelection.AddUnique(InBoneName);
+		}		
+	}
+	else if (EnumHasAnyFlags(SelectionMode, EBoneSelectionMode::Toggle))
+	{
+		const bool bSelected = NewSelection.Contains(InBoneName);
+		if (bSelected)
+		{
+			NewSelection.Remove(InBoneName);
+		}
+		else
+		{
+			NewSelection.Add(InBoneName);
+		}
+	}
+	else
+	{
+		NewSelection.Empty();
+		if (InBoneName != NAME_None)
+		{
+			NewSelection.Add(InBoneName);
+		}
+	}
+
+	Selection = MoveTemp(NewSelection);
+	NormalizeSelection();
+}
+
+void USkeletonEditingTool::NormalizeSelection()
+{
+	if (Selection.Num() < 2)
+	{
+		return;
+	}
+	
+	const FReferenceSkeleton& ReferenceSkeleton = Modifier->GetReferenceSkeleton();
+
+	// ensure they are known and unique
+	TArray<int32> Indexes;
+	for (const FName& BoneName: Selection)
+	{
+		const int32 BoneIndex = ReferenceSkeleton.FindRawBoneIndex(BoneName);
+		if (BoneIndex != INDEX_NONE)
+		{
+			Indexes.AddUnique(BoneIndex);
+		}
+	}
+
+	if (Indexes.IsEmpty())
+	{
+		Selection.Reset();
+		return;
+	}
+
+	// sort them so that parents are placed after one of their children 
+	const TArray<FMeshBoneInfo>& BoneInfos = ReferenceSkeleton.GetRawRefBoneInfo();
+	Indexes.StableSort([&](const int32 Index0, const int32 Index1)
+	{
+		int32 P0 = BoneInfos[Index0].ParentIndex;
+		while (P0 != INDEX_NONE)
+		{
+			if (P0 == Index1)
+			{ // Index1 is a parent
+				return true;
+			}
+			if (Indexes.Contains(P0))
+			{ // parent is selected
+				return true;
+			}
+			P0 = BoneInfos[P0].ParentIndex;
+		}
+		return false;
+	});
+
+	// transform back to names
+	Selection.Reset();
+	Algo::Transform(Indexes, Selection, [&](const int32 Index)
+	{
+		return BoneInfos[Index].Name;
+	});	
+}
+
+void USkeletonEditingTool::OnClickRelease(const FInputDeviceRay& InReleasePos)
 {
 	TGuardValue OperationGuard(Operation, EEditingOperation::Transform);
 	EndChange();
@@ -510,17 +770,37 @@ void USkeletonEditingTool::OnClickRelease(const FInputDeviceRay& ReleasePos)
 void USkeletonEditingTool::OnTerminateDragSequence()
 {}
 
-FInputRayHit USkeletonEditingTool::CanBeginClickDragSequence(const FInputDeviceRay& ClickPos)
+void USkeletonEditingTool::OnTick(float DeltaTime)
 {
+	if (PendingFunction)
+	{
+		PendingFunction();
+		PendingFunction.Reset();
+	}
+	
+	SelectionWatcher.CheckAndUpdate();
+	CoordinateSystemWatcher.CheckAndUpdate();
+}
+
+FInputRayHit USkeletonEditingTool::CanBeginClickDragSequence(const FInputDeviceRay& InPressPos)
+{
+	PendingFunction.Reset();
+	ParentIndex = INDEX_NONE;
+	
 	FViewport* Viewport = GetToolManager()->GetContextQueriesAPI()->GetFocusedViewport();
-	if (!Viewport || !Viewport->HasFocus())
+	if (!Viewport)
+	{
+		return FInputRayHit();
+	}
+	
+	if (GizmoWrapper && GizmoWrapper->IsGizmoHit(InPressPos))
 	{
 		return FInputRayHit();
 	}
 	
 	auto PickBone = [&]() -> int32
 	{
-		if (HHitProxy* HitProxy = Viewport->GetHitProxy(ClickPos.ScreenPosition.X, ClickPos.ScreenPosition.Y))
+		if (HHitProxy* HitProxy = Viewport->GetHitProxy(InPressPos.ScreenPosition.X, InPressPos.ScreenPosition.Y))
 		{
 			if (TOptional<FName> OptBoneName = GetBoneName(HitProxy))
 			{
@@ -534,35 +814,31 @@ FInputRayHit USkeletonEditingTool::CanBeginClickDragSequence(const FInputDeviceR
 	// pick bone in viewport
 	const int32 BoneIndex = PickBone();
 
-	// update parent
-	const FReferenceSkeleton& ReferenceSkeleton = Modifier->GetReferenceSkeleton();
-	ParentIndex = INDEX_NONE;
-
 	// update projection properties
 	const FVector GlobalPosition = Modifier->GetTransform(BoneIndex, true).GetTranslation();
 	ProjectionProperties->UpdatePlane(*ViewContext, GlobalPosition);
-
+	
 	// if we picked a new bone
 	if (BoneIndex > INDEX_NONE)
 	{
 		// parent selection without changing the selection
 		if (Operation == EEditingOperation::Parent)
-		{	
+		{
+			const FReferenceSkeleton& ReferenceSkeleton = Modifier->GetReferenceSkeleton();
 			ParentBones(ReferenceSkeleton.GetBoneName(BoneIndex));
 			return FInputRayHit();
 		}
 		
 		// otherwise, update current selection
-		CurrentBone = ReferenceSkeleton.GetBoneName(BoneIndex);
-
-		Properties->Name = CurrentBone;
-		Properties->Transform = ReferenceSkeleton.GetRefBonePose()[BoneIndex];
-		ParentIndex = ReferenceSkeleton.GetParentIndex(BoneIndex);
-		
-		if (NeedsNotification())
+		PendingFunction = [this, BoneIndex]
 		{
-			GetNotifier().Notify({CurrentBone}, ESkeletalMeshNotifyType::BonesSelected);
-		}
+			const FReferenceSkeleton& ReferenceSkeleton = Modifier->GetReferenceSkeleton();
+			SelectBone(ReferenceSkeleton.GetBoneName(BoneIndex));
+
+			Properties->Name = GetCurrentBone();
+			Properties->Transform = ReferenceSkeleton.GetRefBonePose()[BoneIndex];
+			ParentIndex = ReferenceSkeleton.GetParentIndex(BoneIndex);
+		};
 	
 		return FInputRayHit(0.0);
 	}
@@ -570,15 +846,8 @@ FInputRayHit USkeletonEditingTool::CanBeginClickDragSequence(const FInputDeviceR
 	// if we didn't pick anything
 	if (Operation == EEditingOperation::Select)
 	{
-		// unselect all
-		CurrentBone = NAME_None;
-		Properties->Name = CurrentBone; 
-		
-		if (NeedsNotification())
-		{
-			GetNotifier().Notify({CurrentBone}, ESkeletalMeshNotifyType::BonesSelected);
-		}
-		
+		Selection.Empty();
+		Properties->Name = GetCurrentBone();
 		return FInputRayHit();
 	}
 
@@ -586,15 +855,19 @@ FInputRayHit USkeletonEditingTool::CanBeginClickDragSequence(const FInputDeviceR
 	if (Operation == EEditingOperation::Create)
 	{
 		FVector HitPoint;
-		if (ProjectionProperties->GetProjectionPoint(ClickPos, HitPoint))
+		if (ProjectionProperties->GetProjectionPoint(InPressPos, HitPoint))
 		{
 			// CurrentBone is gonna be the parent
-			ParentIndex = ReferenceSkeleton.FindRawBoneIndex(CurrentBone);
-			const FTransform& ParentGlobalTransform = Modifier->GetTransform(ParentIndex, true);
+			PendingFunction = [&]
+			{
+				const FReferenceSkeleton& ReferenceSkeleton = Modifier->GetReferenceSkeleton();
+				ParentIndex = ReferenceSkeleton.FindRawBoneIndex(GetCurrentBone());
+				const FTransform& ParentGlobalTransform = Modifier->GetTransform(ParentIndex, true);
 
-			// Create the new bone under mouse
-			Properties->Transform.SetLocation(ParentGlobalTransform.InverseTransformPosition(HitPoint));
-			CreateNewBone();
+				// Create the new bone under mouse
+				Properties->Transform.SetLocation(ParentGlobalTransform.InverseTransformPosition(HitPoint));
+				CreateNewBone();
+			};
 			
 			return FInputRayHit(0.0);
 		}
@@ -603,33 +876,47 @@ FInputRayHit USkeletonEditingTool::CanBeginClickDragSequence(const FInputDeviceR
 	return FInputRayHit();
 }
 
-void USkeletonEditingTool::HandleSkeletalMeshModified(const TArray<FName>& BoneNames, const ESkeletalMeshNotifyType InNotifyType)
+void USkeletonEditingTool::HandleSkeletalMeshModified(const TArray<FName>& InBoneNames, const ESkeletalMeshNotifyType InNotifyType)
 {
-	const FName BoneName = BoneNames.IsEmpty() ? NAME_None : BoneNames[0];
+	if (GetNotifier().Notifying())
+	{
+		return;
+	}
+
+	TArray<FName> BoneNames(InBoneNames);
+	const FReferenceSkeleton& RefSkeleton = Modifier->GetReferenceSkeleton();
+	BoneNames.RemoveAll([&](const FName& BoneName)
+	{
+		return RefSkeleton.FindRawBoneIndex(BoneName) == INDEX_NONE;
+	});
+
 	switch (InNotifyType)
 	{
 		case ESkeletalMeshNotifyType::BonesAdded:
-			CurrentBone = BoneName;
+			Selection = BoneNames;
 			break;
 		case ESkeletalMeshNotifyType::BonesRemoved:
-			if (BoneNames.Contains(CurrentBone))
+			Selection.RemoveAll([&](const FName& BoneName)
 			{
-				CurrentBone = NAME_None;
-			}
+				return BoneNames.Contains(BoneName);
+			});
 			break;
 		case ESkeletalMeshNotifyType::BonesMoved:
-			CurrentBone = BoneName;
 			break;
 		case ESkeletalMeshNotifyType::BonesSelected:
-			CurrentBone = BoneName;
+			Selection.Reset();
+			Selection = BoneNames;
 			break;
 		case ESkeletalMeshNotifyType::BonesRenamed:
-			CurrentBone = BoneName;
+			Selection = BoneNames;
 			break;
 		default:
 			break;
 	}
-	Properties->Name = CurrentBone;
+
+	NormalizeSelection();
+	
+	Properties->Name = GetCurrentBone();
 }
 
 void USkeletonEditingTool::Render(IToolsContextRenderAPI* RenderAPI)
@@ -678,24 +965,13 @@ void USkeletonEditingTool::Render(IToolsContextRenderAPI* RenderAPI)
 		HitProxies.Add(new HBoneHitProxy(Index, RefSkeleton.GetBoneName(Index)));
 	}
 
-	// FIXME cache this
-	TArray<int32> SelectedBones;
-	for (const FName& BoneName: GetSelectedBones())
-	{
-		int32 SelectedIndex = RefSkeleton.FindRawBoneIndex(CurrentBone);
-		if (SelectedIndex > INDEX_NONE)
-		{
-			SelectedBones.Add(SelectedIndex);
-		}
-	}
-
 	SkeletalDebugRendering::DrawBones(
 		PDI,
 		ComponentTransform.GetLocation(),
 		RequiredBones,
 		RefSkeleton,
 		WorldTransforms,
-		SelectedBones,
+		GetSelectedBoneIndexes(),
 		BoneColors,
 		HitProxies,
 		DrawConfig
@@ -704,17 +980,15 @@ void USkeletonEditingTool::Render(IToolsContextRenderAPI* RenderAPI)
 
 FBox USkeletonEditingTool::GetWorldSpaceFocusBox()
 {
-	const TArray<FName> Selection = GetSelectedBones();
-	if (!Selection.IsEmpty())
+	const TArray<int32> BoneIndexes = GetSelectedBoneIndexes();
+	if (!BoneIndexes.IsEmpty())
 	{
 		FBox Box(EForceInit::ForceInit);
 		TSet<int32> AllChildren;
 
 		const FReferenceSkeleton& RefSkeleton = Modifier->GetReferenceSkeleton();
-		
-		for (const FName& BoneName: Selection)
+		for (const int32 BoneIndex: BoneIndexes)
 		{
-			const int32 BoneIndex = RefSkeleton.FindRawBoneIndex(BoneName);
 			Box += Modifier->GetTransform(BoneIndex, true).GetTranslation();
 
 			// get direct children
@@ -741,22 +1015,25 @@ FBox USkeletonEditingTool::GetWorldSpaceFocusBox()
 
 TArray<FName> USkeletonEditingTool::GetSelectedBones() const
 {
-	if (Binding.IsValid())
-	{
-		const TArray<FName> Selection = Binding.Pin()->GetSelectedBones();
-		if (!Selection.IsEmpty())
-		{
-			return Selection;
-		}
-	}
+	return Selection;
+}
 
-	if (CurrentBone != NAME_None)
-	{
-		return {CurrentBone};
-	}
+FName USkeletonEditingTool::GetCurrentBone() const
+{
+	return Selection.IsEmpty() ? NAME_None : Selection.Last(); 
+}
 
-	const TArray<FName> Dummy;
-	return Dummy;
+TArray<int32> USkeletonEditingTool::GetSelectedBoneIndexes() const
+{
+	TArray<int32> Indexes;
+	
+	const FReferenceSkeleton& ReferenceSkeleton = Modifier->GetReferenceSkeleton();
+	Algo::Transform(Selection, Indexes, [&](const FName& BoneName)
+	{
+		return ReferenceSkeleton.FindRawBoneIndex(BoneName);
+	});
+	
+	return Indexes;
 }
 
 void USkeletonEditingTool::BeginChange()
@@ -802,7 +1079,6 @@ void USkeletonEditingTool::CancelChange()
 void USkeletonEditingProperties::Initialize(USkeletonEditingTool* ParentToolIn)
 {
 	ParentTool = ParentToolIn;
-	Name = ParentTool->CurrentBone;
 }
 
 #if WITH_EDITOR
diff --git a/Engine/Plugins/Experimental/MeshModelingToolsetExp/Source/MeshModelingToolsExp/Private/SkeletalMesh/SkeletonTransformProxy.cpp b/Engine/Plugins/Experimental/MeshModelingToolsetExp/Source/MeshModelingToolsExp/Private/SkeletalMesh/SkeletonTransformProxy.cpp
new file mode 100644
--- /dev/null
+++ b/Engine/Plugins/Experimental/MeshModelingToolsetExp/Source/MeshModelingToolsExp/Private/SkeletalMesh/SkeletonTransformProxy.cpp
@@ -0,0 +1,49 @@
+﻿// Copyright Epic Games, Inc. All Rights Reserved.
+
+
+#include "SkeletalMesh/SkeletonTransformProxy.h"
+
+FTransform USkeletonTransformProxy::GetTransform() const
+{
+	return SharedTransform;
+}
+
+void USkeletonTransformProxy::Initialize(const FTransform& InTransform, const EToolContextCoordinateSystem& InTransformMode)
+{
+	Objects.Reset();
+	SharedTransform = InitialSharedTransform = InTransform;
+	TransformMode = InTransformMode;
+}
+
+void USkeletonTransformProxy::UpdateSharedTransform()
+{
+	// nothing to do
+}
+
+void USkeletonTransformProxy::UpdateObjects()
+{
+	const FTransform OffsetTransform = SharedTransform.GetRelativeTransform(InitialSharedTransform);
+	const FVector DeltaTranslation = OffsetTransform.GetTranslation();
+	const FQuat DeltaRotation = OffsetTransform.GetRotation();
+	const bool bApplyLocal = TransformMode == EToolContextCoordinateSystem::Local;
+	
+	for (FRelativeObject& Object : Objects)
+	{
+		if (bApplyLocal)
+		{
+			Object.SetTransformFunc(OffsetTransform * Object.StartTransform);
+		}
+		else
+		{
+			FTransform Transform = Object.StartTransform;
+			Transform.SetRotation(DeltaRotation * Transform.GetRotation());
+			Transform.SetTranslation(DeltaTranslation + Transform.GetTranslation());
+			Object.SetTransformFunc(Transform);
+		}
+	}
+}
+
+bool USkeletonTransformProxy::IsValid() const
+{
+	return !Objects.IsEmpty();
+}
\ No newline at end of file
diff --git a/Engine/Plugins/Experimental/MeshModelingToolsetExp/Source/MeshModelingToolsExp/Private/TransformMeshesTool.cpp b/Engine/Plugins/Experimental/MeshModelingToolsetExp/Source/MeshModelingToolsExp/Private/TransformMeshesTool.cpp
--- a/Engine/Plugins/Experimental/MeshModelingToolsetExp/Source/MeshModelingToolsExp/Private/TransformMeshesTool.cpp
+++ b/Engine/Plugins/Experimental/MeshModelingToolsetExp/Source/MeshModelingToolsExp/Private/TransformMeshesTool.cpp
@@ -259,7 +259,7 @@ static void AddInstancedComponentInstance(UInstancedStaticMeshComponent* ISMC, i
 			ISMC->GetInstanceTransform(Index, Tmp, true);
 			return Tmp;
 		},
-		[ISMC, Index](FTransform NewTransform) {
+		[ISMC, Index](const FTransform& NewTransform) {
 			ISMC->UpdateInstanceTransform(Index, NewTransform, true, true, true);
 		},
 		Index, bModifyOnTransform
diff --git a/Engine/Plugins/Experimental/MeshModelingToolsetExp/Source/MeshModelingToolsExp/Public/SkeletalMesh/SkeletalMeshEditionInterface.h b/Engine/Plugins/Experimental/MeshModelingToolsetExp/Source/MeshModelingToolsExp/Public/SkeletalMesh/SkeletalMeshEditionInterface.h
--- a/Engine/Plugins/Experimental/MeshModelingToolsetExp/Source/MeshModelingToolsExp/Public/SkeletalMesh/SkeletalMeshEditionInterface.h
+++ b/Engine/Plugins/Experimental/MeshModelingToolsetExp/Source/MeshModelingToolsExp/Public/SkeletalMesh/SkeletalMeshEditionInterface.h
@@ -7,11 +7,16 @@
 #include "SkeletalMeshNotifier.h"
 
 #include "HitProxies.h"
+#include "ToolContextInterfaces.h"
 #include "GenericPlatform/ICursor.h"
 
 #include "SkeletalMeshEditionInterface.generated.h"
 
 class FSkeletalMeshToolNotifier;
+class UTransformProxy;
+class UInteractiveToolManager;
+class IGizmoStateTarget;
+struct FInputDeviceRay;
 
 /**
  * USkeletalMeshEditionInterface
@@ -35,6 +40,8 @@ public:
 	ISkeletalMeshNotifier& GetNotifier();
 	bool NeedsNotification() const;
 
+	virtual TArray<FName> GetSelectedBones() const;
+	
 	void BindTo(TSharedPtr<ISkeletalMeshEditorBinding> InBinding);
 	void Unbind();
 	
@@ -83,3 +90,45 @@ struct HBoneHitProxy : public HHitProxy
 
 	virtual EMouseCursor::Type GetMouseCursor() override { return EMouseCursor::Crosshairs; }
 };
+
+/**
+ * USkeletalMeshGizmoContextObjectBase
+ */
+
+UCLASS(Abstract)
+class MESHMODELINGTOOLSEXP_API USkeletalMeshGizmoContextObjectBase : public UObject
+{
+	GENERATED_BODY()
+	
+public:
+	virtual USkeletalMeshGizmoWrapperBase* GetNewWrapper(UInteractiveToolManager* InToolManager, UObject* Outer = nullptr, IGizmoStateTarget* InStateTarget = nullptr)
+	PURE_VIRTUAL(UGizmoContextObject::GetNewWrapper, return nullptr;);
+};
+
+/**
+ * USkeletalMeshGizmoWrapperBase
+ */
+
+UCLASS(Abstract)
+class MESHMODELINGTOOLSEXP_API USkeletalMeshGizmoWrapperBase : public UObject
+{
+	GENERATED_BODY()
+
+protected:
+	using FGetTransform = TUniqueFunction<FTransform(void)>;
+	using FSetTransform = TUniqueFunction<void(const FTransform&)>;
+	
+public:
+	virtual void Initialize(const FTransform& InTransform = FTransform::Identity, const EToolContextCoordinateSystem& InTransformMode = EToolContextCoordinateSystem::Local)
+	PURE_VIRTUAL(USkeletalMeshGizmoWrapperBase::Initialize, return;)
+	
+	virtual void HandleBoneTransform(FGetTransform GetTransformFunc, FSetTransform SetTransformFunc)
+	PURE_VIRTUAL(USkeletalMeshGizmoWrapperBase::HandleBoneTransform, return;)
+
+	virtual void Clear() PURE_VIRTUAL(USkeletalMeshGizmoWrapperBase::Clear, return;)
+	
+	virtual bool IsGizmoHit(const FInputDeviceRay& PressPos) const PURE_VIRTUAL(USkeletalMeshGizmoWrapperBase::IsGizmoHit, return false;)
+
+	UPROPERTY()
+	TWeakObjectPtr<USceneComponent> Component;
+};
\ No newline at end of file
diff --git a/Engine/Plugins/Experimental/MeshModelingToolsetExp/Source/MeshModelingToolsExp/Public/SkeletalMesh/SkeletonEditingTool.h b/Engine/Plugins/Experimental/MeshModelingToolsetExp/Source/MeshModelingToolsExp/Public/SkeletalMesh/SkeletonEditingTool.h
--- a/Engine/Plugins/Experimental/MeshModelingToolsetExp/Source/MeshModelingToolsExp/Public/SkeletalMesh/SkeletonEditingTool.h
+++ b/Engine/Plugins/Experimental/MeshModelingToolsetExp/Source/MeshModelingToolsExp/Public/SkeletalMesh/SkeletonEditingTool.h
@@ -9,6 +9,8 @@
 #include "SkeletonModifier.h"
 #include "SkeletalMesh/SkeletalMeshEditionInterface.h"
 #include "Engine/World.h"
+#include "Changes/ValueWatcher.h"
+#include "Misc/EnumClassFlags.h"
 
 #include "SkeletonEditingTool.generated.h"
 
@@ -16,6 +18,9 @@ class USingleClickInputBehavior;
 class UGizmoViewContext;
 class UTransformGizmo;
 class USkeletonEditingTool;
+class USkeletonTransformProxy;
+class USkeletalMeshGizmoContextObjectBase;
+class USkeletalMeshGizmoWrapperBase;
 
 namespace SkeletonEditingTool
 {
@@ -110,12 +115,14 @@ public:
 	virtual bool CanAccept() const override { return true; }
 
 	// ICLickDragBehaviorTarget overrides
-	virtual FInputRayHit CanBeginClickDragSequence(const FInputDeviceRay& PressPos) override;
+	virtual FInputRayHit CanBeginClickDragSequence(const FInputDeviceRay& InPressPos) override;
 	
-	virtual void OnClickPress(const FInputDeviceRay& PressPos) override;
-	virtual void OnClickDrag(const FInputDeviceRay& DragPos) override;
-	virtual void OnClickRelease(const FInputDeviceRay& ReleasePos) override;
+	virtual void OnClickPress(const FInputDeviceRay& InPressPos) override;
+	virtual void OnClickDrag(const FInputDeviceRay& InDragPos) override;
+	virtual void OnClickRelease(const FInputDeviceRay& InReleasePos) override;
 	virtual void OnTerminateDragSequence() override;
+	
+	virtual void OnTick(float DeltaTime) override;
 
 	// IInteractiveToolCameraFocusAPI overrides
 	virtual FBox GetWorldSpaceFocusBox() override;
@@ -125,13 +132,13 @@ public:
 	void RenameBones();
 	void MoveBones();
 	void OrientBones();
-	
-	int32 ParentIndex = INDEX_NONE;
-	FName CurrentBone = NAME_None;
 
-	UPROPERTY()
-	TObjectPtr<USkeletonModifier> Modifier;
-	
+	// ISkeletalMeshEditionInterface overrides
+	virtual TArray<FName> GetSelectedBones() const override;
+
+	// IModifierToggleBehaviorTarget overrides
+	virtual void OnUpdateModifierState(int ModifierID, bool bIsOn) override;
+
 protected:
 
 	// UInteractiveTool overrides
@@ -146,8 +153,9 @@ protected:
 	void UnParentBones();
 	void ParentBones(const FName& InParentName);
 
-	TArray<FName> GetSelectedBones() const;
-
+	UPROPERTY()
+	TObjectPtr<USkeletonModifier> Modifier;
+	
 	UPROPERTY()
 	TObjectPtr<USkeletonEditingProperties> Properties;
 
@@ -175,13 +183,57 @@ protected:
 	UPROPERTY()
 	EEditingOperation Operation = EEditingOperation::Select;
 
+	// gizmo
+	UPROPERTY()
+	TWeakObjectPtr<USkeletalMeshGizmoContextObjectBase> GizmoContext = nullptr;
+
+	UPROPERTY()
+	TObjectPtr<USkeletalMeshGizmoWrapperBase> GizmoWrapper = nullptr;
+
+	void UpdateGizmo() const;
+	
 	// ref skeleton transactions
 	void BeginChange();
 	void EndChange();
 	void CancelChange();
 	TUniquePtr<SkeletonEditingTool::FRefSkeletonChange> ActiveChange;
+
+	friend class SkeletonEditingTool::FRefSkeletonChange;
+	
+private:
+	TArray<int32> GetSelectedBoneIndexes() const;
+	
+	enum class EBoneSelectionMode : uint8
+	{
+		Single				= 0,
+		Additive			= 1 << 0,
+		Toggle				= 1 << 1
+	};
+	FRIEND_ENUM_CLASS_FLAGS(USkeletonEditingTool::EBoneSelectionMode)
+	
+
+	// flags used to identify behavior modifier keys/buttons
+	static constexpr int AddToSelectionModifier = 1;
+	static constexpr int ToggleSelectionModifier = 2;
+	EBoneSelectionMode SelectionMode = EBoneSelectionMode::Single;
+
+	TArray<FName> Selection;
+	
+	void SelectBone(const FName& InBoneName);
+	FName GetCurrentBone() const;
+
+	void NormalizeSelection();
+
+	TValueWatcher<TArray<FName>> SelectionWatcher;
+	TValueWatcher<EToolContextCoordinateSystem> CoordinateSystemWatcher; 
+	
+	int32 ParentIndex = INDEX_NONE;
+
+	TFunction<void()> PendingFunction;
 };
 
+ENUM_CLASS_FLAGS(USkeletonEditingTool::EBoneSelectionMode);
+
 /**
  * USkeletonEditingProperties
  */
diff --git a/Engine/Plugins/Experimental/MeshModelingToolsetExp/Source/MeshModelingToolsExp/Public/SkeletalMesh/SkeletonTransformProxy.h b/Engine/Plugins/Experimental/MeshModelingToolsetExp/Source/MeshModelingToolsExp/Public/SkeletalMesh/SkeletonTransformProxy.h
new file mode 100644
--- /dev/null
+++ b/Engine/Plugins/Experimental/MeshModelingToolsetExp/Source/MeshModelingToolsExp/Public/SkeletalMesh/SkeletonTransformProxy.h
@@ -0,0 +1,37 @@
+﻿// Copyright Epic Games, Inc. All Rights Reserved.
+
+#pragma once
+
+#include "BaseGizmos/TransformProxy.h"
+
+#include "SkeletonTransformProxy.generated.h"
+
+class UObject;
+
+/**
+ * USkeletonTransformProxy is a derivation of UTransformProxy that manages several bones and update their transform individually.
+ */
+
+UCLASS(Transient)
+class MESHMODELINGTOOLSEXP_API USkeletonTransformProxy : public UTransformProxy
+{
+	GENERATED_BODY()
+
+public:
+
+	void Initialize(const FTransform& InTransform, const EToolContextCoordinateSystem& InTransformMode);
+	
+	bool IsValid() const;
+
+	// UTransformProxy overrides
+	virtual FTransform GetTransform() const override;
+
+	UPROPERTY()
+	EToolContextCoordinateSystem TransformMode = EToolContextCoordinateSystem::Local;
+	
+protected:
+	// UTransformProxy overrides
+	virtual void UpdateSharedTransform() override;
+	virtual void UpdateObjects() override;
+};
+
diff --git a/Engine/Source/Developer/SkeletalMeshUtilitiesCommon/Private/SkeletalMeshNotifier.cpp b/Engine/Source/Developer/SkeletalMeshUtilitiesCommon/Private/SkeletalMeshNotifier.cpp
--- a/Engine/Source/Developer/SkeletalMeshUtilitiesCommon/Private/SkeletalMeshNotifier.cpp
+++ b/Engine/Source/Developer/SkeletalMeshUtilitiesCommon/Private/SkeletalMeshNotifier.cpp
@@ -7,8 +7,17 @@ FSkeletalMeshNotifyDelegate& ISkeletalMeshNotifier::Delegate()
 	return NotifyDelegate;
 }
 
+bool ISkeletalMeshNotifier::Notifying() const
+{
+	return bNotifying;
+}
+
 void ISkeletalMeshNotifier::Notify(const TArray<FName>& BoneNames, const ESkeletalMeshNotifyType InNotifyType) const
 {
-	NotifyDelegate.Broadcast(BoneNames, InNotifyType);
+	if (!bNotifying)
+	{
+		TGuardValue<bool> RecursionGuard(bNotifying, true);
+		NotifyDelegate.Broadcast(BoneNames, InNotifyType);
+	}
 }
 
diff --git a/Engine/Source/Developer/SkeletalMeshUtilitiesCommon/Public/SkeletalMeshNotifier.h b/Engine/Source/Developer/SkeletalMeshUtilitiesCommon/Public/SkeletalMeshNotifier.h
--- a/Engine/Source/Developer/SkeletalMeshUtilitiesCommon/Public/SkeletalMeshNotifier.h
+++ b/Engine/Source/Developer/SkeletalMeshUtilitiesCommon/Public/SkeletalMeshNotifier.h
@@ -33,9 +33,11 @@ public:
 	virtual void HandleNotification(const TArray<FName>& BoneNames, const ESkeletalMeshNotifyType InNotifyType) = 0;
 
 	void Notify(const TArray<FName>& BoneNames, const ESkeletalMeshNotifyType InNotifyType) const;
-	
+	bool Notifying() const;
+
 private:
 	FSkeletalMeshNotifyDelegate NotifyDelegate;
+	mutable bool bNotifying = false;
 };
 
 class SKELETALMESHUTILITIESCOMMON_API ISkeletalMeshEditorBinding
diff --git a/Engine/Source/Editor/Experimental/EditorInteractiveToolsFramework/Private/EditorGizmos/TransformGizmo.cpp b/Engine/Source/Editor/Experimental/EditorInteractiveToolsFramework/Private/EditorGizmos/TransformGizmo.cpp
--- a/Engine/Source/Editor/Experimental/EditorInteractiveToolsFramework/Private/EditorGizmos/TransformGizmo.cpp
+++ b/Engine/Source/Editor/Experimental/EditorInteractiveToolsFramework/Private/EditorGizmos/TransformGizmo.cpp
@@ -31,6 +31,11 @@ void UTransformGizmo::SetDisallowNegativeScaling(bool bDisallow)
 
 void UTransformGizmo::Setup()
 {
+	if (IsValid(GizmoElementRoot))
+	{
+		return;
+	}
+	
 	UInteractiveGizmo::Setup();
 
 	SetupBehaviors();
@@ -110,14 +115,23 @@ void UTransformGizmo::Shutdown()
 
 FTransform UTransformGizmo::GetGizmoTransform() const
 {
-	float Scale = 1.0f;
-
-	if (TransformGizmoSource)
+	const float Scale = TransformGizmoSource ? TransformGizmoSource->GetGizmoScale() : 1.0f;
+	
+	auto GetCoordinateSystem = [&]()
+	{
+		if (TransformGizmoSource)
+		{
+			return TransformGizmoSource->GetGizmoCoordSystemSpace(); 
+		}
+		return GetGizmoManager()->GetContextQueriesAPI()->GetCurrentCoordinateSystem();
+	};
+	const bool bLocal = GetCoordinateSystem() == EToolContextCoordinateSystem::Local;
+	
+	FTransform GizmoLocalToWorldTransform(CurrentTransform.GetTranslation());
+	if (bLocal)
 	{
-		Scale = TransformGizmoSource->GetGizmoScale();
+		GizmoLocalToWorldTransform.SetRotation(CurrentTransform.GetRotation());
 	}
-
-	FTransform GizmoLocalToWorldTransform = CurrentTransform;
 	GizmoLocalToWorldTransform.SetScale3D(FVector(Scale, Scale, Scale));
 
 	return GizmoLocalToWorldTransform;
@@ -242,11 +256,46 @@ FInputRayHit UTransformGizmo::CanBeginClickDragSequence(const FInputDeviceRay& P
 
 void UTransformGizmo::UpdateMode()
 {
-	if (TransformGizmoSource && TransformGizmoSource->GetVisible())
+	auto GetTransformMode = [&]()
 	{
-		EGizmoTransformMode NewMode = TransformGizmoSource->GetGizmoMode();
-		EAxisList::Type NewAxisToDraw = TransformGizmoSource->GetGizmoAxisToDraw(NewMode);
+		if (TransformGizmoSource)
+		{
+			return TransformGizmoSource->GetGizmoMode();
+		}
+		
+		const EToolContextTransformGizmoMode ActiveGizmoMode = GetGizmoManager()->GetContextQueriesAPI()->GetCurrentTransformGizmoMode();
+		switch (ActiveGizmoMode)
+		{
+			case EToolContextTransformGizmoMode::Translation: return EGizmoTransformMode::Translate;
+			case EToolContextTransformGizmoMode::Rotation: return EGizmoTransformMode::Rotate;
+			case EToolContextTransformGizmoMode::Scale: return EGizmoTransformMode::Scale;
+		}
+		return EGizmoTransformMode::None;
+	};
 
+	auto GetAxisToDraw = [&]()
+	{
+		if (TransformGizmoSource)
+		{
+			return TransformGizmoSource->GetGizmoAxisToDraw(TransformGizmoSource->GetGizmoMode());
+		}
+		return EAxisList::Type::All;
+	};
+
+	auto GetVisible = [&]()
+	{
+		if (TransformGizmoSource)
+		{
+			return TransformGizmoSource->GetVisible();
+		}
+		return true;
+	};
+	
+	if (GetVisible())
+	{
+		const EGizmoTransformMode NewMode = GetTransformMode();
+		const EAxisList::Type NewAxisToDraw = GetAxisToDraw();
+		
 		if (NewMode != CurrentMode)
 		{
 			EnableMode(CurrentMode, EAxisList::None);
@@ -554,7 +603,7 @@ void UTransformGizmo::Tick(float DeltaTime)
 	UpdateCameraAxisSource();
 }
 
-void UTransformGizmo::SetActiveTarget(UTransformProxy* Target, IToolContextTransactionProvider* TransactionProvider)
+void UTransformGizmo::SetActiveTarget(UTransformProxy* Target, IToolContextTransactionProvider* TransactionProvider, IGizmoStateTarget* InStateTarget)
 {
 	if (ActiveTarget != nullptr)
 	{
@@ -577,9 +626,17 @@ void UTransformGizmo::SetActiveTarget(UTransformProxy* Target, IToolContextTrans
 	{
 		TransactionProvider = GetGizmoManager();
 	}
-	
-	StateTarget = UGizmoObjectModifyStateTarget::Construct(Target,
-		LOCTEXT("UTransformGizmoTransaction", "Transform"), TransactionProvider, this);
+
+	if (InStateTarget)
+	{
+		StateTarget = Cast<UObject>(InStateTarget);
+	}
+	else
+	{
+		StateTarget = UGizmoObjectModifyStateTarget::Construct(Target,
+			LOCTEXT("UTransformGizmoTransaction", "Transform"), TransactionProvider, this);	
+	}
+
 
 	CameraAxisSource = NewObject<UGizmoConstantFrameAxisSource>(this);
 }
@@ -819,7 +876,16 @@ FQuat UTransformGizmo::RotationSnapFunction(const FQuat& DeltaRotation) const
 
 FVector UTransformGizmo::GetWorldAxis(const FVector& InAxis)
 {
-	if (TransformGizmoSource->GetGizmoCoordSystemSpace() == EToolContextCoordinateSystem::Local)
+	auto GetCoordinateSystem = [&]()
+	{
+		if (TransformGizmoSource)
+		{
+			return TransformGizmoSource->GetGizmoCoordSystemSpace(); 
+		}
+		return GetGizmoManager()->GetContextQueriesAPI()->GetCurrentCoordinateSystem();
+	};
+	
+	if (GetCoordinateSystem() == EToolContextCoordinateSystem::Local)
 	{
 		return CurrentTransform.TransformVectorNoScale(InAxis);
 	}
@@ -1284,7 +1350,16 @@ void UTransformGizmo::OnClickDragScale(const FInputDeviceRay& DragPos)
 {
 	FVector2D ScreenDelta = DragPos.ScreenPosition - InteractionScreenCurrPos;
 
-	if (TransformGizmoSource->GetScaleType() != EGizmoTransformScaleType::PercentageBased)
+	auto GetScaleType = [&]()
+	{
+		if (TransformGizmoSource)
+		{
+			return TransformGizmoSource->GetScaleType();
+		}
+		return EGizmoTransformScaleType::Default;
+	};
+	
+	if (GetScaleType() != EGizmoTransformScaleType::PercentageBased)
 	{
 		ScreenDelta *= ScaleMultiplier;
 	}
@@ -1337,7 +1412,7 @@ FVector UTransformGizmo::ComputeScaleDelta(const FVector2D& InStartPos, const FV
 
 void UTransformGizmo::OnClickPressRotateXAxis(const FInputDeviceRay& PressPos)
 {
-	InteractionScreenAxisDirection = GetScreenRotateAxisDir(FVector::YAxisVector, FVector::ZAxisVector).GetSafeNormal();
+	InteractionScreenAxisDirection = GetScreenRotateAxisDir(FVector::ZAxisVector, FVector::YAxisVector).GetSafeNormal();
 	InteractionAxisList = EAxisList::X;
 	InteractionScreenStartPos = InteractionScreenCurrPos = PressPos.ScreenPosition;
 	bInInteraction = true;
@@ -1345,7 +1420,7 @@ void UTransformGizmo::OnClickPressRotateXAxis(const FInputDeviceRay& PressPos)
 
 void UTransformGizmo::OnClickPressRotateYAxis(const FInputDeviceRay& PressPos)
 {
-	InteractionScreenAxisDirection = GetScreenRotateAxisDir(FVector::ZAxisVector, FVector::XAxisVector).GetSafeNormal();
+	InteractionScreenAxisDirection = GetScreenRotateAxisDir(FVector::XAxisVector, FVector::ZAxisVector).GetSafeNormal();
 	InteractionAxisList = EAxisList::Y;
 	InteractionScreenStartPos = InteractionScreenCurrPos = PressPos.ScreenPosition;
 	bInInteraction = true;
@@ -1372,9 +1447,10 @@ FVector2D UTransformGizmo::GetScreenRotateAxisDir(const FVector& InAxis0, const
 	const bool bMirrorAxis1 = (FVector::DotProduct(Axis1, DirectionToWidget) <= 0.0f);
 	const float Direction = (bMirrorAxis0 ^ bMirrorAxis1) ? -1.0f : 1.0f;
 
-	const FVector AxisDir = (Axis1 - Axis0) * Direction;
-
-	return GetScreenProjectedAxis(GizmoViewContext, AxisDir);
+	const FVector2D Axis0Screen = GetScreenProjectedAxis(GizmoViewContext, bMirrorAxis0 ? Axis0:-Axis0);
+	const FVector2D Axis1Screen = GetScreenProjectedAxis(GizmoViewContext, bMirrorAxis1 ? Axis1:-Axis1);
+	
+	return ((Axis1Screen - Axis0Screen) * Direction).GetSafeNormal();
 }
 
 void UTransformGizmo::OnClickDragRotateAxis(const FInputDeviceRay& DragPos)
@@ -1401,8 +1477,16 @@ FQuat UTransformGizmo::ComputeAxisRotateDelta(const FVector2D& InStartPos, const
 		DeltaRot.Yaw = FVector2D::DotProduct(InteractionScreenAxisDirection, DragDir);
 	}
 
-	check(TransformGizmoSource);
-	if (TransformGizmoSource->GetGizmoCoordSystemSpace() == EToolContextCoordinateSystem::Local)
+	auto GetCoordinateSystem = [&]()
+	{
+		if (TransformGizmoSource)
+		{
+			return TransformGizmoSource->GetGizmoCoordSystemSpace(); 
+		}
+		return GetGizmoManager()->GetContextQueriesAPI()->GetCurrentCoordinateSystem();
+	};
+	
+	if (GetCoordinateSystem() == EToolContextCoordinateSystem::Local)
 	{
 		check(ActiveTarget);
 		FMatrix CurrCoordSystem = ActiveTarget->GetTransform().ToMatrixNoScale();
diff --git a/Engine/Source/Editor/Experimental/EditorInteractiveToolsFramework/Public/EditorGizmos/TransformGizmo.h b/Engine/Source/Editor/Experimental/EditorInteractiveToolsFramework/Public/EditorGizmos/TransformGizmo.h
--- a/Engine/Source/Editor/Experimental/EditorInteractiveToolsFramework/Public/EditorGizmos/TransformGizmo.h
+++ b/Engine/Source/Editor/Experimental/EditorInteractiveToolsFramework/Public/EditorGizmos/TransformGizmo.h
@@ -187,7 +187,7 @@ public:
 	 * @param Target active target
 	 * @param TransactionProvider optional IToolContextTransactionProvider implementation to use - by default uses GizmoManager
 	 */
-	virtual void SetActiveTarget(UTransformProxy* Target, IToolContextTransactionProvider* TransactionProvider = nullptr);
+	virtual void SetActiveTarget(UTransformProxy* Target, IToolContextTransactionProvider* TransactionProvider = nullptr, IGizmoStateTarget* InStateTarget = nullptr);
 
 	/**
 	 * Clear the active target object for the Gizmo
@@ -362,10 +362,10 @@ protected:
 	// internal function that updates CameraAxisSource by getting current view state from GizmoManager
 	void UpdateCameraAxisSource();
 
-	/** The state target is created internally during SetActiveTarget() */
+	/** The state target is created internally during SetActiveTarget() if no one is provided. */
 	UPROPERTY()
-	TObjectPtr<UGizmoObjectModifyStateTarget> StateTarget;
-
+	TScriptInterface<IGizmoStateTarget> StateTarget;
+	
 	/**
 	 * These are used to let the translation subgizmos use raycasts into the scene to align the gizmo with scene geometry.
 	 * See comment for SetWorldAlignmentFunctions().
diff --git a/Engine/Source/Editor/SkeletalMeshEditor/Private/SkeletalMeshEditor.cpp b/Engine/Source/Editor/SkeletalMeshEditor/Private/SkeletalMeshEditor.cpp
--- a/Engine/Source/Editor/SkeletalMeshEditor/Private/SkeletalMeshEditor.cpp
+++ b/Engine/Source/Editor/SkeletalMeshEditor/Private/SkeletalMeshEditor.cpp
@@ -1668,7 +1668,10 @@ void FSkeletalMeshEditorNotifier::HandleNotification(const TArray<FName>& BoneNa
 		SkeletonTree->DeselectAll();
 		if (!BoneNames.IsEmpty())
 		{
-			SkeletonTree->SetSelectedBone(BoneNames[0], ESelectInfo::Direct);
+			for (const FName BoneName: BoneNames)
+			{
+				SkeletonTree->SetSelectedBone(BoneName, ESelectInfo::Direct);
+			}
 		}
 		break;
 	case ESkeletalMeshNotifyType::BonesRenamed:
diff --git a/Engine/Source/Runtime/InteractiveToolsFramework/Private/BaseGizmos/TransformProxy.cpp b/Engine/Source/Runtime/InteractiveToolsFramework/Private/BaseGizmos/TransformProxy.cpp
--- a/Engine/Source/Runtime/InteractiveToolsFramework/Private/BaseGizmos/TransformProxy.cpp
+++ b/Engine/Source/Runtime/InteractiveToolsFramework/Private/BaseGizmos/TransformProxy.cpp
@@ -35,7 +35,7 @@ void UTransformProxy::AddComponent(USceneComponent* Component, bool bModifyCompo
 void UTransformProxy::AddComponentCustom(
 	USceneComponent* Component,
 	TUniqueFunction<FTransform(void)> GetTransformFunc,
-	TUniqueFunction<void(FTransform)> SetTransformFunc,
+	TUniqueFunction<void(const FTransform&)> SetTransformFunc,
 	int64 UserDefinedIndex,
 	bool bModifyComponentOnTransform)
 {
diff --git a/Engine/Source/Runtime/InteractiveToolsFramework/Public/BaseGizmos/TransformProxy.h b/Engine/Source/Runtime/InteractiveToolsFramework/Public/BaseGizmos/TransformProxy.h
--- a/Engine/Source/Runtime/InteractiveToolsFramework/Public/BaseGizmos/TransformProxy.h
+++ b/Engine/Source/Runtime/InteractiveToolsFramework/Public/BaseGizmos/TransformProxy.h
@@ -60,7 +60,7 @@ public:
 	virtual void AddComponentCustom(
 		USceneComponent* Component, 
 		TUniqueFunction<FTransform(void)> GetTransformFunc,
-		TUniqueFunction<void(FTransform)> SetTransformFunc,
+		TUniqueFunction<void(const FTransform&)> SetTransformFunc,
 		int64 UserDefinedIndex = 0,
 		bool bModifyComponentOnTransform = true);
 
@@ -162,7 +162,7 @@ protected:
 		int64 UserDefinedIndex;		// no current purpose, reserved for subclass/future use
 
 		TUniqueFunction<FTransform(void)> GetTransformFunc;
-		TUniqueFunction<void(FTransform)> SetTransformFunc;
+		TUniqueFunction<void(const FTransform&)> SetTransformFunc;
 
 		/** The initial transform of the object, set during UpdateSharedTransform() */
 		FTransform StartTransform;
