CombineMeshInstances: can now compute 'swept solid' approximation based on (approximate) projected silhouette of shapes, to get 2D profile curves, and then extrude along axis. This is a bit rough right now but doing polygon union + closure tends to clean up the messy bits. This method can capture parts with holes, and optionally filter out very small holes. Currently hole filtering is hardcoded-enabled, will be exposed as parameters. Triangle cost metric is somewhat broken for meshes with holes, needs to be improved.
GeometryProcessing: added FilterHoles and VtxNormalOffset functions to GeneralPolygon2

diff --git a/Engine/Plugins/Experimental/MeshModelingToolsetExp/Source/GeometryProcessingAdapters/Private/CombineMeshInstancesImpl.cpp b/Engine/Plugins/Experimental/MeshModelingToolsetExp/Source/GeometryProcessingAdapters/Private/CombineMeshInstancesImpl.cpp
--- a/Engine/Plugins/Experimental/MeshModelingToolsetExp/Source/GeometryProcessingAdapters/Private/CombineMeshInstancesImpl.cpp
+++ b/Engine/Plugins/Experimental/MeshModelingToolsetExp/Source/GeometryProcessingAdapters/Private/CombineMeshInstancesImpl.cpp
@@ -18,6 +18,7 @@
 #include "DynamicMeshEditor.h"
 #include "DynamicMesh/DynamicMeshAttributeSet.h"
 #include "DynamicMesh/MeshNormals.h"
+#include "DynamicMesh/MeshTransforms.h"
 #include "DynamicMesh/DynamicMeshAABBTree3.h"
 
 #include "ShapeApproximation/ShapeDetection3.h"
@@ -34,6 +35,14 @@
 #include "Operations/RemoveOccludedTriangles.h"
 #include "Operations/MeshResolveTJunctions.h"
 
+#include "MeshBoundaryLoops.h"
+#include "Curve/PlanarComplex.h"
+#include "Curve/PolygonIntersectionUtils.h"
+#include "Curve/PolygonOffsetUtils.h"
+#include "ConstrainedDelaunay2.h"
+#include "Generators/FlatTriangulationMeshGenerator.h"
+#include "Operations/ExtrudeMesh.h"
+
 #include "Physics/CollisionGeometryConversion.h"
 #include "Physics/PhysicsDataCollection.h"
 
@@ -753,6 +762,160 @@ static void ComputeSimplePartApproximation(
 }
 
 
+
+
+
+
+
+static void ComputeSweptSolidApproximation(
+	const FDynamicMesh3& SourcePartMesh,
+	FDynamicMesh3& DestMesh,
+	FVector3d Direction,
+	double MergeOffset = 0.1,
+	double SimplifyTolerance = 1.0,
+	double MinHoleArea = 1.0)
+{
+	FFrame3d ProjectFrame(FVector3d::Zero(), Direction);
+	FVector3d XAxis(ProjectFrame.GetAxis(0));
+	FVector3d YAxis(ProjectFrame.GetAxis(1));
+
+	FDynamicMesh3 FilteredMesh(SourcePartMesh);
+	FInterval1d AxisRange = FInterval1d::Empty();
+	for (FVector3d Position : FilteredMesh.VerticesItr())
+	{
+		AxisRange.Contain(Position.Dot(Direction));
+	}
+
+	TArray<int32> DeleteTris;
+	for (int32 tid : FilteredMesh.TriangleIndicesItr())
+	{
+		if (FilteredMesh.GetTriNormal(tid).Dot(Direction) < 0.1)
+		{
+			DeleteTris.Add(tid);
+		}
+	}
+	for (int32 tid : DeleteTris)
+	{
+		FilteredMesh.RemoveTriangle(tid);
+	}
+
+	FMeshBoundaryLoops Loops(&FilteredMesh);
+	FPlanarComplexd PlanarComplex;
+	for (FEdgeLoop& Loop : Loops.Loops)
+	{
+		TArray<FVector3d> Vertices;
+		Loop.GetVertices<FVector3d>(Vertices);
+		FPolygon2d Polygon;
+		for (FVector3d V : Vertices)
+		{
+			Polygon.AppendVertex(FVector2d(V.Dot(XAxis), V.Dot(YAxis)));
+		}
+		Polygon.Reverse();		// mesh orientation comes out backwards...
+		PlanarComplex.Polygons.Add(MoveTemp(Polygon));
+	}
+	PlanarComplex.bTrustOrientations = true;		// have to do this or overlapping projections will create holes
+	PlanarComplex.FindSolidRegions();
+	TArray<FGeneralPolygon2d> Polygons = PlanarComplex.ConvertOutputToGeneralPolygons();
+
+	if (Polygons.Num() == 0)
+	{
+		// failed to find anything??
+		ComputeSimplePartApproximation(SourcePartMesh, DestMesh, EApproximatePartMethod::OrientedBox);
+		return;
+	}
+
+	double UnionMergeOffset = 0.1;
+	if (Polygons.Num() > 1)
+	{
+		// nudge all polygons outwards to ensure that when we boolean union exactly-coincident polygons
+		// they intersect a bit, otherwise we may end up with zero-area cracks/holes
+		if (UnionMergeOffset > 0)
+		{
+			for (FGeneralPolygon2d& Polygon : Polygons)
+			{
+				Polygon.VtxNormalOffset(UnionMergeOffset);
+			}
+		}
+
+		TArray<FGeneralPolygon2d> ResultPolygons;
+		PolygonsUnion(Polygons, ResultPolygons, true);
+		Polygons = MoveTemp(ResultPolygons);
+
+		if (UnionMergeOffset > 0)
+		{
+			for (FGeneralPolygon2d& Polygon : Polygons)
+			{
+				Polygon.VtxNormalOffset(-UnionMergeOffset);	// undo offset
+			}
+		}
+	}
+
+	// can optionally try to reduce polygon complexity by topological closure (dilate/erode)
+	if (MergeOffset > 0)
+	{
+		TArray<FGeneralPolygon2d> TmpPolygons;
+		PolygonsOffsets(MergeOffset, -MergeOffset,
+			Polygons, TmpPolygons, true, 1.0,
+			EPolygonOffsetJoinType::Square,
+			EPolygonOffsetEndType::Polygon);
+
+		Polygons = MoveTemp(TmpPolygons);
+	}
+
+	FConstrainedDelaunay2d Triangulator;
+	for (FGeneralPolygon2d& Polygon : Polygons)
+	{
+		if (SimplifyTolerance > 0)
+		{
+			Polygon.Simplify(SimplifyTolerance, SimplifyTolerance * 0.25);		// 0.25 is kind of arbitrary here...
+		}
+		if (MinHoleArea > 0)
+		{
+			Polygon.FilterHoles([&](const FPolygon2d& HolePoly) { return HolePoly.Area() < MinHoleArea; });
+		}
+		Triangulator.Add(Polygon);
+	}
+
+	Triangulator.Triangulate([&Polygons](const TArray<FVector2d>& Vertices, FIndex3i Tri)
+	{
+		FVector2d Point = (Vertices[Tri.A] + Vertices[Tri.B] + Vertices[Tri.C]) / 3.0;
+		for (const FGeneralPolygon2d& Polygon : Polygons)
+		{
+			if (Polygon.Contains(Point))
+			{
+				return true;
+			}
+		}
+		return false;
+	});
+
+	FFlatTriangulationMeshGenerator TriangulationMeshGen;
+	TriangulationMeshGen.Vertices2D = Triangulator.Vertices;
+	TriangulationMeshGen.Triangles2D = Triangulator.Triangles;
+	FDynamicMesh3 ResultMesh(&TriangulationMeshGen.Generate());
+
+	if (ResultMesh.TriangleCount() < 3)
+	{
+		// failed to find anything??
+		ComputeSimplePartApproximation(SourcePartMesh, DestMesh, EApproximatePartMethod::OrientedBox);
+		return;
+	}
+
+	ProjectFrame.Origin = FVector3d::Zero() + AxisRange.Min * Direction;
+	MeshTransforms::FrameCoordsToWorld(ResultMesh, ProjectFrame);
+
+	FExtrudeMesh Extruder(&ResultMesh);
+	Extruder.DefaultExtrudeDistance = AxisRange.Length();
+	Extruder.UVScaleFactor = 1.0;
+	FVector3d ExtrudeNormal = Direction;
+	Extruder.Apply();
+
+	DestMesh = MoveTemp(ResultMesh);
+}
+
+
+
+
 static void SelectBestFittingMeshApproximation(
 	const FDynamicMesh3& OriginalMesh, 
 	const FDynamicMeshAABBTree3& OriginalMeshSpatial,
@@ -776,6 +939,40 @@ static void SelectBestFittingMeshApproximation(
 		ComputeSimplePartApproximation(PartMeshInOut, PartMeshInOut, EApproximatePartMethod::ConvexHull);
 	}, (int32)EApproximatePartMethod::ConvexHull );
 
+	// Add swept-solid approximations
+	// Currently this is a bit hardcoded and some of these numbers should be exposed as parameters
+	{
+		const int32 UseExtrudeAxis = -1;		// this means auto, ie try all 3
+		const double MinHoleSize = 10.0;		// very aggressive, should be exposed as a parameter
+		const double MinHoleArea = MinHoleSize * MinHoleSize;
+		const double PolyMergeTol = 0.1;
+		const double PolySimplifyTol = AcceptableDeviationTol;
+
+		if (UseExtrudeAxis == 0 || UseExtrudeAxis == -1)
+		{
+			ApproxSelector.AddGeneratedMesh([&](FDynamicMesh3& PartMeshInOut) {
+				ComputeSweptSolidApproximation(PartMeshInOut, PartMeshInOut,
+					FVector3d::UnitX(), PolyMergeTol, PolySimplifyTol, MinHoleArea);
+			}, (int32)EApproximatePartMethod::FlattendExtrusion);
+		}
+
+		if (UseExtrudeAxis == 1 || UseExtrudeAxis == -1)
+		{
+			ApproxSelector.AddGeneratedMesh([&](FDynamicMesh3& PartMeshInOut) {
+				ComputeSweptSolidApproximation(PartMeshInOut, PartMeshInOut,
+					FVector3d::UnitY(), PolyMergeTol, PolySimplifyTol, MinHoleArea);
+			}, (int32)EApproximatePartMethod::FlattendExtrusion);
+		}
+
+		if (UseExtrudeAxis == 2 || UseExtrudeAxis == -1)
+		{
+			ApproxSelector.AddGeneratedMesh([&](FDynamicMesh3& PartMeshInOut) {
+				ComputeSweptSolidApproximation(PartMeshInOut, PartMeshInOut,
+					FVector3d::UnitZ(), PolyMergeTol, PolySimplifyTol, MinHoleArea);
+			}, (int32)EApproximatePartMethod::FlattendExtrusion);
+		}
+	}
+
 	ApproxSelector.SelectBestOption(ResultMesh);
 }
 
@@ -1509,9 +1706,12 @@ void BuildCombinedMesh(
 	using namespace UE::Geometry;
 	bool bVerbose = CVarGeometryCombineMeshInstancesVerbose.GetValueOnGameThread();
 
+	bool bAppendMinimalLOD = false;
+
 	int32 NumLODs = CombineOptions.NumLODs;
+	int NumExtraLODs = (bAppendMinimalLOD) ? 1 : 0;
 	TArray<FCombinedMeshLOD> MeshLODs;
-	MeshLODs.SetNum(NumLODs);
+	MeshLODs.SetNum(NumLODs + NumExtraLODs);
 
 	int FirstVoxWrappedIndex = 9999;
 	TArray<ECombinedLODType> LODTypes;
@@ -1646,7 +1846,7 @@ void BuildCombinedMesh(
 		&& CVarGeometryCombineMeshInstancesRemoveHidden.GetValueOnGameThread() > 0);
 	if (bRemoveHiddenFaces)
 	{
-		for (int32 LODIndex = CombineOptions.RemoveHiddenStartLOD; LODIndex < MeshLODs.Num() && LODIndex < FirstVoxWrappedIndex; ++LODIndex)
+		for (int32 LODIndex = CombineOptions.RemoveHiddenStartLOD; LODIndex < NumLODs && LODIndex < FirstVoxWrappedIndex; ++LODIndex)
 		{
 			if ( bVerbose )
 			{
@@ -1741,11 +1941,17 @@ void BuildCombinedMesh(
 	// wait...
 	UE::Tasks::Wait(PendingRemoveHiddenTasks);
 
+
+	if (bAppendMinimalLOD)
+	{
+		check(false);		// for future use
+	}
+
 	// remove hidden faces on voxel LODs (todo: can do this via shape sorting, much faster)
 	if (bRemoveHiddenFaces)
 	{
 		TRACE_CPUPROFILER_EVENT_SCOPE(RemoveHidden);
-		ParallelFor(MeshLODs.Num(), [&](int32 LODIndex)
+		ParallelFor(NumLODs, [&](int32 LODIndex)
 		{
 			if ( LODIndex >= FirstVoxWrappedIndex )
 			{ 
@@ -1760,7 +1966,9 @@ void BuildCombinedMesh(
 	}
 
 
-	for (int32 LODLevel = 0; LODLevel < NumLODs; ++LODLevel)
+
+
+	for (int32 LODLevel = 0; LODLevel < MeshLODs.Num(); ++LODLevel)
 	{
 		FDynamicMesh3 LODMesh = MoveTemp(MeshLODs[LODLevel].Mesh);
 
diff --git a/Engine/Source/Runtime/GeometryCore/Public/Curve/GeneralPolygon2.h b/Engine/Source/Runtime/GeometryCore/Public/Curve/GeneralPolygon2.h
--- a/Engine/Source/Runtime/GeometryCore/Public/Curve/GeneralPolygon2.h
+++ b/Engine/Source/Runtime/GeometryCore/Public/Curve/GeneralPolygon2.h
@@ -118,6 +118,24 @@ public:
 	}
 
 
+	/**
+	 * Remove any Hole polygons for which RemoveHolePredicateFunc(Hole) returns true
+	 */
+	void FilterHoles(TFunctionRef<bool(const TPolygon2<T>&)> RemoveHolePredicateFunc)
+	{
+		int32 NumHoles = Holes.Num();
+		for (int32 k = 0; k < NumHoles; ++k)
+		{
+			if ( RemoveHolePredicateFunc(Holes[k]) )
+			{
+				Holes.RemoveAtSwap(k);
+				k--;		// need to reconsider index k
+				NumHoles--;
+			}
+		}
+	}
+
+
     double SignedArea() const
     {
         double Sign = (bOuterIsCW) ? -1.0 : 1.0;
@@ -163,7 +181,8 @@ public:
     }
 
 
-	void Translate(TVector2<T> translate) {
+	void Translate(TVector2<T> translate) 
+	{
 		Outer.Translate(translate);
 		for (TPolygon2<T>& Hole : Holes)
 		{
@@ -312,6 +331,21 @@ public:
 		}
     }
 
+
+	/**
+	 * Offset each polygon by the given Distance along vertex "normal" direction (ie for positive offset, outer polygon grows and holes shrink)
+	 * @param OffsetDistance the distance to offset
+	 * @param bUseFaceAvg if true, we offset by the average-face normal instead of the perpendicular-tangent normal
+	 */
+	void VtxNormalOffset(T OffsetDistance, bool bUseFaceAvg = false)
+	{
+		Outer.VtxNormalOffset(OffsetDistance, bUseFaceAvg);
+		for (TPolygon2<T>& Hole : Holes)
+		{
+			Hole.VtxNormalOffset(OffsetDistance, bUseFaceAvg);
+		}
+	}
+
 };
 
 typedef TGeneralPolygon2<double> FGeneralPolygon2d;
