CombineMeshInstances: fix PVS issue

diff --git a/Engine/Plugins/Experimental/MeshModelingToolsetExp/Source/GeometryProcessingAdapters/Private/CombineMeshInstancesImpl.cpp b/Engine/Plugins/Experimental/MeshModelingToolsetExp/Source/GeometryProcessingAdapters/Private/CombineMeshInstancesImpl.cpp
--- a/Engine/Plugins/Experimental/MeshModelingToolsetExp/Source/GeometryProcessingAdapters/Private/CombineMeshInstancesImpl.cpp
+++ b/Engine/Plugins/Experimental/MeshModelingToolsetExp/Source/GeometryProcessingAdapters/Private/CombineMeshInstancesImpl.cpp
@@ -921,7 +921,9 @@ static void SelectBestFittingMeshApproximation(
 	const FDynamicMeshAABBTree3& OriginalMeshSpatial,
 	FDynamicMesh3& ResultMesh,
 	double AcceptableDeviationTol,
-	double TriangleCost)
+	double TriangleCost,
+	int32 UseExtrudeAxis = -1  // axis index, or -1 means auto / try all 3
+	)
 {
 	FPartApproxSelector ApproxSelector;
 	ApproxSelector.Initialize(&OriginalMesh, &OriginalMeshSpatial);
@@ -942,7 +944,6 @@ static void SelectBestFittingMeshApproximation(
 	// Add swept-solid approximations
 	// Currently this is a bit hardcoded and some of these numbers should be exposed as parameters
 	{
-		const int32 UseExtrudeAxis = -1;		// this means auto, ie try all 3
 		const double MinHoleSize = 10.0;		// very aggressive, should be exposed as a parameter
 		const double MinHoleArea = MinHoleSize * MinHoleSize;
 		const double PolyMergeTol = 0.1;
