Skin Binding Fixes
    - invalidate the operator when needed
    - clamp stiffness in the UI

diff --git a/Engine/Plugins/Experimental/MeshModelingToolsetExp/Source/MeshModelingToolsExp/Private/SkinWeightsBindingTool.cpp b/Engine/Plugins/Experimental/MeshModelingToolsetExp/Source/MeshModelingToolsExp/Private/SkinWeightsBindingTool.cpp
--- a/Engine/Plugins/Experimental/MeshModelingToolsetExp/Source/MeshModelingToolsExp/Private/SkinWeightsBindingTool.cpp
+++ b/Engine/Plugins/Experimental/MeshModelingToolsetExp/Source/MeshModelingToolsExp/Private/SkinWeightsBindingTool.cpp
@@ -118,8 +118,31 @@ void USkinWeightsBindingTool::Setup()
 								 ToolSetupUtil::GetDefaultWorkingMaterial(GetToolManager())
 	);
 
+	// setup watchers
+	auto HandlePropertyChange = [this](const bool bUpdateOperator)
+	{
+		if (bUpdateOperator)
+		{
+			Preview->InvalidateResult();
+		}
+		UpdateVisualization();
+	};
+
 	Properties->WatchProperty(Properties->CurrentBone.BoneName,
-							  [this](FName) { UpdateVisualization();});
+							  [HandlePropertyChange](FName) { HandlePropertyChange(false); });
+	Properties->WatchProperty(Properties->BindingType,
+							  [HandlePropertyChange](ESkinWeightsBindType) { HandlePropertyChange(true); });
+	Properties->WatchProperty(Properties->Stiffness,
+							  [HandlePropertyChange](float) { HandlePropertyChange(true); });
+	Properties->WatchProperty(Properties->MaxInfluences,
+							  [HandlePropertyChange](int32) { HandlePropertyChange(true); });
+	Properties->WatchProperty(Properties->VoxelResolution,
+							  [this](int32)
+							  {
+							  	Occupancy = MakeShared<UE::Geometry::FOccupancyGrid3>(*OriginalMesh, Properties->VoxelResolution);
+								Preview->InvalidateResult();
+								UpdateVisualization();
+							  });
 
 	OriginalMesh = MakeShared<FDynamicMesh3, ESPMode::ThreadSafe>();
 	FMeshDescriptionToDynamicMesh Converter;
@@ -250,26 +273,6 @@ bool USkinWeightsBindingTool::CanAccept() const
 	return Super::CanAccept() && Preview->HaveValidResult();
 }
 
-
-void USkinWeightsBindingTool::OnPropertyModified(UObject* PropertySet, FProperty* Property)
-{
-	if ( Property )
-	{
-		if (Property->GetFName() == GET_MEMBER_NAME_CHECKED(FBoneReference, BoneName) /* ||
-			Property->GetFName() == GET_MEMBER_NAME_CHECKED(USkinWeightsBindingToolProperties, bDebugDraw) */ )
-		{
-			// Handled by the property watcher.
-		}
-		else if (Property->GetFName() == GET_MEMBER_NAME_CHECKED(USkinWeightsBindingToolProperties, VoxelResolution))
-		{
-			Occupancy = MakeShared<UE::Geometry::FOccupancyGrid3>(*OriginalMesh, Properties->VoxelResolution);
-			
-			Preview->InvalidateResult();
-		}
-	}
-}
-
-
 TUniquePtr<UE::Geometry::FDynamicMeshOperator> USkinWeightsBindingTool::MakeNewOperator()
 {
 	TUniquePtr<UE::Geometry::FSkinBindingOp> Op = MakeUnique<UE::Geometry::FSkinBindingOp>();
diff --git a/Engine/Plugins/Experimental/MeshModelingToolsetExp/Source/MeshModelingToolsExp/Public/SkinWeightsBindingTool.h b/Engine/Plugins/Experimental/MeshModelingToolsetExp/Source/MeshModelingToolsExp/Public/SkinWeightsBindingTool.h
--- a/Engine/Plugins/Experimental/MeshModelingToolsetExp/Source/MeshModelingToolsExp/Public/SkinWeightsBindingTool.h
+++ b/Engine/Plugins/Experimental/MeshModelingToolsetExp/Source/MeshModelingToolsExp/Public/SkinWeightsBindingTool.h
@@ -58,7 +58,7 @@ public:
 	ESkinWeightsBindType BindingType = ESkinWeightsBindType::DirectDistance;
 
 	/** Stiffness of binding. Lower values allow more distant bones to contribute more */
-	UPROPERTY(EditAnywhere, Category = Binding)
+	UPROPERTY(EditAnywhere, Category = Binding, meta=(ClampMin="0.0", ClampMax="1.0", UIMin="0.0", UIMax="1.0"))
 	float Stiffness = 0.2f;
 
 	/** Maximum bones that will influence each vertex */
@@ -102,8 +102,6 @@ public:
 	virtual bool HasAccept() const override { return true; }
 	virtual bool CanAccept() const override;
 
-	virtual void OnPropertyModified(UObject* PropertySet, FProperty* Property) override;
-
 	// IDynamicMeshOperatorFactory API
 	virtual TUniquePtr<UE::Geometry::FDynamicMeshOperator> MakeNewOperator() override;
 
