Skeleton Editing Tool
    - fixed bone creation
    - removed exposure to the default name

diff --git a/Engine/Plugins/Experimental/MeshModelingToolsetExp/Source/MeshModelingToolsExp/Private/SkeletalMesh/SkeletonEditingTool.cpp b/Engine/Plugins/Experimental/MeshModelingToolsetExp/Source/MeshModelingToolsExp/Private/SkeletalMesh/SkeletonEditingTool.cpp
--- a/Engine/Plugins/Experimental/MeshModelingToolsetExp/Source/MeshModelingToolsExp/Private/SkeletalMesh/SkeletonEditingTool.cpp
+++ b/Engine/Plugins/Experimental/MeshModelingToolsetExp/Source/MeshModelingToolsExp/Private/SkeletalMesh/SkeletonEditingTool.cpp
@@ -400,7 +400,8 @@ void USkeletonEditingTool::CreateNewBone()
 
 	BeginChange();
 
-	const FName BoneName = Modifier->GetUniqueName(Properties->DefaultName);
+	static const FName DefaultName("joint");
+	const FName BoneName = Modifier->GetUniqueName(DefaultName);
 	const FName ParentName = Selection.IsEmpty() ? NAME_None : Selection.Last(); 
 	const bool bBoneAdded = Modifier->AddBone(BoneName, ParentName, Properties->Transform);
 	if (bBoneAdded)
@@ -858,7 +859,7 @@ FInputRayHit USkeletonEditingTool::CanBeginClickDragSequence(const FInputDeviceR
 		if (ProjectionProperties->GetProjectionPoint(InPressPos, HitPoint))
 		{
 			// CurrentBone is gonna be the parent
-			PendingFunction = [&]
+			PendingFunction = [this, HitPoint]
 			{
 				const FReferenceSkeleton& ReferenceSkeleton = Modifier->GetReferenceSkeleton();
 				ParentIndex = ReferenceSkeleton.FindRawBoneIndex(GetCurrentBone());
diff --git a/Engine/Plugins/Experimental/MeshModelingToolsetExp/Source/MeshModelingToolsExp/Public/SkeletalMesh/SkeletonEditingTool.h b/Engine/Plugins/Experimental/MeshModelingToolsetExp/Source/MeshModelingToolsExp/Public/SkeletalMesh/SkeletonEditingTool.h
--- a/Engine/Plugins/Experimental/MeshModelingToolsetExp/Source/MeshModelingToolsExp/Public/SkeletalMesh/SkeletonEditingTool.h
+++ b/Engine/Plugins/Experimental/MeshModelingToolsetExp/Source/MeshModelingToolsExp/Public/SkeletalMesh/SkeletonEditingTool.h
@@ -254,9 +254,6 @@ public:
 	UPROPERTY(EditAnywhere, Category = "Move")
 	bool bUpdateChildren = false;
 
-	UPROPERTY(EditAnywhere, Category = "Create")
-	FName DefaultName = "joint";
-
 	UPROPERTY(EditAnywhere, Category = "Viewport Axis Settings",  meta = (ClampMin = "0.0", UIMin = "0.0"))
 	float AxisLength = 1.f;
 
