ModelingTools: add support for Mesh Selection to Generate Polygroups / Convert To Polygons Tool.
Base class is changed to USingleTargetWithSelectionTool.
When selection is provided, a UPreviewMesh for the unselected area is created and used.
A FSharedConstDynamicMesh3 is used to wrap either the original mesh or the selected submesh, and passed to the Operator (which also now takes a FSharedConstDynamicMesh3)
On shutdown, modified submesh is re-integrated into original mesh (in this case a trivial group copy suffices) and a new selection is emitted

diff --git a/Engine/Plugins/Experimental/MeshModelingToolsetExp/Source/MeshModelingToolsExp/Private/ConvertToPolygonsTool.cpp b/Engine/Plugins/Experimental/MeshModelingToolsetExp/Source/MeshModelingToolsExp/Private/ConvertToPolygonsTool.cpp
--- a/Engine/Plugins/Experimental/MeshModelingToolsetExp/Source/MeshModelingToolsExp/Private/ConvertToPolygonsTool.cpp
+++ b/Engine/Plugins/Experimental/MeshModelingToolsetExp/Source/MeshModelingToolsExp/Private/ConvertToPolygonsTool.cpp
@@ -13,11 +13,15 @@
 
 #include "DynamicMesh/DynamicMesh3.h"
 #include "DynamicMesh/MeshNormals.h"
+#include "DynamicMesh/MeshSharingUtil.h"
 #include "DynamicMeshEditor.h"
+#include "DynamicSubmesh3.h"
 #include "MeshRegionBoundaryLoops.h"
 #include "Util/ColorConstants.h"
 #include "Polygroups/PolygroupUtil.h"
 #include "Util/ColorConstants.h"
+#include "Selections/GeometrySelectionUtil.h"
+#include "Selection/StoredMeshSelectionUtil.h"
 
 #include "Polygroups/PolygroupsGenerator.h"
 
@@ -31,7 +35,7 @@ using namespace UE::Geometry;
  * ToolBuilder
  */
 
-USingleSelectionMeshEditingTool* UConvertToPolygonsToolBuilder::CreateNewTool(const FToolBuilderState& SceneState) const
+USingleTargetWithSelectionTool* UConvertToPolygonsToolBuilder::CreateNewTool(const FToolBuilderState& SceneState) const
 {
 	return NewObject<UConvertToPolygonsTool>(SceneState.ToolManager);
 }
@@ -55,11 +59,12 @@ public:
 	int QuadSearchRounds = 1;
 
 	int32 MinGroupSize = 2;
+	int32 InitialGroupID = 0;
 
 	bool bCalculateNormals = false;
 
 	// input mesh
-	TSharedPtr<FDynamicMesh3, ESPMode::ThreadSafe> OriginalMesh;
+	TSharedPtr<FSharedConstDynamicMesh3, ESPMode::ThreadSafe> OriginalMesh;
 	// input polygroups, if available
 	TSharedPtr<FPolygroupSet, ESPMode::ThreadSafe> SourcePolyGroups;
 
@@ -68,12 +73,15 @@ public:
 
 	virtual void CalculateResult(FProgressCancel* Progress) override
 	{
-		if ((Progress && Progress->Cancelled()) || !OriginalMesh)
+		if ((Progress && Progress->Cancelled()) || OriginalMesh.IsValid() == false)
 		{
 			return;
 		}
 
-		ResultMesh->Copy(*OriginalMesh, true, true, true, true);
+		OriginalMesh->AccessSharedObject([&](const FDynamicMesh3& ReadMesh)
+		{
+			ResultMesh->Copy(ReadMesh, true, true, true, true);
+		});
 	
 		if (Progress && Progress->Cancelled())
 		{
@@ -92,6 +100,7 @@ public:
 
 		Generator = FPolygroupsGenerator(ResultMesh.Get());
 		Generator.MinGroupSize = MinGroupSize;
+		Generator.InitialGroupID = InitialGroupID;
 
 		switch (ConversionMode)
 		{
@@ -134,8 +143,11 @@ public:
 		{
 			if (bSubdivideExisting)
 			{
-				FPolygroupSet InputGroups(OriginalMesh.Get());
-				Generator.FindPolygroupsFromFurthestPointSampling(NumPoints, WeightingType, WeightingCoeffs, &InputGroups);
+				OriginalMesh->AccessSharedObject([&](const FDynamicMesh3& ReadMesh)
+				{
+					FPolygroupSet InputGroups(&ReadMesh);
+					Generator.FindPolygroupsFromFurthestPointSampling(NumPoints, WeightingType, WeightingCoeffs, &InputGroups);
+				});
 			}
 			else
 			{
@@ -217,6 +229,27 @@ void UConvertToPolygonsTool::Setup()
 
 	OriginalDynamicMesh = MakeShared<FDynamicMesh3, ESPMode::ThreadSafe>(UE::ToolTarget::GetDynamicMeshCopy(Target));
 
+	// initialize triangle ROI if one exists
+	SelectionTriangleROI = MakeShared<TSet<int32>, ESPMode::ThreadSafe>();
+	if (HasGeometrySelection())
+	{
+		const FGeometrySelection& InputSelection = GetGeometrySelection();
+		UE::Geometry::EnumerateSelectionTriangles(InputSelection, *OriginalDynamicMesh,
+			[&](int32 TriangleID) { SelectionTriangleROI->Add(TriangleID); });
+
+		OriginalSubmesh = MakeShared<FDynamicSubmesh3, ESPMode::ThreadSafe>(OriginalDynamicMesh.Get(), SelectionTriangleROI->Array());
+		bUsingSelection = true;
+	}
+
+	if (bUsingSelection)
+	{
+		ComputeOperatorSharedMesh = MakeShared<FSharedConstDynamicMesh3>(&OriginalSubmesh->GetSubmesh());
+	}
+	else
+	{
+		ComputeOperatorSharedMesh = MakeShared<FSharedConstDynamicMesh3>(OriginalDynamicMesh.Get());
+	}
+
 	Settings = NewObject<UConvertToPolygonsToolProperties>(this);
 	Settings->RestoreProperties(this);
 	AddToolPropertySource(Settings);
@@ -255,6 +288,22 @@ void UConvertToPolygonsTool::Setup()
 
 	}
 	
+	if (bUsingSelection)
+	{
+		// create the preview object for the unmodified area
+		UnmodifiedAreaPreviewMesh = NewObject<UPreviewMesh>();
+		UnmodifiedAreaPreviewMesh->CreateInWorld(GetTargetWorld(), MeshTransform);
+		ToolSetupUtil::ApplyRenderingConfigurationToPreview(UnmodifiedAreaPreviewMesh, Target);
+		UnmodifiedAreaPreviewMesh->SetMaterials(MaterialSet.Materials);
+		UnmodifiedAreaPreviewMesh->EnableSecondaryTriangleBuffers( [this](const FDynamicMesh3* Mesh, int32 TriangleID) {
+			return SelectionTriangleROI->Contains(TriangleID);
+		});
+		UnmodifiedAreaPreviewMesh->SetSecondaryBuffersVisibility(false);
+		UnmodifiedAreaPreviewMesh->SetTangentsMode(EDynamicMeshComponentTangentsMode::AutoCalculated);
+		UnmodifiedAreaPreviewMesh->UpdatePreview(OriginalDynamicMesh.Get());
+	}
+
+
 	PreviewGeometry = NewObject<UPreviewGeometry>(this);
 	PreviewGeometry->CreateInWorld(GetTargetWorld(), MeshTransform);
 
@@ -324,10 +373,20 @@ void UConvertToPolygonsTool::UpdateOpParameters(FConvertToPolygonsOp& ConvertToP
 	ConvertToPolygonsOp.bRespectUVSeams = Settings->bRespectUVSeams;
 	ConvertToPolygonsOp.bRespectHardNormals = Settings->bRespectHardNormals;
 
-	ConvertToPolygonsOp.OriginalMesh = OriginalDynamicMesh;
+	ConvertToPolygonsOp.OriginalMesh = ComputeOperatorSharedMesh;
+
+	if (bUsingSelection)
+	{
+		ConvertToPolygonsOp.InitialGroupID = OriginalDynamicMesh->MaxGroupID();
+	}
+
 	if (ActiveFromGroupSet.IsValid())
 	{
 		ConvertToPolygonsOp.SourcePolyGroups = ActiveFromGroupSet;
+		if (bUsingSelection)
+		{
+			ConvertToPolygonsOp.InitialGroupID = ActiveFromGroupSet->MaxGroupID;
+		}
 	}
 
 	FTransform LocalToWorld = (FTransform)UE::ToolTarget::GetLocalToWorldTransform(Target);
@@ -341,8 +400,17 @@ void UConvertToPolygonsTool::OnShutdown(EToolShutdownType ShutdownType)
 	Settings->SaveProperties(this);
 	UE::ToolTarget::ShowSourceObject(Target);
 
-	PreviewGeometry->Disconnect();
-	PreviewGeometry = nullptr;
+	if (PreviewGeometry)
+	{
+		PreviewGeometry->Disconnect();
+		PreviewGeometry = nullptr;
+	}
+
+	if (UnmodifiedAreaPreviewMesh != nullptr)
+	{
+		UnmodifiedAreaPreviewMesh->Disconnect();
+		UnmodifiedAreaPreviewMesh = nullptr;
+	}
 
 	if (PreviewCompute)
 	{
@@ -355,13 +423,13 @@ void UConvertToPolygonsTool::OnShutdown(EToolShutdownType ShutdownType)
 			{
 				if (OutputProperties->GroupLayer != TEXT("Default"))
 				{
-					// if we want to write to any layer other than default, we have to find or create it
 					FDynamicMesh3 UseResultMesh = *OriginalDynamicMesh;
 					if (UseResultMesh.HasAttributes() == false)
 					{
 						UseResultMesh.EnableAttributes();
 					}
 
+					// if we want to write to any layer other than default, we have to find or create it
 					FDynamicMeshPolygroupAttribute* UseAttribLayer = nullptr;
 					if (OutputProperties->GroupLayer == TEXT("Create New..."))
 					{
@@ -380,9 +448,20 @@ void UConvertToPolygonsTool::OnShutdown(EToolShutdownType ShutdownType)
 					if (UseAttribLayer)
 					{
 						// copy the generated groups from the Op output mesh (stored in primary groups) to the target layer
-						for (int32 tid : UseResultMesh.TriangleIndicesItr())
+						if (bUsingSelection)
+						{
+							for (int32 tid : *SelectionTriangleROI)
+							{
+								int32 GroupID = DynamicMeshResult->GetTriangleGroup(OriginalSubmesh->MapTriangleToSubmesh(tid));
+								UseAttribLayer->SetValue(tid, GroupID);
+							}
+						}
+						else
 						{
-							UseAttribLayer->SetValue(tid, DynamicMeshResult->GetTriangleGroup(tid));
+							for (int32 tid : UseResultMesh.TriangleIndicesItr())
+							{
+								UseAttribLayer->SetValue(tid, DynamicMeshResult->GetTriangleGroup(tid));
+							}
 						}
 						UE::ToolTarget::CommitDynamicMeshUpdate(Target, UseResultMesh, true);
 					}
@@ -396,12 +475,43 @@ void UConvertToPolygonsTool::OnShutdown(EToolShutdownType ShutdownType)
 				else
 				{
 					// todo: have not actually modified topology here, but groups-only update is not supported yet
-					UE::ToolTarget::CommitDynamicMeshUpdate(Target, *DynamicMeshResult, true);
+
+					if (bUsingSelection)
+					{
+						for (int32 tid : *SelectionTriangleROI)
+						{
+							int32 GroupID = DynamicMeshResult->GetTriangleGroup(OriginalSubmesh->MapTriangleToSubmesh(tid));
+							OriginalDynamicMesh->SetTriangleGroup(tid, GroupID);
+						}
+						UE::ToolTarget::CommitDynamicMeshUpdate(Target, *OriginalDynamicMesh, true);
+					}
+					else
+					{
+						UE::ToolTarget::CommitDynamicMeshUpdate(Target, *DynamicMeshResult, true);
+					}
+				}
+			}
+
+			if (bUsingSelection)
+			{
+				// if the input was a group selection, that selection is no longer valid. But if we output
+				// a triangle selection it should be converted to the group selection
+				FGeometrySelection OutputSelection;
+				for (int32 tid : *SelectionTriangleROI)
+				{
+					OutputSelection.Selection.Add(FGeoSelectionID::MeshTriangle(tid).Encoded());
 				}
+				UE::Geometry::SetToolOutputGeometrySelectionForTarget(this, Target, OutputSelection);
 			}
+
 			GetToolManager()->EndUndoTransaction();
 		}
 	}
+
+	if (ComputeOperatorSharedMesh.IsValid())
+	{
+		ComputeOperatorSharedMesh->ReleaseSharedObject();
+	}
 }
 
 void UConvertToPolygonsTool::OnSettingsModified()
@@ -426,17 +536,20 @@ void UConvertToPolygonsTool::OnSelectedFromGroupLayerChanged()
 
 void UConvertToPolygonsTool::UpdateFromGroupLayer()
 {
-	if (CopyFromLayerProperties->HasSelectedPolygroup() == false)
-	{
-		ActiveFromGroupSet = MakeShared<UE::Geometry::FPolygroupSet, ESPMode::ThreadSafe>(OriginalDynamicMesh.Get());
-	}
-	else
+	ComputeOperatorSharedMesh->AccessSharedObject([&](const FDynamicMesh3& ReadMesh)
 	{
-		FName SelectedName = CopyFromLayerProperties->ActiveGroupLayer;
-		FDynamicMeshPolygroupAttribute* FoundAttrib = UE::Geometry::FindPolygroupLayerByName(*OriginalDynamicMesh, SelectedName);
-		ensureMsgf(FoundAttrib, TEXT("Selected Attribute Not Found! Falling back to Default group layer."));
-		ActiveFromGroupSet = MakeShared<UE::Geometry::FPolygroupSet, ESPMode::ThreadSafe>(OriginalDynamicMesh.Get(), FoundAttrib);
-	}
+		if (CopyFromLayerProperties->HasSelectedPolygroup() == false)
+		{
+			ActiveFromGroupSet = MakeShared<UE::Geometry::FPolygroupSet, ESPMode::ThreadSafe>(&ReadMesh);
+		}
+		else
+		{
+			FName SelectedName = CopyFromLayerProperties->ActiveGroupLayer;
+			const FDynamicMeshPolygroupAttribute* FoundAttrib = UE::Geometry::FindPolygroupLayerByName(ReadMesh, SelectedName);
+			ensureMsgf(FoundAttrib, TEXT("Selected Attribute Not Found! Falling back to Default group layer."));
+			ActiveFromGroupSet = MakeShared<UE::Geometry::FPolygroupSet, ESPMode::ThreadSafe>(&ReadMesh, FoundAttrib);
+		}
+	});
 }
 
 
@@ -473,13 +586,15 @@ void UConvertToPolygonsTool::UpdateVisualization()
 	FColor GroupLineColor = FColor::Red;
 	float GroupLineThickness = 2.0f;
 
-	PreviewGeometry->CreateOrUpdateLineSet(TEXT("GroupBorders"), PolygonEdges.Num(), 
-		[&](int32 k, TArray<FRenderableLine>& LinesOut) {
-			FVector3d A, B;
-			OriginalDynamicMesh->GetEdgeV(PolygonEdges[k], A, B);
-			LinesOut.Add(FRenderableLine(A, B, GroupLineColor, GroupLineThickness));
-		}, 1);
-
+	ComputeOperatorSharedMesh->AccessSharedObject([&](const FDynamicMesh3& ReadMesh)
+	{
+		PreviewGeometry->CreateOrUpdateLineSet(TEXT("GroupBorders"), PolygonEdges.Num(), 
+			[&](int32 k, TArray<FRenderableLine>& LinesOut) {
+				FVector3d A, B;
+				ReadMesh.GetEdgeV(PolygonEdges[k], A, B);
+				LinesOut.Add(FRenderableLine(A, B, GroupLineColor, GroupLineThickness));
+			}, 1);
+	});
 }
 
 
diff --git a/Engine/Plugins/Experimental/MeshModelingToolsetExp/Source/MeshModelingToolsExp/Public/ConvertToPolygonsTool.h b/Engine/Plugins/Experimental/MeshModelingToolsetExp/Source/MeshModelingToolsExp/Public/ConvertToPolygonsTool.h
--- a/Engine/Plugins/Experimental/MeshModelingToolsetExp/Source/MeshModelingToolsExp/Public/ConvertToPolygonsTool.h
+++ b/Engine/Plugins/Experimental/MeshModelingToolsetExp/Source/MeshModelingToolsExp/Public/ConvertToPolygonsTool.h
@@ -3,28 +3,35 @@
 #pragma once
 
 #include "CoreMinimal.h"
-#include "BaseTools/SingleSelectionMeshEditingTool.h"
+#include "BaseTools/SingleTargetWithSelectionTool.h"
 #include "PreviewMesh.h"
 #include "ModelingOperators.h"
 #include "MeshOpPreviewHelpers.h"
 #include "PropertySets/PolygroupLayersProperties.h"
+#include "DynamicMesh/MeshSharingUtil.h"
 #include "ConvertToPolygonsTool.generated.h"
 
 // predeclaration
 class UConvertToPolygonsTool;
 class FConvertToPolygonsOp;
 class UPreviewGeometry;
-PREDECLARE_GEOMETRY(class FDynamicMesh3);
+namespace UE::Geometry
+{
+	class FDynamicMesh3;
+	struct FDynamicSubmesh3;
+}
 
 /**
  *
  */
 UCLASS()
-class MESHMODELINGTOOLSEXP_API UConvertToPolygonsToolBuilder : public USingleSelectionMeshEditingToolBuilder
+class MESHMODELINGTOOLSEXP_API UConvertToPolygonsToolBuilder : public USingleTargetWithSelectionToolBuilder
 {
 	GENERATED_BODY()
 public:
-	virtual USingleSelectionMeshEditingTool* CreateNewTool(const FToolBuilderState & SceneState) const override;
+	virtual USingleTargetWithSelectionTool* CreateNewTool(const FToolBuilderState & SceneState) const override;
+
+	virtual bool RequiresInputSelection() const override { return false; }
 };
 
 
@@ -170,7 +177,7 @@ public:
  *
  */
 UCLASS()
-class MESHMODELINGTOOLSEXP_API UConvertToPolygonsTool : public USingleSelectionMeshEditingTool
+class MESHMODELINGTOOLSEXP_API UConvertToPolygonsTool : public USingleTargetWithSelectionTool
 {
 	GENERATED_BODY()
 
@@ -210,15 +217,29 @@ protected:
 	UPROPERTY()
 	TObjectPtr<UPreviewGeometry> PreviewGeometry = nullptr;
 
+	// If a selection was provided (bUsingSelection = true), UnmodifiedAreaPreviewMesh is used to render the unmodified (non-selected) part of the input mesh,
+	// as the PreviewCompute input mesh will be limited to the input selected area
+	UPROPERTY()
+	TObjectPtr<UPreviewMesh> UnmodifiedAreaPreviewMesh = nullptr;
+
 protected:
+	// copy of the input mesh (possibly does not need to be SharedPtr anymore now that FSharedConstDynamicMesh3 is used below, but then needs PimplPtr...)
 	TSharedPtr<UE::Geometry::FDynamicMesh3, ESPMode::ThreadSafe> OriginalDynamicMesh;
 
-	// for visualization
+	// If there is an active selection, bUsingSelection will be true and the ROI and Submesh will be initialized
+	bool bUsingSelection = false;
+	TSharedPtr<TSet<int32>> SelectionTriangleROI;
+	TSharedPtr<UE::Geometry::FDynamicSubmesh3, ESPMode::ThreadSafe> OriginalSubmesh;
+
+	// The mesh passed to the Compute operator to base result on, will either be OriginalDynamicMesh or the submesh from OriginalSubmesh if using a selection
+	TSharedPtr<UE::Geometry::FSharedConstDynamicMesh3> ComputeOperatorSharedMesh;
+
+	// current set of detected polygroup edges, relative to ComputeOperatorSharedMesh
 	TArray<int> PolygonEdges;
 	
 	void UpdateVisualization();
 
-
+	// current input group set used in CopyFromLayer mode, relative to ComputeOperatorSharedMesh
 	TSharedPtr<UE::Geometry::FPolygroupSet, ESPMode::ThreadSafe> ActiveFromGroupSet;
 	void OnSelectedFromGroupLayerChanged();
 	void UpdateFromGroupLayer();
diff --git a/Engine/Plugins/Runtime/GeometryProcessing/Source/DynamicMesh/Private/Polygroups/PolygroupsGenerator.cpp b/Engine/Plugins/Runtime/GeometryProcessing/Source/DynamicMesh/Private/Polygroups/PolygroupsGenerator.cpp
--- a/Engine/Plugins/Runtime/GeometryProcessing/Source/DynamicMesh/Private/Polygroups/PolygroupsGenerator.cpp
+++ b/Engine/Plugins/Runtime/GeometryProcessing/Source/DynamicMesh/Private/Polygroups/PolygroupsGenerator.cpp
@@ -918,7 +918,7 @@ void FPolygroupsGenerator::CopyPolygroupsToMesh()
 		int NumTriangles = Polygroup.Num();
 		for (int k = 0; k < NumTriangles; ++k)
 		{
-			Mesh->SetTriangleGroup(Polygroup[k], (PolyIdx + 1));
+			Mesh->SetTriangleGroup(Polygroup[k], InitialGroupID + (PolyIdx + 1));
 		}
 	}
 }
@@ -935,7 +935,7 @@ void FPolygroupsGenerator::CopyPolygroupsToPolygroupSet(FPolygroupSet& Polygroup
 		int NumTriangles = Polygroup.Num();
 		for (int k = 0; k < NumTriangles; ++k)
 		{
-			Polygroups.SetGroup(Polygroup[k], (PolyIdx + 1), TargetMesh);
+			Polygroups.SetGroup(Polygroup[k], InitialGroupID + (PolyIdx + 1), TargetMesh);
 		}
 	}
 }
diff --git a/Engine/Plugins/Runtime/GeometryProcessing/Source/DynamicMesh/Public/Polygroups/PolygroupsGenerator.h b/Engine/Plugins/Runtime/GeometryProcessing/Source/DynamicMesh/Public/Polygroups/PolygroupsGenerator.h
--- a/Engine/Plugins/Runtime/GeometryProcessing/Source/DynamicMesh/Public/Polygroups/PolygroupsGenerator.h
+++ b/Engine/Plugins/Runtime/GeometryProcessing/Source/DynamicMesh/Public/Polygroups/PolygroupsGenerator.h
@@ -41,6 +41,9 @@ public:
 	// if true, after groups are computed they will be copied to the output mesh
 	bool bCopyToMesh = true;
 
+	// start at this GroupID when assigning new groups
+	int32 InitialGroupID = 0;
+
 	//
 	// Outputs
 	//
