ModelingMode: Adjust CanBuildTool rules on several tools to disallow their execution on Volume targets since they don't apply conceptually.

diff --git a/Engine/Plugins/Experimental/MeshModelingToolsetExp/Source/MeshModelingToolsEditorOnlyExp/Private/AttributeEditorTool.cpp b/Engine/Plugins/Experimental/MeshModelingToolsetExp/Source/MeshModelingToolsEditorOnlyExp/Private/AttributeEditorTool.cpp
--- a/Engine/Plugins/Experimental/MeshModelingToolsetExp/Source/MeshModelingToolsEditorOnlyExp/Private/AttributeEditorTool.cpp
+++ b/Engine/Plugins/Experimental/MeshModelingToolsetExp/Source/MeshModelingToolsEditorOnlyExp/Private/AttributeEditorTool.cpp
@@ -22,6 +22,7 @@
 #include "TargetInterfaces/MeshDescriptionCommitter.h"
 #include "TargetInterfaces/PrimitiveComponentBackedTarget.h"
 #include "ModelingToolTargetUtil.h"
+#include "ToolTargetManager.h"
 
 #include UE_INLINE_GENERATED_CPP_BY_NAME(AttributeEditorTool)
 
@@ -38,7 +39,12 @@ UMultiSelectionMeshEditingTool* UAttributeEditorToolBuilder::CreateNewTool(const
 	return NewObject<UAttributeEditorTool>(SceneState.ToolManager);
 }
 
-
+bool UAttributeEditorToolBuilder::CanBuildTool(const FToolBuilderState& SceneState) const
+{
+	return UMultiSelectionMeshEditingToolBuilder::CanBuildTool(SceneState) &&
+		SceneState.TargetManager->CountSelectedAndTargetableWithPredicate(SceneState, GetTargetRequirements(),
+			[](UActorComponent& Component) { return !ToolBuilderUtil::IsVolume(Component); }) >= 1;
+}
 
 
 void UAttributeEditorActionPropertySet::PostAction(EAttributeEditorToolActions Action)
diff --git a/Engine/Plugins/Experimental/MeshModelingToolsetExp/Source/MeshModelingToolsEditorOnlyExp/Private/EditMeshMaterialsTool.cpp b/Engine/Plugins/Experimental/MeshModelingToolsetExp/Source/MeshModelingToolsEditorOnlyExp/Private/EditMeshMaterialsTool.cpp
--- a/Engine/Plugins/Experimental/MeshModelingToolsetExp/Source/MeshModelingToolsEditorOnlyExp/Private/EditMeshMaterialsTool.cpp
+++ b/Engine/Plugins/Experimental/MeshModelingToolsetExp/Source/MeshModelingToolsEditorOnlyExp/Private/EditMeshMaterialsTool.cpp
@@ -15,6 +15,7 @@
 #include "ToolSetupUtil.h"
 #include "ModelingToolTargetUtil.h"
 #include "Materials/MaterialInterface.h"
+#include "ToolTargetManager.h"
 
 #include UE_INLINE_GENERATED_CPP_BY_NAME(EditMeshMaterialsTool)
 
@@ -80,7 +81,12 @@ UMeshSurfacePointTool* UEditMeshMaterialsToolBuilder::CreateNewTool(const FToolB
 	return SelectionTool;
 }
 
-
+bool UEditMeshMaterialsToolBuilder::CanBuildTool(const FToolBuilderState& SceneState) const
+{
+	return UMeshSelectionToolBuilder::CanBuildTool(SceneState) &&
+		SceneState.TargetManager->CountSelectedAndTargetableWithPredicate(SceneState, GetTargetRequirements(),
+			[](UActorComponent& Component) { return !ToolBuilderUtil::IsVolume(Component); }) >= 1;
+}
 
 
 void UEditMeshMaterialsTool::Setup()
diff --git a/Engine/Plugins/Experimental/MeshModelingToolsetExp/Source/MeshModelingToolsEditorOnlyExp/Private/MeshTangentsTool.cpp b/Engine/Plugins/Experimental/MeshModelingToolsetExp/Source/MeshModelingToolsEditorOnlyExp/Private/MeshTangentsTool.cpp
--- a/Engine/Plugins/Experimental/MeshModelingToolsetExp/Source/MeshModelingToolsEditorOnlyExp/Private/MeshTangentsTool.cpp
+++ b/Engine/Plugins/Experimental/MeshModelingToolsetExp/Source/MeshModelingToolsEditorOnlyExp/Private/MeshTangentsTool.cpp
@@ -20,6 +20,7 @@
 #include "TargetInterfaces/PrimitiveComponentBackedTarget.h"
 #include "TargetInterfaces/StaticMeshBackedTarget.h"
 #include "ModelingToolTargetUtil.h"
+#include "ToolTargetManager.h"
 
 #include UE_INLINE_GENERATED_CPP_BY_NAME(MeshTangentsTool)
 
@@ -47,6 +48,13 @@ USingleSelectionMeshEditingTool* UMeshTangentsToolBuilder::CreateNewTool(const F
 	return NewObject<UMeshTangentsTool>(SceneState.ToolManager);
 }
 
+bool UMeshTangentsToolBuilder::CanBuildTool(const FToolBuilderState& SceneState) const
+{
+	return USingleSelectionMeshEditingToolBuilder::CanBuildTool(SceneState) &&
+		SceneState.TargetManager->CountSelectedAndTargetableWithPredicate(SceneState, GetTargetRequirements(),
+			[](UActorComponent& Component) { return !ToolBuilderUtil::IsVolume(Component); }) >= 1;
+}
+
 /*
  * Tool
  */
diff --git a/Engine/Plugins/Experimental/MeshModelingToolsetExp/Source/MeshModelingToolsEditorOnlyExp/Public/AttributeEditorTool.h b/Engine/Plugins/Experimental/MeshModelingToolsetExp/Source/MeshModelingToolsEditorOnlyExp/Public/AttributeEditorTool.h
--- a/Engine/Plugins/Experimental/MeshModelingToolsetExp/Source/MeshModelingToolsEditorOnlyExp/Public/AttributeEditorTool.h
+++ b/Engine/Plugins/Experimental/MeshModelingToolsetExp/Source/MeshModelingToolsEditorOnlyExp/Public/AttributeEditorTool.h
@@ -24,6 +24,7 @@ class MESHMODELINGTOOLSEDITORONLYEXP_API UAttributeEditorToolBuilder : public UM
 
 public:
 	virtual UMultiSelectionMeshEditingTool* CreateNewTool(const FToolBuilderState& SceneState) const override;
+	virtual bool CanBuildTool(const FToolBuilderState& SceneState) const override;
 };
 
 
diff --git a/Engine/Plugins/Experimental/MeshModelingToolsetExp/Source/MeshModelingToolsEditorOnlyExp/Public/EditMeshMaterialsTool.h b/Engine/Plugins/Experimental/MeshModelingToolsetExp/Source/MeshModelingToolsEditorOnlyExp/Public/EditMeshMaterialsTool.h
--- a/Engine/Plugins/Experimental/MeshModelingToolsetExp/Source/MeshModelingToolsEditorOnlyExp/Public/EditMeshMaterialsTool.h
+++ b/Engine/Plugins/Experimental/MeshModelingToolsetExp/Source/MeshModelingToolsEditorOnlyExp/Public/EditMeshMaterialsTool.h
@@ -16,6 +16,7 @@ class MESHMODELINGTOOLSEDITORONLYEXP_API UEditMeshMaterialsToolBuilder : public
 
 public:
 	virtual UMeshSurfacePointTool* CreateNewTool(const FToolBuilderState& SceneState) const override;
+	virtual bool CanBuildTool(const FToolBuilderState& SceneState) const override;
 };
 
 
diff --git a/Engine/Plugins/Experimental/MeshModelingToolsetExp/Source/MeshModelingToolsEditorOnlyExp/Public/MeshTangentsTool.h b/Engine/Plugins/Experimental/MeshModelingToolsetExp/Source/MeshModelingToolsEditorOnlyExp/Public/MeshTangentsTool.h
--- a/Engine/Plugins/Experimental/MeshModelingToolsetExp/Source/MeshModelingToolsEditorOnlyExp/Public/MeshTangentsTool.h
+++ b/Engine/Plugins/Experimental/MeshModelingToolsetExp/Source/MeshModelingToolsEditorOnlyExp/Public/MeshTangentsTool.h
@@ -30,6 +30,7 @@ class MESHMODELINGTOOLSEDITORONLYEXP_API UMeshTangentsToolBuilder : public USing
 
 public:
 	virtual USingleSelectionMeshEditingTool* CreateNewTool(const FToolBuilderState& SceneState) const override;
+	virtual bool CanBuildTool(const FToolBuilderState& SceneState) const override;
 
 protected:
 	virtual const FToolTargetTypeRequirements& GetTargetRequirements() const override;
diff --git a/Engine/Plugins/Experimental/MeshModelingToolsetExp/Source/MeshModelingToolsExp/Private/ConvertToPolygonsTool.cpp b/Engine/Plugins/Experimental/MeshModelingToolsetExp/Source/MeshModelingToolsExp/Private/ConvertToPolygonsTool.cpp
--- a/Engine/Plugins/Experimental/MeshModelingToolsetExp/Source/MeshModelingToolsExp/Private/ConvertToPolygonsTool.cpp
+++ b/Engine/Plugins/Experimental/MeshModelingToolsetExp/Source/MeshModelingToolsExp/Private/ConvertToPolygonsTool.cpp
@@ -22,6 +22,7 @@
 #include "Util/ColorConstants.h"
 #include "Selections/GeometrySelectionUtil.h"
 #include "Selection/StoredMeshSelectionUtil.h"
+#include "ToolTargetManager.h"
 
 #include "Polygroups/PolygroupsGenerator.h"
 
@@ -40,6 +41,13 @@ USingleTargetWithSelectionTool* UConvertToPolygonsToolBuilder::CreateNewTool(con
 	return NewObject<UConvertToPolygonsTool>(SceneState.ToolManager);
 }
 
+bool UConvertToPolygonsToolBuilder::CanBuildTool(const FToolBuilderState& SceneState) const
+{
+	return USingleTargetWithSelectionToolBuilder::CanBuildTool(SceneState) &&
+		SceneState.TargetManager->CountSelectedAndTargetableWithPredicate(SceneState, GetTargetRequirements(),
+			[](UActorComponent& Component) { return !ToolBuilderUtil::IsVolume(Component); }) >= 1;
+}
+
 class FConvertToPolygonsOp : public  FDynamicMeshOperator
 {
 public:
diff --git a/Engine/Plugins/Experimental/MeshModelingToolsetExp/Source/MeshModelingToolsExp/Private/EditNormalsTool.cpp b/Engine/Plugins/Experimental/MeshModelingToolsetExp/Source/MeshModelingToolsExp/Private/EditNormalsTool.cpp
--- a/Engine/Plugins/Experimental/MeshModelingToolsetExp/Source/MeshModelingToolsExp/Private/EditNormalsTool.cpp
+++ b/Engine/Plugins/Experimental/MeshModelingToolsetExp/Source/MeshModelingToolsExp/Private/EditNormalsTool.cpp
@@ -25,6 +25,7 @@
 
 #include "TargetInterfaces/PrimitiveComponentBackedTarget.h"
 #include "ModelingToolTargetUtil.h"
+#include "ToolTargetManager.h"
 
 #include UE_INLINE_GENERATED_CPP_BY_NAME(EditNormalsTool)
 
@@ -61,6 +62,12 @@ void UEditNormalsToolBuilder::InitializeNewTool(UMultiSelectionMeshEditingTool*
 	}
 }
 
+bool UEditNormalsToolBuilder::CanBuildTool(const FToolBuilderState& SceneState) const
+{
+	return UMultiSelectionMeshEditingToolBuilder::CanBuildTool(SceneState) &&
+		SceneState.TargetManager->CountSelectedAndTargetableWithPredicate(SceneState, GetTargetRequirements(),
+			[](UActorComponent& Component) { return !ToolBuilderUtil::IsVolume(Component); }) >= 1;
+}
 
 
 
diff --git a/Engine/Plugins/Experimental/MeshModelingToolsetExp/Source/MeshModelingToolsExp/Private/HoleFillTool.cpp b/Engine/Plugins/Experimental/MeshModelingToolsetExp/Source/MeshModelingToolsExp/Private/HoleFillTool.cpp
--- a/Engine/Plugins/Experimental/MeshModelingToolsetExp/Source/MeshModelingToolsExp/Private/HoleFillTool.cpp
+++ b/Engine/Plugins/Experimental/MeshModelingToolsetExp/Source/MeshModelingToolsExp/Private/HoleFillTool.cpp
@@ -16,6 +16,7 @@
 #include "TargetInterfaces/MaterialProvider.h"
 #include "TargetInterfaces/PrimitiveComponentBackedTarget.h"
 #include "ModelingToolTargetUtil.h"
+#include "ToolTargetManager.h"
 
 #include UE_INLINE_GENERATED_CPP_BY_NAME(HoleFillTool)
 
@@ -32,6 +33,14 @@ USingleSelectionMeshEditingTool* UHoleFillToolBuilder::CreateNewTool(const FTool
 	return NewObject<UHoleFillTool>(SceneState.ToolManager);
 }
 
+bool UHoleFillToolBuilder::CanBuildTool(const FToolBuilderState& SceneState) const
+{
+	// We're disallowing volumes because they shouldn't have holes to fill intrinsically.
+	return USingleSelectionMeshEditingToolBuilder::CanBuildTool(SceneState) &&
+		SceneState.TargetManager->CountSelectedAndTargetableWithPredicate(SceneState, GetTargetRequirements(),
+			[](UActorComponent& Component) { return !ToolBuilderUtil::IsVolume(Component); }) >= 1;
+}
+
 /*
  * Tool properties
  */
diff --git a/Engine/Plugins/Experimental/MeshModelingToolsetExp/Source/MeshModelingToolsExp/Private/MeshAttributePaintTool.cpp b/Engine/Plugins/Experimental/MeshModelingToolsetExp/Source/MeshModelingToolsExp/Private/MeshAttributePaintTool.cpp
--- a/Engine/Plugins/Experimental/MeshModelingToolsetExp/Source/MeshModelingToolsExp/Private/MeshAttributePaintTool.cpp
+++ b/Engine/Plugins/Experimental/MeshModelingToolsetExp/Source/MeshModelingToolsExp/Private/MeshAttributePaintTool.cpp
@@ -12,6 +12,7 @@
 #include "TargetInterfaces/PrimitiveComponentBackedTarget.h"
 #include "ModelingToolTargetUtil.h"
 #include "DynamicMesh/NonManifoldMappingSupport.h"
+#include "ToolTargetManager.h"
 
 #include UE_INLINE_GENERATED_CPP_BY_NAME(MeshAttributePaintTool)
 
@@ -165,6 +166,13 @@ UMeshSurfacePointTool* UMeshAttributePaintToolBuilder::CreateNewTool(const FTool
 	return SelectionTool;
 }
 
+bool UMeshAttributePaintToolBuilder::CanBuildTool(const FToolBuilderState& SceneState) const
+{
+	return UMeshSurfacePointMeshEditingToolBuilder::CanBuildTool(SceneState) &&
+		SceneState.TargetManager->CountSelectedAndTargetableWithPredicate(SceneState, GetTargetRequirements(),
+			[](UActorComponent& Component) { return !ToolBuilderUtil::IsVolume(Component); }) >= 1;
+}
+
 
 
 
diff --git a/Engine/Plugins/Experimental/MeshModelingToolsetExp/Source/MeshModelingToolsExp/Private/MeshGroupPaintTool.cpp b/Engine/Plugins/Experimental/MeshModelingToolsetExp/Source/MeshModelingToolsExp/Private/MeshGroupPaintTool.cpp
--- a/Engine/Plugins/Experimental/MeshModelingToolsetExp/Source/MeshModelingToolsExp/Private/MeshGroupPaintTool.cpp
+++ b/Engine/Plugins/Experimental/MeshModelingToolsetExp/Source/MeshModelingToolsExp/Private/MeshGroupPaintTool.cpp
@@ -36,6 +36,8 @@
 #include "CanvasItem.h"
 #include "Engine/Engine.h"  // for GEngine->GetSmallFont()
 #include "SceneView.h"
+#include "ToolTargetManager.h"
+#include "ToolBuilderUtil.h"
 
 #include UE_INLINE_GENERATED_CPP_BY_NAME(MeshGroupPaintTool)
 
@@ -63,7 +65,12 @@ UMeshSurfacePointTool* UMeshGroupPaintToolBuilder::CreateNewTool(const FToolBuil
 	SculptTool->SetWorld(SceneState.World);
 	return SculptTool;
 }
-
+bool UMeshGroupPaintToolBuilder::CanBuildTool(const FToolBuilderState& SceneState) const
+{
+	return UMeshSurfacePointMeshEditingToolBuilder::CanBuildTool(SceneState) &&
+		SceneState.TargetManager->CountSelectedAndTargetableWithPredicate(SceneState, GetTargetRequirements(),
+			[](UActorComponent& Component) { return !ToolBuilderUtil::IsVolume(Component); }) >= 1;
+}
 
 /*
  * Properties
diff --git a/Engine/Plugins/Experimental/MeshModelingToolsetExp/Source/MeshModelingToolsExp/Private/RevolveBoundaryTool.cpp b/Engine/Plugins/Experimental/MeshModelingToolsetExp/Source/MeshModelingToolsExp/Private/RevolveBoundaryTool.cpp
--- a/Engine/Plugins/Experimental/MeshModelingToolsetExp/Source/MeshModelingToolsExp/Private/RevolveBoundaryTool.cpp
+++ b/Engine/Plugins/Experimental/MeshModelingToolsetExp/Source/MeshModelingToolsExp/Private/RevolveBoundaryTool.cpp
@@ -16,6 +16,7 @@
 #include "ModelingToolTargetUtil.h"
 #include "ToolSceneQueriesUtil.h"
 #include "ToolSetupUtil.h"
+#include "ToolTargetManager.h"
 
 #include "TargetInterfaces/PrimitiveComponentBackedTarget.h"
 
@@ -31,7 +32,13 @@ USingleSelectionMeshEditingTool* URevolveBoundaryToolBuilder::CreateNewTool(cons
 {
 	return NewObject<URevolveBoundaryTool>(SceneState.ToolManager);
 }
-
+bool URevolveBoundaryToolBuilder::CanBuildTool(const FToolBuilderState& SceneState) const
+{
+	// We're disallowing volumes here because volumes won't have open boundaries intrinsically.
+	return USingleSelectionMeshEditingToolBuilder::CanBuildTool(SceneState) &&
+		SceneState.TargetManager->CountSelectedAndTargetableWithPredicate(SceneState, GetTargetRequirements(),
+			[](UActorComponent& Component) { return !ToolBuilderUtil::IsVolume(Component); }) >= 1;
+}
 
 // Operator factory
 
diff --git a/Engine/Plugins/Experimental/MeshModelingToolsetExp/Source/MeshModelingToolsExp/Public/ConvertToPolygonsTool.h b/Engine/Plugins/Experimental/MeshModelingToolsetExp/Source/MeshModelingToolsExp/Public/ConvertToPolygonsTool.h
--- a/Engine/Plugins/Experimental/MeshModelingToolsetExp/Source/MeshModelingToolsExp/Public/ConvertToPolygonsTool.h
+++ b/Engine/Plugins/Experimental/MeshModelingToolsetExp/Source/MeshModelingToolsExp/Public/ConvertToPolygonsTool.h
@@ -30,8 +30,8 @@ class MESHMODELINGTOOLSEXP_API UConvertToPolygonsToolBuilder : public USingleTar
 	GENERATED_BODY()
 public:
 	virtual USingleTargetWithSelectionTool* CreateNewTool(const FToolBuilderState & SceneState) const override;
-
 	virtual bool RequiresInputSelection() const override { return false; }
+	virtual bool CanBuildTool(const FToolBuilderState& SceneState) const override;
 };
 
 
diff --git a/Engine/Plugins/Experimental/MeshModelingToolsetExp/Source/MeshModelingToolsExp/Public/EditNormalsTool.h b/Engine/Plugins/Experimental/MeshModelingToolsetExp/Source/MeshModelingToolsExp/Public/EditNormalsTool.h
--- a/Engine/Plugins/Experimental/MeshModelingToolsetExp/Source/MeshModelingToolsExp/Public/EditNormalsTool.h
+++ b/Engine/Plugins/Experimental/MeshModelingToolsetExp/Source/MeshModelingToolsExp/Public/EditNormalsTool.h
@@ -29,8 +29,8 @@ class MESHMODELINGTOOLSEXP_API UEditNormalsToolBuilder : public UMultiSelectionM
 
 public:
 	virtual UMultiSelectionMeshEditingTool* CreateNewTool(const FToolBuilderState& SceneState) const override;
-
 	virtual void InitializeNewTool(UMultiSelectionMeshEditingTool* NewTool, const FToolBuilderState& SceneState) const override;
+	virtual bool CanBuildTool(const FToolBuilderState& SceneState) const override;
 };
 
 
diff --git a/Engine/Plugins/Experimental/MeshModelingToolsetExp/Source/MeshModelingToolsExp/Public/HoleFillTool.h b/Engine/Plugins/Experimental/MeshModelingToolsetExp/Source/MeshModelingToolsExp/Public/HoleFillTool.h
--- a/Engine/Plugins/Experimental/MeshModelingToolsetExp/Source/MeshModelingToolsExp/Public/HoleFillTool.h
+++ b/Engine/Plugins/Experimental/MeshModelingToolsetExp/Source/MeshModelingToolsExp/Public/HoleFillTool.h
@@ -33,6 +33,7 @@ class MESHMODELINGTOOLSEXP_API UHoleFillToolBuilder : public USingleSelectionMes
 
 public:
 	virtual USingleSelectionMeshEditingTool* CreateNewTool(const FToolBuilderState& SceneState) const override;
+	virtual bool CanBuildTool(const FToolBuilderState& SceneState) const override;
 };
 
 /*
diff --git a/Engine/Plugins/Experimental/MeshModelingToolsetExp/Source/MeshModelingToolsExp/Public/MeshAttributePaintTool.h b/Engine/Plugins/Experimental/MeshModelingToolsetExp/Source/MeshModelingToolsExp/Public/MeshAttributePaintTool.h
--- a/Engine/Plugins/Experimental/MeshModelingToolsetExp/Source/MeshModelingToolsExp/Public/MeshAttributePaintTool.h
+++ b/Engine/Plugins/Experimental/MeshModelingToolsetExp/Source/MeshModelingToolsExp/Public/MeshAttributePaintTool.h
@@ -92,6 +92,7 @@ public:
 	TUniqueFunction<TUniquePtr<FFloatAttributeColorMapper>()> ColorMapFactory;
 
 	virtual UMeshSurfacePointTool* CreateNewTool(const FToolBuilderState& SceneState) const override;
+	virtual bool CanBuildTool(const FToolBuilderState& SceneState) const override;
 };
 
 
diff --git a/Engine/Plugins/Experimental/MeshModelingToolsetExp/Source/MeshModelingToolsExp/Public/MeshGroupPaintTool.h b/Engine/Plugins/Experimental/MeshModelingToolsetExp/Source/MeshModelingToolsExp/Public/MeshGroupPaintTool.h
--- a/Engine/Plugins/Experimental/MeshModelingToolsetExp/Source/MeshModelingToolsExp/Public/MeshGroupPaintTool.h
+++ b/Engine/Plugins/Experimental/MeshModelingToolsetExp/Source/MeshModelingToolsExp/Public/MeshGroupPaintTool.h
@@ -55,6 +55,7 @@ class MESHMODELINGTOOLSEXP_API UMeshGroupPaintToolBuilder : public UMeshSurfaceP
 
 public:
 	virtual UMeshSurfacePointTool* CreateNewTool(const FToolBuilderState& SceneState) const override;
+	virtual bool CanBuildTool(const FToolBuilderState& SceneState) const override;
 };
 
 
diff --git a/Engine/Plugins/Experimental/MeshModelingToolsetExp/Source/MeshModelingToolsExp/Public/RevolveBoundaryTool.h b/Engine/Plugins/Experimental/MeshModelingToolsetExp/Source/MeshModelingToolsExp/Public/RevolveBoundaryTool.h
--- a/Engine/Plugins/Experimental/MeshModelingToolsetExp/Source/MeshModelingToolsExp/Public/RevolveBoundaryTool.h
+++ b/Engine/Plugins/Experimental/MeshModelingToolsetExp/Source/MeshModelingToolsExp/Public/RevolveBoundaryTool.h
@@ -26,6 +26,7 @@ class MESHMODELINGTOOLSEXP_API URevolveBoundaryToolBuilder : public USingleSelec
 
 public:
 	virtual USingleSelectionMeshEditingTool* CreateNewTool(const FToolBuilderState& SceneState) const override;
+	virtual bool CanBuildTool(const FToolBuilderState& SceneState) const override;
 };
 
 UCLASS()
