add initial triangulate splines tool

diff --git a/Engine/Plugins/Editor/ModelingToolsEditorMode/Source/ModelingToolsEditorMode/Private/ModelingToolsEditorMode.cpp b/Engine/Plugins/Editor/ModelingToolsEditorMode/Source/ModelingToolsEditorMode/Private/ModelingToolsEditorMode.cpp
--- a/Engine/Plugins/Editor/ModelingToolsEditorMode/Source/ModelingToolsEditorMode/Private/ModelingToolsEditorMode.cpp
+++ b/Engine/Plugins/Editor/ModelingToolsEditorMode/Source/ModelingToolsEditorMode/Private/ModelingToolsEditorMode.cpp
@@ -89,6 +89,7 @@
 #include "ConvertMeshesTool.h"
 #include "SplitMeshesTool.h"
 #include "PatternTool.h"
+#include "TriangulateSplinesTool.h"
 
 #include "Polymodeling/ExtrudeMeshSelectionTool.h"
 #include "Polymodeling/OffsetMeshSelectionTool.h"
@@ -235,10 +236,14 @@ bool UModelingToolsEditorMode::CanAutoSave() const
 
 bool UModelingToolsEditorMode::ShouldDrawWidget() const
 { 
-	// hide standard xform gizmo if we have an active tool
+	// hide standard xform gizmo if we have an active tool, unless it explicitly opts in via the IInteractiveToolEditorGizmoAPI
 	if (GetInteractiveToolsContext() != nullptr && GetToolManager()->HasAnyActiveTool())
 	{
-		return false;
+		IInteractiveToolEditorGizmoAPI* GizmoAPI = Cast<IInteractiveToolEditorGizmoAPI>(GetToolManager()->GetActiveTool(EToolSide::Left));
+		if (!GizmoAPI || !GizmoAPI->GetAllowStandardEditorGizmos())
+		{
+			return false;
+		}
 	}
 
 	// hide standard xform gizmo if we have an active selection
@@ -573,6 +578,9 @@ void UModelingToolsEditorMode::Enter()
 	auto CubeGridToolBuilder = NewObject<UCubeGridToolBuilder>();
 	RegisterTool(ToolManagerCommands.BeginCubeGridTool, TEXT("BeginCubeGridTool"), CubeGridToolBuilder);
 
+	auto TriangulateSplinesToolBuilder = NewObject<UTriangulateSplinesToolBuilder>();
+	RegisterTool(ToolManagerCommands.BeginTriangulateSplinesTool, TEXT("BeginTriangulateSplinesTool"), TriangulateSplinesToolBuilder);
+
 	//
 	// vertex deform tools
 	//
diff --git a/Engine/Plugins/Editor/ModelingToolsEditorMode/Source/ModelingToolsEditorMode/Private/ModelingToolsEditorModeToolkit.cpp b/Engine/Plugins/Editor/ModelingToolsEditorMode/Source/ModelingToolsEditorMode/Private/ModelingToolsEditorModeToolkit.cpp
--- a/Engine/Plugins/Editor/ModelingToolsEditorMode/Source/ModelingToolsEditorMode/Private/ModelingToolsEditorModeToolkit.cpp
+++ b/Engine/Plugins/Editor/ModelingToolsEditorMode/Source/ModelingToolsEditorMode/Private/ModelingToolsEditorModeToolkit.cpp
@@ -256,7 +256,8 @@ void FModelingToolsEditorModeToolkit::RegisterPalettes()
 
 		Commands.BeginDrawPolygonTool,
 		Commands.BeginDrawPolyPathTool,
-		Commands.BeginDrawAndRevolveTool
+		Commands.BeginDrawAndRevolveTool,
+		Commands.BeginTriangulateSplinesTool
 	});
 	ToolkitBuilder->AddPalette(
 		MakeShareable( new FToolPalette( Commands.LoadCreateTools.ToSharedRef(), CreatePaletteItems ) ) );
diff --git a/Engine/Plugins/Editor/ModelingToolsEditorMode/Source/ModelingToolsEditorMode/Private/ModelingToolsManagerActions.cpp b/Engine/Plugins/Editor/ModelingToolsEditorMode/Source/ModelingToolsEditorMode/Private/ModelingToolsManagerActions.cpp
--- a/Engine/Plugins/Editor/ModelingToolsEditorMode/Source/ModelingToolsEditorMode/Private/ModelingToolsManagerActions.cpp
+++ b/Engine/Plugins/Editor/ModelingToolsEditorMode/Source/ModelingToolsEditorMode/Private/ModelingToolsManagerActions.cpp
@@ -82,6 +82,7 @@ void FModelingToolsManagerCommands::RegisterCommands()
 		REGISTER_MODELING_TOOL_COMMAND(BeginRevolveBoundaryTool, "BdryRev", "Revolve Mesh boundary loops to create new objects");
 		REGISTER_MODELING_TOOL_COMMAND(BeginCombineMeshesTool, "Merge", "Merge multiple Meshes to create new objects");
 		REGISTER_MODELING_TOOL_COMMAND(BeginDuplicateMeshesTool, "Dupe", "Duplicate single Meshes to create new objects");
+		REGISTER_MODELING_TOOL_COMMAND(BeginTriangulateSplinesTool, "MshSpln", "Triangulate the Spline Components of selected actors to create new objects");
 	}
 	else
 	{
@@ -91,6 +92,7 @@ void FModelingToolsManagerCommands::RegisterCommands()
 		REGISTER_MODELING_TOOL_COMMAND(BeginRevolveBoundaryTool, "Revolve Boundary", "Revolve Mesh boundary loops to create new objects");
 		REGISTER_MODELING_TOOL_COMMAND(BeginCombineMeshesTool, "Merge", "Merge multiple Meshes to create new objects");
 		REGISTER_MODELING_TOOL_COMMAND(BeginDuplicateMeshesTool, "Duplicate", "Duplicate single Meshes to create new objects");
+		REGISTER_MODELING_TOOL_COMMAND(BeginTriangulateSplinesTool, "Mesh Splines", "Triangulate the Spline Components of selected actors to create new objects");
 	}
 
 	// PolyModel
diff --git a/Engine/Plugins/Editor/ModelingToolsEditorMode/Source/ModelingToolsEditorMode/Public/ModelingToolsManagerActions.h b/Engine/Plugins/Editor/ModelingToolsEditorMode/Source/ModelingToolsEditorMode/Public/ModelingToolsManagerActions.h
--- a/Engine/Plugins/Editor/ModelingToolsEditorMode/Source/ModelingToolsEditorMode/Public/ModelingToolsManagerActions.h
+++ b/Engine/Plugins/Editor/ModelingToolsEditorMode/Source/ModelingToolsEditorMode/Public/ModelingToolsManagerActions.h
@@ -47,6 +47,7 @@ public:
 	TSharedPtr<FUICommandInfo> BeginDrawPolyPathTool;
 	TSharedPtr<FUICommandInfo> BeginDrawAndRevolveTool;
 	TSharedPtr<FUICommandInfo> BeginShapeSprayTool;
+	TSharedPtr<FUICommandInfo> BeginTriangulateSplinesTool;
 
 	TSharedPtr<FUICommandInfo> BeginSculptMeshTool;
 	TSharedPtr<FUICommandInfo> BeginCubeGridTool;
diff --git a/Engine/Plugins/Experimental/MeshModelingToolsetExp/Source/MeshModelingToolsExp/Private/TriangulateSplinesTool.cpp b/Engine/Plugins/Experimental/MeshModelingToolsetExp/Source/MeshModelingToolsExp/Private/TriangulateSplinesTool.cpp
new file mode 100644
--- /dev/null
+++ b/Engine/Plugins/Experimental/MeshModelingToolsetExp/Source/MeshModelingToolsExp/Private/TriangulateSplinesTool.cpp
@@ -0,0 +1,262 @@
+// Copyright Epic Games, Inc. All Rights Reserved.
+
+#include "TriangulateSplinesTool.h"
+
+#include "CurveOps/TriangulateCurvesOp.h"
+#include "InteractiveToolManager.h"
+#include "ToolBuilderUtil.h"
+#include "ToolSetupUtil.h"
+#include "Selection/ToolSelectionUtil.h"
+#include "DynamicMesh/MeshTransforms.h"
+
+#include "Components/SplineComponent.h"
+
+#include UE_INLINE_GENERATED_CPP_BY_NAME(TriangulateSplinesTool)
+
+#define LOCTEXT_NAMESPACE "UTriangulateSplinesTool"
+
+using namespace UE::Geometry;
+
+
+
+void UTriangulateSplinesTool::Setup()
+{
+	UInteractiveTool::Setup();
+
+	// initialize our properties
+
+	TriangulateProperties = NewObject<UTriangulateSplinesToolProperties>(this);
+	TriangulateProperties->RestoreProperties(this);
+	AddToolPropertySource(TriangulateProperties);
+
+	OutputTypeProperties = NewObject<UCreateMeshObjectTypeProperties>(this);
+	OutputTypeProperties->InitializeDefault();
+	OutputTypeProperties->RestoreProperties(this);
+	OutputTypeProperties->WatchProperty(OutputTypeProperties->OutputType, [this](FString) { OutputTypeProperties->UpdatePropertyVisibility(); });
+	AddToolPropertySource(OutputTypeProperties);
+
+	SetToolDisplayName(LOCTEXT("TriangulateSplinesToolName", "Triangulate Splines"));
+	GetToolManager()->DisplayMessage(
+		LOCTEXT("TriangulateSplinesToolToolDescription", "Triangulate the shapes of the selected splines."),
+		EToolMessageLevel::UserNotification);
+
+	Preview = NewObject<UMeshOpPreviewWithBackgroundCompute>(this);
+	Preview->Setup(GetTargetWorld(), this);
+	ToolSetupUtil::ApplyRenderingConfigurationToPreview(Preview->PreviewMesh, nullptr);
+
+	Preview->OnMeshUpdated.AddLambda(
+		[this](const UMeshOpPreviewWithBackgroundCompute* UpdatedPreview)
+		{
+			UpdateAcceptWarnings(UpdatedPreview->HaveEmptyResult() ? EAcceptWarning::EmptyForbidden : EAcceptWarning::NoWarning);
+		}
+	);
+
+	PollSplineUpdates();
+}
+
+void UTriangulateSplinesTool::PollSplineUpdates()
+{
+	bool bSplinesUpdated = false;
+	int32 SplineIdx = 0;
+	TArray<TObjectPtr<USplineComponent>, TInlineAllocator<8>> SplineComponents;
+	for (int32 ActorIdx = 0; ActorIdx < ActorsWithSplines.Num(); ++ActorIdx)
+	{
+		if (AActor* Actor = ActorsWithSplines[ActorIdx].Get())
+		{
+			SplineComponents.Reset();
+			Actor->GetComponents<USplineComponent, TInlineAllocator<8>>(SplineComponents);
+			for (USplineComponent* SplineComponent : SplineComponents)
+			{
+				int32 Version = SplineComponent->SplineCurves.Version;
+				FTransform Transform = SplineComponent->GetComponentTransform();
+				if (SplineIdx >= LastSplineVersions.Num())
+				{
+					bSplinesUpdated = true;
+					LastSplineVersions.Add(Version);
+					LastSplineTransforms.Add(Transform);
+				}
+				else if (LastSplineVersions[SplineIdx] != Version || !LastSplineTransforms[SplineIdx].Equals(Transform))
+				{
+					bSplinesUpdated = true;
+				}
+				LastSplineVersions[SplineIdx] = Version;
+				LastSplineTransforms[SplineIdx] = Transform;
+				++SplineIdx;
+			}
+		}
+	}
+	if (LastSplineVersions.Num() != SplineIdx)
+	{
+		LastSplineVersions.SetNum(SplineIdx);
+		LastSplineTransforms.SetNum(SplineIdx);
+		bSplinesUpdated = true;
+	}
+	if (bSplinesUpdated)
+	{
+		Preview->InvalidateResult();
+	}
+}
+
+void UTriangulateSplinesTool::OnTick(float DeltaTime)
+{
+	PollSplineUpdates();
+
+	Preview->Tick(DeltaTime);
+}
+
+void UTriangulateSplinesTool::SetWorld(UWorld* World)
+{
+	TargetWorld = World;
+}
+
+UWorld* UTriangulateSplinesTool::GetTargetWorld()
+{
+	return TargetWorld.Get();
+}
+
+void UTriangulateSplinesTool::GenerateAsset(const FDynamicMeshOpResult& OpResult)
+{
+	if (OpResult.Mesh.Get() == nullptr) return;
+
+	FTransform3d NewTransform;
+	if (ActorsWithSplines.Num() == 1 && ActorsWithSplines[0].IsValid()) // in the single-actor case, shove the result back into the original actor transform space
+	{
+		FTransform3d ActorToWorld = (FTransform3d)ActorsWithSplines[0]->GetTransform();
+		MeshTransforms::ApplyTransform(*OpResult.Mesh, OpResult.Transform, true);
+		MeshTransforms::ApplyTransformInverse(*OpResult.Mesh, ActorToWorld, true);
+		NewTransform = ActorToWorld;
+	}
+	else // in the multi-selection case, center the pivot for the combined result
+	{
+		FVector3d Center = OpResult.Mesh->GetBounds().Center();
+		double Rescale = OpResult.Transform.GetScale().X;
+		FTransform3d LocalTransform(-Center * Rescale);
+		LocalTransform.SetScale3D(FVector3d(Rescale, Rescale, Rescale));
+		MeshTransforms::ApplyTransform(*OpResult.Mesh, LocalTransform, true);
+		NewTransform = OpResult.Transform;
+		NewTransform.SetScale3D(FVector3d::One());
+		NewTransform.SetTranslation(NewTransform.GetTranslation() + NewTransform.TransformVector(Center * Rescale));
+	}
+
+	FString BaseName = TEXT("Triangulation");
+
+	FCreateMeshObjectParams NewMeshObjectParams;
+	NewMeshObjectParams.TargetWorld = GetTargetWorld();
+	NewMeshObjectParams.Transform = (FTransform)NewTransform;
+	NewMeshObjectParams.BaseName = BaseName;
+	NewMeshObjectParams.Materials = Preview->StandardMaterials;
+	NewMeshObjectParams.SetMesh(OpResult.Mesh.Get());
+	OutputTypeProperties->ConfigureCreateMeshObjectParams(NewMeshObjectParams);
+	FCreateMeshObjectResult Result = UE::Modeling::CreateMeshObject(GetToolManager(), MoveTemp(NewMeshObjectParams));
+	if (Result.IsOK() && Result.NewActor != nullptr)
+	{
+		ToolSelectionUtil::SetNewActorSelection(GetToolManager(), Result.NewActor);
+	}
+}
+
+void UTriangulateSplinesTool::OnPropertyModified(UObject* PropertySet, FProperty* Property)
+{
+	if (Property &&
+		(PropertySet == OutputTypeProperties
+			))
+	{
+		// nothing
+	}
+	else
+	{
+		Preview->InvalidateResult();
+	}
+}
+
+void UTriangulateSplinesTool::Shutdown(EToolShutdownType ShutdownType)
+{
+	OutputTypeProperties->SaveProperties(this);
+	TriangulateProperties->SaveProperties(this);
+
+	FDynamicMeshOpResult Result = Preview->Shutdown();
+
+	if (ShutdownType == EToolShutdownType::Accept)
+	{
+		GetToolManager()->BeginUndoTransaction(LOCTEXT("TriangulateSplinesAction", "Spline Triangulation"));
+
+		// Generate the result asset
+		GenerateAsset(Result);
+
+		GetToolManager()->EndUndoTransaction();
+	}
+
+	TargetWorld = nullptr;
+	Super::Shutdown(ShutdownType);
+}
+
+bool UTriangulateSplinesTool::CanAccept() const
+{
+	return Preview->HaveValidNonEmptyResult();
+}
+
+TUniquePtr<FDynamicMeshOperator> UTriangulateSplinesTool::MakeNewOperator()
+{
+	TUniquePtr<FTriangulateCurvesOp> Op = MakeUnique<FTriangulateCurvesOp>();
+
+	Op->Thickness = TriangulateProperties->Thickness;
+
+	TArray<TObjectPtr<USplineComponent>, TInlineAllocator<8>> SplineComponents;
+	for (int32 ActorIdx = 0; ActorIdx < ActorsWithSplines.Num(); ++ActorIdx)
+	{
+		if (AActor* Actor = ActorsWithSplines[ActorIdx].Get())
+		{
+			SplineComponents.Reset();
+			Actor->GetComponents<USplineComponent, TInlineAllocator<8>>(SplineComponents);
+			for (USplineComponent* SplineComponent : SplineComponents)
+			{
+				Op->AddSpline(SplineComponent, TriangulateProperties->ErrorTolerance);
+			}
+		}
+	}
+
+	return Op;
+}
+
+/// Tool builder
+
+bool UTriangulateSplinesToolBuilder::CanBuildTool(const FToolBuilderState& SceneState) const
+{
+	int32 NumSplines = ToolBuilderUtil::CountComponents(SceneState, [&](UActorComponent* Object) -> bool
+	{
+		return Object->IsA<USplineComponent>();
+	});
+	return NumSplines > 0;
+}
+
+UInteractiveTool* UTriangulateSplinesToolBuilder::BuildTool(const FToolBuilderState& SceneState) const
+{
+	UTriangulateSplinesTool* NewTool = NewObject<UTriangulateSplinesTool>(SceneState.ToolManager);
+	InitializeNewTool(NewTool, SceneState);
+	return NewTool;
+}
+
+void UTriangulateSplinesToolBuilder::InitializeNewTool(UTriangulateSplinesTool* NewTool, const FToolBuilderState& SceneState) const
+{
+	TArray<UActorComponent*> Components = ToolBuilderUtil::FindAllComponents(SceneState, [&](UActorComponent* Object)
+	{
+		return Object->IsA<USplineComponent>();
+	});
+	TArray<TWeakObjectPtr<AActor>> ActorsWithSplines;
+	TSet<AActor*> FoundActors;
+	for (UActorComponent* Component : Components)
+	{
+		AActor* ActorWithSpline = Component->GetOwner();
+		if (!FoundActors.Contains(ActorWithSpline))
+		{
+			FoundActors.Add(ActorWithSpline);
+			ActorsWithSplines.Add(ActorWithSpline);
+		}
+	}
+	NewTool->SetSplineActors(MoveTemp(ActorsWithSplines));
+	NewTool->SetWorld(SceneState.World);
+}
+
+
+
+#undef LOCTEXT_NAMESPACE
+
diff --git a/Engine/Plugins/Experimental/MeshModelingToolsetExp/Source/MeshModelingToolsExp/Public/TriangulateSplinesTool.h b/Engine/Plugins/Experimental/MeshModelingToolsetExp/Source/MeshModelingToolsExp/Public/TriangulateSplinesTool.h
new file mode 100644
--- /dev/null
+++ b/Engine/Plugins/Experimental/MeshModelingToolsetExp/Source/MeshModelingToolsExp/Public/TriangulateSplinesTool.h
@@ -0,0 +1,138 @@
+// Copyright Epic Games, Inc. All Rights Reserved.
+
+#pragma once
+
+#include "CoreMinimal.h"
+#include "UObject/NoExportTypes.h"
+
+#include "InteractiveTool.h"
+#include "InteractiveToolBuilder.h"
+#include "InteractiveToolManager.h"
+#include "MeshOpPreviewHelpers.h"
+#include "InteractiveToolQueryInterfaces.h"
+#include "PropertySets/CreateMeshObjectTypeProperties.h"
+
+
+#include "TriangulateSplinesTool.generated.h"
+
+class USplineComponent;
+
+
+
+/**
+ * Parameters for controlling the spline triangulation
+ */
+UCLASS()
+class MESHMODELINGTOOLSEXP_API UTriangulateSplinesToolProperties : public UInteractiveToolPropertySet
+{
+	GENERATED_BODY()
+public:
+
+	// How far to allow the triangulation boundary can deviate from the spline curve before we add more vertices
+	UPROPERTY(EditAnywhere, Category = Spline, meta = (ClampMin = 0.001))
+	double ErrorTolerance = 1.0;
+
+	// If > 0, Extrude the triangulation by this amount
+	UPROPERTY(EditAnywhere, Category = Spline, meta = (ClampMin = 0.0))
+	double Thickness = 0.0;
+	
+};
+
+/**
+ * Tool to create a mesh from a set of selected Spline Components
+ */
+UCLASS()
+class MESHMODELINGTOOLSEXP_API UTriangulateSplinesTool : public UInteractiveTool, public IInteractiveToolEditorGizmoAPI, public UE::Geometry::IDynamicMeshOperatorFactory
+{
+	GENERATED_BODY()
+
+public:
+
+	UTriangulateSplinesTool() = default;
+
+	// IInteractiveToolEditorGizmoAPI -- allow editor gizmo so users can live-edit the splines
+
+	virtual bool GetAllowStandardEditorGizmos() override
+	{
+		return true;
+	}
+
+	//
+	// InteractiveTool API - generally does not need to be modified by subclasses
+	//
+
+	virtual void Setup() override;
+	virtual void Shutdown(EToolShutdownType ShutdownType) override;
+
+	virtual void OnTick(float DeltaTime) override;
+
+	virtual bool HasCancel() const override { return true; }
+	virtual bool HasAccept() const override { return true; }
+	virtual bool CanAccept() const override;
+
+	virtual void OnPropertyModified(UObject* PropertySet, FProperty* Property) override;
+
+	void SetSplineActors(TArray<TWeakObjectPtr<AActor>> InSplineActors)
+	{
+		ActorsWithSplines = MoveTemp(InSplineActors);
+	}
+
+	virtual void SetWorld(UWorld* World);
+	virtual UWorld* GetTargetWorld();
+
+	// IDynamicMeshOperatorFactory API
+	virtual TUniquePtr<UE::Geometry::FDynamicMeshOperator> MakeNewOperator();
+
+private:
+
+	virtual void GenerateAsset(const FDynamicMeshOpResult& OpResult);
+
+	UPROPERTY()
+	TObjectPtr<UTriangulateSplinesToolProperties> TriangulateProperties;
+
+	UPROPERTY()
+	TObjectPtr<UCreateMeshObjectTypeProperties> OutputTypeProperties;
+
+	UPROPERTY()
+	TObjectPtr<UMeshOpPreviewWithBackgroundCompute> Preview;
+
+	UPROPERTY()
+	TWeakObjectPtr<UWorld> TargetWorld = nullptr;
+
+	// Note: We track actors instead of the USplineComponents here because the USplineComponents objects are often deleted / swapped for identical but new objects
+	UPROPERTY()
+	TArray<TWeakObjectPtr<AActor>> ActorsWithSplines;
+
+private:
+
+	// Helper to track the splines we are triangulating, so we can re-triangulate when they are moved or changed
+	void PollSplineUpdates();
+	// Track the spline 'Version' integer, which is incremented when splines are changed
+	TArray<uint32> LastSplineVersions;
+	// Track the spline component's transform (to world space)
+	TArray<FTransform> LastSplineTransforms;
+};
+
+
+
+/**
+ * Base Tool Builder for tools that operate on a selection of Spline Components
+ */
+UCLASS(Transient)
+class MESHMODELINGTOOLSEXP_API UTriangulateSplinesToolBuilder : public UInteractiveToolBuilder
+{
+	GENERATED_BODY()
+
+public:
+	/** @return true if spline component sources can be found in the active selection */
+	virtual bool CanBuildTool(const FToolBuilderState& SceneState) const override;
+
+	/** @return new Tool instance initialized with selected spline source(s) */
+	virtual UInteractiveTool* BuildTool(const FToolBuilderState& SceneState) const override;
+
+	/** Called by BuildTool to configure the Tool with the input spline source(s) based on the SceneState */
+	virtual void InitializeNewTool(UTriangulateSplinesTool* Tool, const FToolBuilderState& SceneState) const;
+};
+
+
+
diff --git a/Engine/Plugins/Runtime/MeshModelingToolset/Source/ModelingOperators/Private/CurveOps/TriangulateCurvesOp.cpp b/Engine/Plugins/Runtime/MeshModelingToolset/Source/ModelingOperators/Private/CurveOps/TriangulateCurvesOp.cpp
new file mode 100644
--- /dev/null
+++ b/Engine/Plugins/Runtime/MeshModelingToolset/Source/ModelingOperators/Private/CurveOps/TriangulateCurvesOp.cpp
@@ -0,0 +1,100 @@
+// Copyright Epic Games, Inc. All Rights Reserved.
+
+#include "CurveOps/TriangulateCurvesOp.h"
+
+#include "Util/ProgressCancel.h"
+#include "BoxTypes.h"
+#include "CompGeom/PolygonTriangulation.h"
+#include "Operations/ExtrudeMesh.h"
+#include "DynamicMesh/MeshNormals.h"
+
+#include "Components/SplineComponent.h"
+
+using namespace UE::Geometry;
+
+
+void FTriangulateCurvesOp::AddSpline(USplineComponent* Spline, double ErrorTolerance)
+{
+	FCurvePath& Path = Paths.Emplace_GetRef();
+	Spline->ConvertSplineToPolyLine(ESplineCoordinateSpace::Type::World, ErrorTolerance * ErrorTolerance, Path.Vertices);
+	Path.bClosed = Spline->IsClosedLoop();
+	if (Paths.Num() == 1)
+	{
+		FirstPathTransform = Spline->GetComponentTransform();
+	}
+}
+
+void FTriangulateCurvesOp::CalculateResult(FProgressCancel* Progress)
+{
+	if (Progress && Progress->Cancelled())
+	{
+		return;
+	}
+
+	ResultMesh->EnableAttributes();
+	ResultMesh->EnableTriangleGroups();
+	ResultMesh->EnableVertexNormals(FVector3f::ZAxisVector);
+
+	SetResultTransform(FTransformSRT3d(FirstPathTransform));
+
+	FAxisAlignedBox3d InputBounds;
+	for (const FCurvePath& Path : Paths)
+	{
+		InputBounds.Contain(Path.Vertices);
+	}
+	double UseUVScaleFactor = UVScaleFactor / InputBounds.MaxDim();
+
+	auto AppendTriangles = [this, UseUVScaleFactor](FDynamicMesh3& Mesh, const TArray<FVector3d>& Vertices, const TArray<FIndex3i>& Tris, int32 GroupID, FVector3d PlaneOrigin, FVector3d Normal) -> void
+	{
+		if (Tris.IsEmpty())
+		{
+			return;
+		}
+
+		FFrame3d ProjectionFrame(ResultTransform.GetTranslation(), Normal);
+		FVector3f LocalNormal = (FVector3f)ResultTransform.InverseTransformNormal(Normal);
+
+		checkSlow(Mesh.IsCompact());
+		int32 VertStart = Mesh.MaxVertexID();
+		int32 TriStart = Mesh.MaxTriangleID();
+		for (const FVector3d V : Vertices)
+		{
+			int32 VID = Mesh.AppendVertex(ResultTransform.InverseTransformPosition(V));
+			Mesh.SetVertexNormal(VID, LocalNormal);
+			int32 NormalEID = Mesh.Attributes()->PrimaryNormals()->AppendElement(LocalNormal);
+			int32 UVEID = Mesh.Attributes()->PrimaryUV()->AppendElement(FVector2f(ProjectionFrame.ToPlaneUV(V) * UseUVScaleFactor));
+			// since we always add vertices and overlay elements at the same time, the IDs must match up
+			checkSlow(VID == NormalEID && NormalEID == UVEID);
+		}
+		for (const FIndex3i& T : Tris)
+		{
+			FIndex3i NewTri(T.A + VertStart, T.B + VertStart, T.C + VertStart);
+			int32 NewTID = Mesh.AppendTriangle(NewTri, GroupID);
+			Mesh.Attributes()->PrimaryNormals()->SetTriangle(NewTID, NewTri);
+			Mesh.Attributes()->PrimaryUV()->SetTriangle(NewTID, NewTri);
+		}
+	};
+
+	for (int32 PathIdx = 0; PathIdx < Paths.Num(); ++PathIdx)
+	{
+		FVector3d Normal, PlanePoint;
+		PolygonTriangulation::ComputePolygonPlane<double>(Paths[PathIdx].Vertices, Normal, PlanePoint);
+		TArray<FIndex3i> Tris;
+		PolygonTriangulation::TriangulateSimplePolygon<double>(Paths[PathIdx].Vertices, Tris, false);
+		AppendTriangles(*ResultMesh, Paths[PathIdx].Vertices, Tris, PathIdx, PlanePoint, Normal);
+	}
+
+	if (Thickness > 0)
+	{
+		FExtrudeMesh ExtrudeMesh(ResultMesh.Get());
+		ExtrudeMesh.DefaultExtrudeDistance = Thickness;
+		ExtrudeMesh.UVScaleFactor = UseUVScaleFactor;
+		ExtrudeMesh.Apply();
+	}
+
+	if (Progress && Progress->Cancelled())
+	{
+		return;
+	}
+}
+
diff --git a/Engine/Plugins/Runtime/MeshModelingToolset/Source/ModelingOperators/Public/CurveOps/TriangulateCurvesOp.h b/Engine/Plugins/Runtime/MeshModelingToolset/Source/ModelingOperators/Public/CurveOps/TriangulateCurvesOp.h
new file mode 100644
--- /dev/null
+++ b/Engine/Plugins/Runtime/MeshModelingToolset/Source/ModelingOperators/Public/CurveOps/TriangulateCurvesOp.h
@@ -0,0 +1,63 @@
+// Copyright Epic Games, Inc. All Rights Reserved.
+
+#pragma once
+
+#include "CoreMinimal.h"
+#include "ModelingOperators.h"
+
+class USplineComponent;
+
+
+namespace UE {
+namespace Geometry {
+
+class FDynamicMesh3;
+
+/**
+ * FTriangulateCurvesOp triangulates polygons/paths generated from USplineComponent inputs.
+ */
+class MODELINGOPERATORS_API FTriangulateCurvesOp : public FDynamicMeshOperator
+{
+public:
+	virtual ~FTriangulateCurvesOp() {}
+
+	//
+	// Inputs
+	//
+
+	void AddSpline(USplineComponent* Spline, double ErrorTolerance);
+
+	//
+	// Parameters
+	//
+
+	// TODO: Add more options, e.g. for flattening the input curves, setting triangulation method, and additional processing of the curves (e.g., union, offset, etc)
+
+	// Scaling applied to the default UV values
+	double UVScaleFactor = 1.0;
+
+	// If > 0, thicken the result mesh to make a solid
+	double Thickness = 0.0;
+
+	//
+	// FDynamicMeshOperator interface 
+	//
+	virtual void CalculateResult(FProgressCancel* Progress) override;
+
+private:
+
+	struct FCurvePath
+	{
+		bool bClosed = false;
+		TArray<FVector3d> Vertices;
+	};
+
+	// Paths for all splines, in world space
+	TArray<FCurvePath> Paths;
+
+	// Local to World transform of the first path
+	FTransform FirstPathTransform;
+
+};
+
+}} // end UE::Geometry
diff --git a/Engine/Source/Runtime/InteractiveToolsFramework/Public/InteractiveToolQueryInterfaces.h b/Engine/Source/Runtime/InteractiveToolsFramework/Public/InteractiveToolQueryInterfaces.h
--- a/Engine/Source/Runtime/InteractiveToolsFramework/Public/InteractiveToolQueryInterfaces.h
+++ b/Engine/Source/Runtime/InteractiveToolsFramework/Public/InteractiveToolQueryInterfaces.h
@@ -151,3 +151,28 @@ class IInteractiveToolExclusiveToolAPI
 
 
 
+// UInterface for IInteractiveToolEditorGizmoAPI
+UINTERFACE(MinimalAPI)
+class UInteractiveToolEditorGizmoAPI : public UInterface
+{
+	GENERATED_BODY()
+};
+
+/**
+ * IInteractiveToolEditorGizmoAPI provides an API to indicate whether
+ * the standard editor gizmos can be enabled while this tool is active.
+ */
+class IInteractiveToolEditorGizmoAPI
+{
+	GENERATED_BODY()
+public:
+
+	/**
+	 * @return true if the tool implementing this API allows the editor gizmos to be enabled while the tool is active
+	 */
+	virtual bool GetAllowStandardEditorGizmos() { return false; }
+};
+
+
+
+
