cis fix non-unity build for TriangulateSplinesTool.cpp

diff --git a/Engine/Plugins/Experimental/MeshModelingToolsetExp/Source/MeshModelingToolsExp/Private/TriangulateSplinesTool.cpp b/Engine/Plugins/Experimental/MeshModelingToolsetExp/Source/MeshModelingToolsExp/Private/TriangulateSplinesTool.cpp
--- a/Engine/Plugins/Experimental/MeshModelingToolsetExp/Source/MeshModelingToolsExp/Private/TriangulateSplinesTool.cpp
+++ b/Engine/Plugins/Experimental/MeshModelingToolsetExp/Source/MeshModelingToolsExp/Private/TriangulateSplinesTool.cpp
@@ -9,6 +9,8 @@
 #include "Selection/ToolSelectionUtil.h"
 #include "DynamicMesh/MeshTransforms.h"
 
+#include "Engine/World.h"
+
 #include "Components/SplineComponent.h"
 
 #include UE_INLINE_GENERATED_CPP_BY_NAME(TriangulateSplinesTool)
diff --git a/Engine/Plugins/Experimental/MeshModelingToolsetExp/Source/MeshModelingToolsExp/Public/TriangulateSplinesTool.h b/Engine/Plugins/Experimental/MeshModelingToolsetExp/Source/MeshModelingToolsExp/Public/TriangulateSplinesTool.h
--- a/Engine/Plugins/Experimental/MeshModelingToolsetExp/Source/MeshModelingToolsExp/Public/TriangulateSplinesTool.h
+++ b/Engine/Plugins/Experimental/MeshModelingToolsetExp/Source/MeshModelingToolsExp/Public/TriangulateSplinesTool.h
@@ -16,7 +16,7 @@
 #include "TriangulateSplinesTool.generated.h"
 
 class USplineComponent;
-
+class UWorld;
 
 
 /**
