MeshModelingTools: Add the "Copy From Inputs" option in Mesh to Collision tool transfer existing simple geometry over to the target.

diff --git a/Engine/Plugins/Experimental/MeshModelingToolsetExp/Source/MeshModelingToolsExp/Private/Physics/SetCollisionGeometryTool.cpp b/Engine/Plugins/Experimental/MeshModelingToolsetExp/Source/MeshModelingToolsExp/Private/Physics/SetCollisionGeometryTool.cpp
--- a/Engine/Plugins/Experimental/MeshModelingToolsetExp/Source/MeshModelingToolsExp/Private/Physics/SetCollisionGeometryTool.cpp
+++ b/Engine/Plugins/Experimental/MeshModelingToolsetExp/Source/MeshModelingToolsExp/Private/Physics/SetCollisionGeometryTool.cpp
@@ -45,6 +45,9 @@ class FPhysicsCollectionOp : public TGenericDataOperator<FPhysicsDataCollection>
 public:
 
 	TSharedPtr<FPhysicsDataCollection, ESPMode::ThreadSafe> InitialCollision;
+	TSharedPtr<TArray<FPhysicsDataCollection>, ESPMode::ThreadSafe> OtherInputsCollision;
+	FTransformSequence3d TargetInverseTransform;
+	TSharedPtr<TArray<FTransform3d>, ESPMode::ThreadSafe> OtherInputsTransforms; // null if not using world space
 
 	TUniquePtr<FMeshSimpleShapeApproximation> UseShapeGenerator;
 	// Note: UseShapeGenerator holds raw pointers to these meshes, so we keep the array of shared pointers to prevent them from getting deleted while the op runs
@@ -68,15 +71,42 @@ public:
 		// calculate new collision
 		TUniquePtr<FPhysicsDataCollection> NewCollision = MakeUnique<FPhysicsDataCollection>();
 		NewCollision->InitializeFromExisting(*InitialCollision);
-		if (bAppendToExisting || ComputeType == ECollisionGeometryType::KeepExisting)
+		if (bAppendToExisting)
 		{
 			NewCollision->CopyGeometryFromExisting(*InitialCollision);
 		}
 
 		switch (ComputeType)
 		{
-		case ECollisionGeometryType::KeepExisting:
-		case ECollisionGeometryType::None:
+		case ECollisionGeometryType::Empty:
+			break;
+		case ECollisionGeometryType::CopyFromInputs:
+			if (ensure(OtherInputsCollision))
+			{
+				if (OtherInputsTransforms)
+				{
+					for (int32 i = 0; i < OtherInputsCollision->Num(); ++i)
+					{
+						const FPhysicsDataCollection& Collision = (*OtherInputsCollision)[i];
+						if (!ensure(i < OtherInputsTransforms->Num()))
+						{
+							break;
+						}
+
+						FTransformSequence3d TransformStack;
+						TransformStack.Append((*OtherInputsTransforms)[i]);
+						TransformStack.Append(TargetInverseTransform);
+						NewCollision->Geometry.Append(Collision.Geometry, TransformStack);
+					}
+				}
+				else
+				{
+					for (const FPhysicsDataCollection& Collision : *OtherInputsCollision)
+					{
+						NewCollision->Geometry.Append(Collision.Geometry);
+					}
+				}
+			}
 			break;
 		case ECollisionGeometryType::AlignedBoxes:
 			UseShapeGenerator->Generate_AlignedBoxes(NewCollision->Geometry);
@@ -166,7 +196,6 @@ UMultiSelectionMeshEditingTool* USetCollisionGeometryToolBuilder::CreateNewTool(
 	return NewObject<USetCollisionGeometryTool>(SceneState.ToolManager);
 }
 
-
 void USetCollisionGeometryTool::Setup()
 {
 	UInteractiveTool::Setup();
@@ -186,17 +215,45 @@ void USetCollisionGeometryTool::Setup()
 		}
 	}
 
+	UToolTarget* CollisionTarget = Targets[Targets.Num() - 1];
+	OrigTargetTransform = (FTransform)UE::ToolTarget::GetLocalToWorldTransform(CollisionTarget);
+	
+	// We need an inverse transform for copying simple collision in CopyFromInputs mode. We have
+	// to separate out the scale since we can't represent the inverse properly just in SRT form.
+	FTransform OrigTransformRT = OrigTargetTransform;
+	OrigTransformRT.SetScale3D(FVector3d::One());
+
+	TargetInverseTransform.Append(OrigTransformRT.Inverse());
+	TargetInverseTransform.Append(FTransformSRT3d(FQuaterniond::Identity(), FVector3d::Zero(),
+		OrigTargetTransform.GetScale3D().Reciprocal()));
+
+	// The "OtherInputs" variables are only needed if we have multiple meshes selected. They
+	// do not include the target mesh
+	OtherInputsCollision = MakeShared<TArray<FPhysicsDataCollection>, ESPMode::ThreadSafe>();
+	OtherInputsTransforms = MakeShared<TArray<FTransform3d>, ESPMode::ThreadSafe>();
+	bool bHaveMultipleInputs = Targets.Num() > 1;
+	if (bHaveMultipleInputs)
+	{
+		// The preallocation is so we can do the initialization in the parallel for below
+		OtherInputsCollision->SetNum(SourceObjectIndices.Num());
+		OtherInputsTransforms->SetNum(SourceObjectIndices.Num());
+	}
+
 	// collect input meshes
 	InitialSourceMeshes.SetNum(SourceObjectIndices.Num());
 	ParallelFor(SourceObjectIndices.Num(), [&](int32 k)
 	{
 		InitialSourceMeshes[k] = UE::ToolTarget::GetDynamicMeshCopy(Targets[k]);
+
+		if (bHaveMultipleInputs)
+		{
+			(*OtherInputsCollision)[k].InitializeFromComponent(UE::ToolTarget::GetTargetComponent(Targets[k]), true);
+			(*OtherInputsTransforms)[k] = UE::ToolTarget::GetLocalToWorldTransform(Targets[k]);
+		}
 	});
 
-	UToolTarget* CollisionTarget = Targets[Targets.Num() - 1];
 	PreviewGeom = NewObject<UPreviewGeometry>(this);
-	FTransform PreviewTransform = (FTransform)UE::ToolTarget::GetLocalToWorldTransform(CollisionTarget);
-	OrigTargetTransform = PreviewTransform;
+	FTransform PreviewTransform = OrigTargetTransform;
 	TargetScale3D = PreviewTransform.GetScale3D();
 	PreviewTransform.SetScale3D(FVector::OneVector);
 	PreviewGeom->CreateInWorld(UE::ToolTarget::GetTargetActor(CollisionTarget)->GetWorld(), PreviewTransform);
@@ -209,8 +266,8 @@ void USetCollisionGeometryTool::Setup()
 	// create tool options
 	Settings = NewObject<USetCollisionGeometryToolProperties>(this);
 	Settings->RestoreProperties(this);
+	Settings->bUsingMultipleInputs = Targets.Num() > 1;
 	AddToolPropertySource(Settings);
-	Settings->bUseWorldSpace = (SourceObjectIndices.Num() > 1);
 	Settings->WatchProperty(Settings->InputMode, [this](ESetCollisionGeometryInputMode) { OnInputModeChanged(); });
 	Settings->WatchProperty(Settings->GeometryType, [this](ECollisionGeometryType) { InvalidateCompute(); });
 	Settings->WatchProperty(Settings->bUseWorldSpace, [this](bool) { bInputMeshesValid = false; });
@@ -269,6 +326,12 @@ TUniquePtr<UE::Geometry::TGenericDataOperator<FPhysicsDataCollection>> USetColli
 	TUniquePtr<FPhysicsCollectionOp> Op = MakeUnique<FPhysicsCollectionOp>();
 
 	Op->InitialCollision = InitialCollision;
+	Op->OtherInputsCollision = OtherInputsCollision;
+	Op->TargetInverseTransform = TargetInverseTransform;
+	if (Settings->bUseWorldSpace)
+	{
+		Op->OtherInputsTransforms = OtherInputsTransforms;
+	}
 
 	// Pick the approximator and input meshes that will be used by the op
 	TSharedPtr<UE::Geometry::FMeshSimpleShapeApproximation, ESPMode::ThreadSafe>* Approximator = nullptr;
@@ -311,6 +374,17 @@ TUniquePtr<UE::Geometry::TGenericDataOperator<FPhysicsDataCollection>> USetColli
 
 	Op->ComputeType = Settings->GeometryType;
 	Op->bAppendToExisting = Settings->bAppendToExisting;
+
+	// If we are in single-target mode, the CopyFromInputs option does not make sense because
+	// we don't want to let the user accidentally duplicate all the geometry if bAppendToExisting
+	// is true, and if it is false, the user can do the same thing with Empty and appending.
+	// TODO: add a way to disable this enum option in the UI.
+	if (Targets.Num() <= 1 && Op->ComputeType == ECollisionGeometryType::CopyFromInputs)
+	{
+		Op->ComputeType = ECollisionGeometryType::Empty;
+		Op->bAppendToExisting = true;
+	}
+
 	Op->bUseMaxCount = Settings->bEnableMaxCount;
 	Op->MaxCount = Settings->MaxCount;
 	Op->bRemoveContained = Settings->bRemoveContained;
diff --git a/Engine/Plugins/Experimental/MeshModelingToolsetExp/Source/MeshModelingToolsExp/Public/Physics/SetCollisionGeometryTool.h b/Engine/Plugins/Experimental/MeshModelingToolsetExp/Source/MeshModelingToolsExp/Public/Physics/SetCollisionGeometryTool.h
--- a/Engine/Plugins/Experimental/MeshModelingToolsetExp/Source/MeshModelingToolsExp/Public/Physics/SetCollisionGeometryTool.h
+++ b/Engine/Plugins/Experimental/MeshModelingToolsetExp/Source/MeshModelingToolsExp/Public/Physics/SetCollisionGeometryTool.h
@@ -13,6 +13,7 @@
 #include "Physics/PhysicsDataCollection.h"
 #include "PropertySets/PolygroupLayersProperties.h"
 #include "Polygroups/PolygroupSet.h"
+#include "TransformSequence.h"
 #include "ModelingOperators.h"
 #include "MeshOpPreviewHelpers.h"
 #include "SetCollisionGeometryTool.generated.h"
@@ -50,7 +51,11 @@ enum class ESetCollisionGeometryInputMode
 UENUM()
 enum class ECollisionGeometryType
 {
-	KeepExisting = 0,
+	//~ TODO: We should add a way to remove this option in single-input mode
+	// Copy the existing collision geometry shapes from the inputs to the target. With a single-selection,
+	// always does the same thing as Empty with Append To Existing set to true.
+	CopyFromInputs = 0,
+
 	AlignedBoxes = 1,
 	OrientedBoxes = 2,
 	MinimalSpheres = 3,
@@ -60,7 +65,10 @@ enum class ECollisionGeometryType
 	LevelSets = 7,
 	MinVolume = 10,
 
-	None = 11
+	// Do not produce new collision for inputs. If Append To Existing is false, this gives a way
+	// to empty the simple collision on the target. If Append To Existing is true, the existing collision
+	// is kept and can be passed through the optional filters in the tool, like removing enclosed shapes.
+	Empty = 11,
 };
 
 
@@ -86,11 +94,15 @@ public:
 	ECollisionGeometryType GeometryType = ECollisionGeometryType::AlignedBoxes;
 
 	UPROPERTY(EditAnywhere, Category = Options)
-	ESetCollisionGeometryInputMode InputMode = ESetCollisionGeometryInputMode::PerInputObject;
+	bool bAppendToExisting = false;
 
-	UPROPERTY(EditAnywhere, Category = Options)
+	UPROPERTY(EditAnywhere, Category = Options, meta = (EditCondition = "bUsingMultipleInputs", EditConditionHides, HideEditConditionToggle))
 	bool bUseWorldSpace = false;
 
+	UPROPERTY(EditAnywhere, Category = Options, meta = (
+		EditCondition = "bUsingMultipleInputs && GeometryType != ECollisionGeometryType::Empty && GeometryType != ECollisionGeometryType::CopyFromInputs"))
+	ESetCollisionGeometryInputMode InputMode = ESetCollisionGeometryInputMode::PerInputObject;
+
 	UPROPERTY(EditAnywhere, Category = Options)
 	bool bRemoveContained = true;
 
@@ -157,11 +169,12 @@ public:
 		EditConditionHides, EditCondition = "GeometryType == ECollisionGeometryType::LevelSets"))
 	int32 LevelSetResolution = 10;
 
-	UPROPERTY(EditAnywhere, Category = OutputOptions)
-	bool bAppendToExisting = false;
-
 	UPROPERTY(EditAnywhere, Category = OutputOptions)
 	ECollisionGeometryMode SetCollisionType = ECollisionGeometryMode::SimpleAndComplex;
+
+	// Set by the tool to tell the settings object whether the tool is using multiple inputs.
+	UPROPERTY()
+	bool bUsingMultipleInputs = false;
 };
 
 
@@ -273,11 +286,13 @@ protected:
 	void UpdateActiveGroupLayer();
 
 	FTransform OrigTargetTransform;
+	UE::Geometry::FTransformSequence3d TargetInverseTransform;
 	FVector TargetScale3D;
 
 	TSharedPtr<FPhysicsDataCollection, ESPMode::ThreadSafe> InitialCollision;
 	TSharedPtr<FPhysicsDataCollection, ESPMode::ThreadSafe> GeneratedCollision;
-
+	TSharedPtr<TArray<FPhysicsDataCollection>, ESPMode::ThreadSafe> OtherInputsCollision;
+	TSharedPtr<TArray<FTransform3d>, ESPMode::ThreadSafe> OtherInputsTransforms;
 
 	bool bVisualizationDirty = false;
 	void UpdateVisualization();
diff --git a/Engine/Plugins/Runtime/GeometryProcessing/Source/DynamicMesh/Private/ShapeApproximation/SimpleShapeSet3.cpp b/Engine/Plugins/Runtime/GeometryProcessing/Source/DynamicMesh/Private/ShapeApproximation/SimpleShapeSet3.cpp
--- a/Engine/Plugins/Runtime/GeometryProcessing/Source/DynamicMesh/Private/ShapeApproximation/SimpleShapeSet3.cpp
+++ b/Engine/Plugins/Runtime/GeometryProcessing/Source/DynamicMesh/Private/ShapeApproximation/SimpleShapeSet3.cpp
@@ -389,12 +389,17 @@ static void TransformSphereShape(FSphereShape3d& SphereShape, const TransformTyp
 template<typename TransformType>
 static void TransformBoxShape(FBoxShape3d& BoxShape, const TransformType& Transform)
 {
-	FVector3d CornerVec = BoxShape.Box.Frame.PointAt(BoxShape.Box.Extents) - BoxShape.Box.Frame.Origin;
 	BoxShape.Box.Frame.Transform(Transform);
-	CornerVec = Transform.TransformVector(CornerVec);
-	BoxShape.Box.Extents.X = CornerVec.Dot(BoxShape.Box.AxisX());
-	BoxShape.Box.Extents.Y = CornerVec.Dot(BoxShape.Box.AxisY());
-	BoxShape.Box.Extents.Z = CornerVec.Dot(BoxShape.Box.AxisZ());
+
+	// There isn't a perfect approach to transforming the extents in the presence of non-uniform scaling.
+	// One approach we've done in the past has been to transform the corner vector, but depending on that
+	// vector's orientation relative to scaling axis, the extents could grow/shrink in unexpected ways,
+	// or not change at all if direction happened to be orthogonal.
+	// The approach we now use is to apply the scaling in the reference frame of the box itself, because
+	// this what the editor does for box collision shapes, and it is helpful to have this work the same
+	// way here.
+	// For uniform scaling, it doesn't matter which approach we use.
+	BoxShape.Box.Extents *= Transform.GetScale3D();
 }
 template<typename TransformType>
 static void TransformCapsuleShape(UE::Geometry::FCapsuleShape3d& CapsuleShape, const TransformType& Transform)
@@ -436,15 +441,10 @@ static void TransformSphereShapeByArray(FSphereShape3d& SphereShape, const Trans
 template<typename TransformArrayType>
 static void TransformBoxShapeByArray(FBoxShape3d& BoxShape, const TransformArrayType& TransformSequence)
 {
-	FVector3d CornerVec = BoxShape.Box.Frame.PointAt(BoxShape.Box.Extents) - BoxShape.Box.Frame.Origin;
 	for (const auto& XForm : TransformSequence)
 	{
-		BoxShape.Box.Frame.Transform(XForm);
-		CornerVec = XForm.TransformVector(CornerVec);
+		TransformBoxShape(BoxShape, XForm);
 	}
-	BoxShape.Box.Extents.X = CornerVec.Dot(BoxShape.Box.AxisX());
-	BoxShape.Box.Extents.Y = CornerVec.Dot(BoxShape.Box.AxisY());
-	BoxShape.Box.Extents.Z = CornerVec.Dot(BoxShape.Box.AxisZ());
 }
 template<typename TransformArrayType>
 static void TransformCapsuleShapeByArray(UE::Geometry::FCapsuleShape3d& CapsuleShape, const TransformArrayType& TransformSequence)
diff --git a/Engine/Source/Runtime/GeometryCore/Public/TransformSequence.h b/Engine/Source/Runtime/GeometryCore/Public/TransformSequence.h
--- a/Engine/Source/Runtime/GeometryCore/Public/TransformSequence.h
+++ b/Engine/Source/Runtime/GeometryCore/Public/TransformSequence.h
@@ -39,6 +39,17 @@ public:
 		Transforms.Add( TTransformSRT3<RealType>(Transform) );
 	}
 
+	/**
+	 * Add all transforms in given sequence to the end of this sequence, ie Seq(P) becomes SequenceToAppend * Seq(P)
+	 */
+	void Append(const TTransformSequence3<RealType>& SequenceToAppend)
+	{
+		for (const TTransformSRT3<RealType>& Transform : SequenceToAppend.Transforms)
+		{
+			Append(Transform);
+		}
+	}
+
 	/**
 	 * @return number of transforms in the sequence
 	 */
