MeshModelingTools: Add Revolve Spline tool.

diff --git a/Engine/Plugins/Editor/ModelingToolsEditorMode/Source/ModelingToolsEditorMode/Private/ModelingToolsEditorMode.cpp b/Engine/Plugins/Editor/ModelingToolsEditorMode/Source/ModelingToolsEditorMode/Private/ModelingToolsEditorMode.cpp
--- a/Engine/Plugins/Editor/ModelingToolsEditorMode/Source/ModelingToolsEditorMode/Private/ModelingToolsEditorMode.cpp
+++ b/Engine/Plugins/Editor/ModelingToolsEditorMode/Source/ModelingToolsEditorMode/Private/ModelingToolsEditorMode.cpp
@@ -44,6 +44,7 @@
 #include "DrawPolygonTool.h"
 #include "DrawPolyPathTool.h"
 #include "DrawAndRevolveTool.h"
+#include "RevolveSplineTool.h"
 #include "ShapeSprayTool.h"
 #include "VoxelSolidifyMeshesTool.h"
 #include "VoxelBlendMeshesTool.h"
@@ -572,6 +573,9 @@ void UModelingToolsEditorMode::Enter()
 	auto DrawAndRevolveToolBuilder = NewObject<UDrawAndRevolveToolBuilder>();
 	RegisterTool(ToolManagerCommands.BeginDrawAndRevolveTool, TEXT("BeginDrawAndRevolveTool"), DrawAndRevolveToolBuilder);
 
+	auto RevolveSplineToolBuilder = NewObject<URevolveSplineToolBuilder>();
+	RegisterTool(ToolManagerCommands.BeginRevolveSplineTool, TEXT("BeginRevolveSplineTool"), RevolveSplineToolBuilder);
+
 	auto ShapeSprayToolBuilder = NewObject<UShapeSprayToolBuilder>();
 	RegisterTool(ToolManagerCommands.BeginShapeSprayTool, TEXT("BeginShapeSprayTool"), ShapeSprayToolBuilder);
 
diff --git a/Engine/Plugins/Editor/ModelingToolsEditorMode/Source/ModelingToolsEditorMode/Private/ModelingToolsEditorModeToolkit.cpp b/Engine/Plugins/Editor/ModelingToolsEditorMode/Source/ModelingToolsEditorMode/Private/ModelingToolsEditorModeToolkit.cpp
--- a/Engine/Plugins/Editor/ModelingToolsEditorMode/Source/ModelingToolsEditorMode/Private/ModelingToolsEditorModeToolkit.cpp
+++ b/Engine/Plugins/Editor/ModelingToolsEditorMode/Source/ModelingToolsEditorMode/Private/ModelingToolsEditorModeToolkit.cpp
@@ -257,6 +257,7 @@ void FModelingToolsEditorModeToolkit::RegisterPalettes()
 		Commands.BeginDrawPolygonTool,
 		Commands.BeginDrawPolyPathTool,
 		Commands.BeginDrawAndRevolveTool,
+		Commands.BeginRevolveSplineTool,
 		Commands.BeginTriangulateSplinesTool
 	});
 	ToolkitBuilder->AddPalette(
@@ -1622,6 +1623,7 @@ void FModelingToolsEditorModeToolkit::BuildToolPalette(FName PaletteIndex, class
 		ToolbarBuilder.AddToolBarButton(Commands.BeginDrawPolygonTool);
 		ToolbarBuilder.AddToolBarButton(Commands.BeginDrawPolyPathTool);
 		ToolbarBuilder.AddToolBarButton(Commands.BeginDrawAndRevolveTool);
+		ToolbarBuilder.AddToolBarButton(Commands.BeginRevolveSplineTool);
 		ToolbarBuilder.AddToolBarButton(Commands.BeginRevolveBoundaryTool);
 		ToolbarBuilder.AddToolBarButton(Commands.BeginCombineMeshesTool);
 		ToolbarBuilder.AddToolBarButton(Commands.BeginDuplicateMeshesTool);
diff --git a/Engine/Plugins/Editor/ModelingToolsEditorMode/Source/ModelingToolsEditorMode/Private/ModelingToolsManagerActions.cpp b/Engine/Plugins/Editor/ModelingToolsEditorMode/Source/ModelingToolsEditorMode/Private/ModelingToolsManagerActions.cpp
--- a/Engine/Plugins/Editor/ModelingToolsEditorMode/Source/ModelingToolsEditorMode/Private/ModelingToolsManagerActions.cpp
+++ b/Engine/Plugins/Editor/ModelingToolsEditorMode/Source/ModelingToolsEditorMode/Private/ModelingToolsManagerActions.cpp
@@ -79,6 +79,7 @@ void FModelingToolsManagerCommands::RegisterCommands()
 		REGISTER_MODELING_TOOL_COMMAND(BeginDrawPolygonTool, "PolyExt", "Draw and extrude 2D Polygons to create new objects");
 		REGISTER_MODELING_TOOL_COMMAND(BeginDrawPolyPathTool, "PathExt", "Draw and extrude 2D Paths to create new objects");
 		REGISTER_MODELING_TOOL_COMMAND(BeginDrawAndRevolveTool, "PathRev", "Draw and revolve 2D Paths to create new objects");
+		REGISTER_MODELING_TOOL_COMMAND(BeginRevolveSplineTool, "SplnRev", "Revolve splines to create new objects");
 		REGISTER_MODELING_TOOL_COMMAND(BeginRevolveBoundaryTool, "BdryRev", "Revolve Mesh boundary loops to create new objects");
 		REGISTER_MODELING_TOOL_COMMAND(BeginCombineMeshesTool, "Merge", "Merge multiple Meshes to create new objects");
 		REGISTER_MODELING_TOOL_COMMAND(BeginDuplicateMeshesTool, "Dupe", "Duplicate single Meshes to create new objects");
@@ -89,6 +90,7 @@ void FModelingToolsManagerCommands::RegisterCommands()
 		REGISTER_MODELING_TOOL_COMMAND(BeginDrawPolygonTool, "Extrude Polygon", "Draw and extrude 2D Polygons to create new objects");
 		REGISTER_MODELING_TOOL_COMMAND(BeginDrawPolyPathTool, "Extrude Path", "Draw and extrude 2D Paths to create new objects");
 		REGISTER_MODELING_TOOL_COMMAND(BeginDrawAndRevolveTool, "Revolve Path", "Draw and revolve 2D Paths to create new objects");
+		REGISTER_MODELING_TOOL_COMMAND(BeginRevolveSplineTool, "Revolve Spline", "Revolve splines to create new objects");
 		REGISTER_MODELING_TOOL_COMMAND(BeginRevolveBoundaryTool, "Revolve Boundary", "Revolve Mesh boundary loops to create new objects");
 		REGISTER_MODELING_TOOL_COMMAND(BeginCombineMeshesTool, "Merge", "Merge multiple Meshes to create new objects");
 		REGISTER_MODELING_TOOL_COMMAND(BeginDuplicateMeshesTool, "Duplicate", "Duplicate single Meshes to create new objects");
diff --git a/Engine/Plugins/Editor/ModelingToolsEditorMode/Source/ModelingToolsEditorMode/Public/ModelingToolsManagerActions.h b/Engine/Plugins/Editor/ModelingToolsEditorMode/Source/ModelingToolsEditorMode/Public/ModelingToolsManagerActions.h
--- a/Engine/Plugins/Editor/ModelingToolsEditorMode/Source/ModelingToolsEditorMode/Public/ModelingToolsManagerActions.h
+++ b/Engine/Plugins/Editor/ModelingToolsEditorMode/Source/ModelingToolsEditorMode/Public/ModelingToolsManagerActions.h
@@ -46,6 +46,7 @@ public:
 	TSharedPtr<FUICommandInfo> BeginDrawPolygonTool;
 	TSharedPtr<FUICommandInfo> BeginDrawPolyPathTool;
 	TSharedPtr<FUICommandInfo> BeginDrawAndRevolveTool;
+	TSharedPtr<FUICommandInfo> BeginRevolveSplineTool;
 	TSharedPtr<FUICommandInfo> BeginShapeSprayTool;
 	TSharedPtr<FUICommandInfo> BeginTriangulateSplinesTool;
 
diff --git a/Engine/Plugins/Experimental/MeshModelingToolsetExp/Source/MeshModelingToolsExp/Private/RevolveSplineTool.cpp b/Engine/Plugins/Experimental/MeshModelingToolsetExp/Source/MeshModelingToolsExp/Private/RevolveSplineTool.cpp
new file mode 100644
--- /dev/null
+++ b/Engine/Plugins/Experimental/MeshModelingToolsetExp/Source/MeshModelingToolsExp/Private/RevolveSplineTool.cpp
@@ -0,0 +1,474 @@
+// Copyright Epic Games, Inc. All Rights Reserved.
+
+#include "RevolveSplineTool.h"
+
+#include "InteractiveToolManager.h"
+#include "ToolBuilderUtil.h"
+#include "ToolSetupUtil.h"
+#include "BaseBehaviors/SingleClickBehavior.h"
+#include "GameFramework/Actor.h"
+#include "Selection/ToolSelectionUtil.h"
+#include "Properties/MeshMaterialProperties.h"
+#include "CompGeom/PolygonTriangulation.h"
+
+#include "Mechanics/ConstructionPlaneMechanic.h"
+#include "CompositionOps/CurveSweepOp.h"
+#include "DynamicMesh/MeshTransforms.h"
+#include "PropertySets/CreateMeshObjectTypeProperties.h"
+#include "Components/SplineComponent.h"
+#include "MeshOpPreviewHelpers.h" // UMeshOpPreviewWithBackgroundCompute
+
+#define LOCTEXT_NAMESPACE "URevolveSplineTool"
+
+using namespace UE::Geometry;
+
+
+
+void URevolveSplineTool::Setup()
+{
+	UInteractiveTool::Setup();
+
+	Settings = NewObject<URevolveSplineToolProperties>(this);
+	Settings->RestoreProperties(this);
+	AddToolPropertySource(Settings);
+
+	ToolActions = NewObject<URevolveSplineToolActionPropertySet>(this);
+	ToolActions->Initialize(this);
+	AddToolPropertySource(ToolActions);
+
+	MaterialProperties = NewObject<UNewMeshMaterialProperties>(this);
+	AddToolPropertySource(MaterialProperties);
+	MaterialProperties->RestoreProperties(this);
+
+	OutputTypeProperties = NewObject<UCreateMeshObjectTypeProperties>(this);
+	OutputTypeProperties->InitializeDefault();
+	OutputTypeProperties->RestoreProperties(this);
+	OutputTypeProperties->WatchProperty(OutputTypeProperties->OutputType, [this](FString) { OutputTypeProperties->UpdatePropertyVisibility(); });
+	AddToolPropertySource(OutputTypeProperties);
+
+	SetToolDisplayName(LOCTEXT("RevolveSplineToolName", "Revolve Spline"));
+	GetToolManager()->DisplayMessage(
+		LOCTEXT("RevolveSplineToolDescription", "Revolve the selected spline to create a mesh."),
+		EToolMessageLevel::UserNotification);
+
+	Preview = NewObject<UMeshOpPreviewWithBackgroundCompute>(this);
+	Preview->Setup(GetTargetWorld(), this);
+	Preview->PreviewMesh->EnableWireframe(MaterialProperties->bShowWireframe);
+	Preview->ConfigureMaterials(MaterialProperties->Material.Get(),
+		ToolSetupUtil::GetDefaultWorkingMaterial(GetToolManager()));
+	ToolSetupUtil::ApplyRenderingConfigurationToPreview(Preview->PreviewMesh, nullptr);
+
+	Preview->OnMeshUpdated.AddLambda(
+		[this](const UMeshOpPreviewWithBackgroundCompute* UpdatedPreview)
+		{
+			UpdateAcceptWarnings(UpdatedPreview->HaveEmptyResult() ? EAcceptWarning::EmptyForbidden : EAcceptWarning::NoWarning);
+		}
+	);
+
+	// TODO: We'll probably want a click behavior someday for clicking on the spline to align to a tangent at a point
+
+	// The plane mechanic is used for the revolution axis.
+	// Note: The only thing we really end up using from it is the gizmo and the control+click. We
+	// could use our own gizmo directly.
+	PlaneMechanic = NewObject<UConstructionPlaneMechanic>(this);
+	PlaneMechanic->Setup(this);
+	PlaneMechanic->Initialize(GetTargetWorld(), FFrame3d(Settings->AxisOrigin, 
+		FRotator(Settings->AxisOrientation.X, Settings->AxisOrientation.Y, 0).Quaternion()));
+	PlaneMechanic->bShowGrid = false;
+	PlaneMechanic->OnPlaneChanged.AddLambda([this]() {
+		Settings->AxisOrigin = (FVector)PlaneMechanic->Plane.Origin;
+		FRotator AxisOrientation = ((FQuat)PlaneMechanic->Plane.Rotation).Rotator();
+		Settings->AxisOrientation.X = AxisOrientation.Pitch;
+		Settings->AxisOrientation.Y = AxisOrientation.Yaw;
+		NotifyOfPropertyChangeByTool(Settings);
+		UpdateRevolutionAxis();
+		});
+	// Add if we get our own click behavior:
+	//PlaneMechanic->UpdateClickPriority(ClickBehavior->GetPriority().MakeLower());
+
+	// TODO: It would be nice to have a drag alignment mechanic for the above gizmo, but we currently
+	// don't have a way to pass in a custom alignment raycast, which we would want in order to snap
+	// and align to spline points.
+
+	PollSplineUpdates();
+
+	if (Settings->bResetAxisOnStart)
+	{
+		ResetAxis();
+	}
+	else
+	{
+		UpdateRevolutionAxis();
+	}
+}
+
+void URevolveSplineTool::ResetAxis()
+{
+	if (!Spline.IsValid())
+	{
+		return;
+	}
+	int32 NumSplinePoints = Spline->GetNumberOfSplinePoints();
+	if (NumSplinePoints == 0)
+	{
+		return;
+	}
+	Settings->AxisOrigin = Spline->GetLocationAtSplinePoint(0, ESplineCoordinateSpace::World);
+
+	// Our axis is the X axis of the frame, and we align it to Last-First
+	FVector3d PlaneX = Spline->GetLocationAtSplinePoint(NumSplinePoints-1, ESplineCoordinateSpace::World) - Settings->AxisOrigin;
+	PlaneX.Normalize();
+	FFrame3d PlaneFrame(Settings->AxisOrigin, SplineFitPlaneNormal);
+	if (!PlaneX.IsZero())
+	{
+		FVector3d PlaneY = SplineFitPlaneNormal.Cross(PlaneX);
+		FVector3d PlaneZ = PlaneX.Cross(PlaneY);
+		PlaneFrame = FFrame3d(Settings->AxisOrigin, PlaneX, PlaneY, PlaneZ);
+	}
+	FRotator AxisOrientation = ((FQuat)PlaneFrame.Rotation).Rotator();
+	Settings->AxisOrientation.X = AxisOrientation.Pitch;
+	Settings->AxisOrientation.Y = AxisOrientation.Yaw;
+
+	NotifyOfPropertyChangeByTool(Settings);
+	UpdateRevolutionAxis();
+}
+
+void URevolveSplineTool::PollSplineUpdates()
+{
+	if (bLostInputSpline)
+	{
+		return;
+	}
+
+	if (!Spline.IsValid())
+	{
+		// Attempt to recapture the spline.
+		TArray<USplineComponent*> Splines;
+		if (SplineOwningActor.IsValid())
+		{
+			SplineOwningActor->GetComponents<USplineComponent>(Splines);
+		}
+		if (SplineComponentIndex < Splines.Num() && ensure(SplineComponentIndex >= 0))
+		{
+			Spline = Splines[SplineComponentIndex];
+		}
+		else
+		{
+			GetToolManager()->DisplayMessage(
+				LOCTEXT("LostSpline", "Tool lost reference to the input spline. The spline input will no longer be updated."),
+				EToolMessageLevel::UserWarning);
+			bLostInputSpline = true;
+			return;
+		}
+	}
+
+	if (Spline->SplineCurves.Version == LastSplineVersion)
+	{
+		// No update necessary
+		return;
+	}
+
+	UpdatePointsFromSpline();
+
+	LastSplineVersion = Spline->SplineCurves.Version;
+}
+
+void URevolveSplineTool::UpdatePointsFromSpline()
+{
+	if (!Spline.IsValid())
+	{
+		return;
+	}
+
+	bProfileCurveIsClosed = Spline->IsClosedLoop();
+
+	// Update the curve plane
+	TArray<FVector> SplineControlPoints;
+	int32 NumSplinePoints = Spline->GetNumberOfSplinePoints();
+	for (int32 i = 0; i < NumSplinePoints; ++i)
+	{
+		SplineControlPoints.Add(Spline->GetLocationAtSplinePoint(i, ESplineCoordinateSpace::World));
+	}
+	PolygonTriangulation::ComputePolygonPlane<double>(SplineControlPoints, SplineFitPlaneNormal, SplineFitPlaneOrigin);
+
+	// Update the points we actually revolve
+	ProfileCurve.Reset();
+	switch (Settings->SampleMode)
+	{
+	case ERevolveSplineSampleMode::ControlPointsOnly:
+	{
+		ProfileCurve = SplineControlPoints;
+	}
+	break;
+	case ERevolveSplineSampleMode::PolyLineMaxError:
+	{
+		Spline->ConvertSplineToPolyLine(ESplineCoordinateSpace::World,
+			Settings->ErrorTolerance * Settings->ErrorTolerance,
+			ProfileCurve);
+	};
+	break;
+	case ERevolveSplineSampleMode::UniformSpacingAlongCurve:
+	{
+		double Length = Spline->GetSplineLength();
+		int32 NumSegments = FMath::RoundFromZero(Length / FMath::Max(0.01, Settings->MaxSampleDistance));
+		double SegmentLengthToUse = Length / NumSegments;
+		for (int32 i = 0; i <= NumSegments; ++i)
+		{
+			double DistanceToSampleAt = Length * ((double)i / NumSegments);
+			ProfileCurve.Add(Spline->GetLocationAtDistanceAlongSpline(DistanceToSampleAt, 
+				ESplineCoordinateSpace::World));
+		}
+	}
+	break;
+	}
+
+	Preview->InvalidateResult();
+}
+
+void URevolveSplineTool::OnTick(float DeltaTime)
+{
+	PollSplineUpdates();
+
+	if (PlaneMechanic)
+	{
+		PlaneMechanic->Tick(DeltaTime);
+	}
+
+	if (Preview)
+	{
+		Preview->Tick(DeltaTime);
+	}
+}
+
+void URevolveSplineTool::Render(IToolsContextRenderAPI* RenderAPI)
+{
+	Super::Render(RenderAPI);
+
+	FViewCameraState CameraState;
+	GetToolManager()->GetContextQueriesAPI()->GetCurrentViewState(CameraState);
+
+	if (PlaneMechanic)
+	{
+		PlaneMechanic->Render(RenderAPI);
+
+		// Draw the axis of rotation
+		float PdiScale = CameraState.GetPDIScalingFactor();
+		FPrimitiveDrawInterface* PDI = RenderAPI->GetPrimitiveDrawInterface();
+
+		FColor AxisColor(240, 16, 240);
+		double AxisThickness = 1 * PdiScale;
+		double AxisHalfLength = ToolSceneQueriesUtil::CalculateDimensionFromVisualAngleD(CameraState, RevolutionAxisOrigin, 90);
+
+		FVector3d StartPoint = RevolutionAxisOrigin - (RevolutionAxisDirection * (AxisHalfLength * PdiScale));
+		FVector3d EndPoint = RevolutionAxisOrigin + (RevolutionAxisDirection * (AxisHalfLength * PdiScale));
+
+		PDI->DrawLine((FVector)StartPoint, (FVector)EndPoint, AxisColor, SDPG_Foreground,
+			AxisThickness, 0.0f, true);
+	}
+}
+
+void URevolveSplineTool::OnPropertyModified(UObject* PropertySet, FProperty* Property)
+{
+	if (PropertySet == OutputTypeProperties)
+	{
+		return;
+	}
+
+	if (Property)
+	{
+		if (Property->GetFName() == GET_MEMBER_NAME_CHECKED(UNewMeshMaterialProperties, Material))
+		{
+			Preview->ConfigureMaterials(MaterialProperties->Material.Get(),
+				ToolSetupUtil::GetDefaultWorkingMaterial(GetToolManager()));
+		}
+		else if (Property->GetFName() == GET_MEMBER_NAME_CHECKED(URevolveSplineToolProperties, SampleMode)
+			|| Property->GetFName() == GET_MEMBER_NAME_CHECKED(URevolveSplineToolProperties, ErrorTolerance)
+			|| Property->GetFName() == GET_MEMBER_NAME_CHECKED(URevolveSplineToolProperties, MaxSampleDistance))
+		{
+			UpdatePointsFromSpline();
+		}
+		else if (Property->GetFName() == GET_MEMBER_NAME_CHECKED(URevolveSplineToolProperties, AxisOrigin)
+			|| Property->GetFName() == GET_MEMBER_NAME_CHECKED(URevolveSplineToolProperties, AxisOrientation))
+		{
+			PlaneMechanic->SetPlaneWithoutBroadcast(FFrame3d(Settings->AxisOrigin,
+				FRotator(Settings->AxisOrientation.X, Settings->AxisOrientation.Y, 0).Quaternion()));
+			UpdateRevolutionAxis();
+		}
+	}
+
+	Preview->PreviewMesh->EnableWireframe(MaterialProperties->bShowWireframe);
+
+	Preview->InvalidateResult();
+}
+
+void URevolveSplineTool::Shutdown(EToolShutdownType ShutdownType)
+{
+	Settings->SaveProperties(this);
+	OutputTypeProperties->SaveProperties(this);
+	MaterialProperties->SaveProperties(this);
+
+	FDynamicMeshOpResult Result = Preview->Shutdown();
+
+	if (ShutdownType == EToolShutdownType::Accept)
+	{
+		GetToolManager()->BeginUndoTransaction(LOCTEXT("RevolveSplineAction", "Revolve Spline"));
+
+		// Generate the result asset
+		GenerateAsset(Result);
+
+		GetToolManager()->EndUndoTransaction();
+	}
+
+	PlaneMechanic->Shutdown();
+
+	Preview = nullptr;
+	Settings = nullptr;
+	MaterialProperties = nullptr;
+	OutputTypeProperties = nullptr;
+	ToolActions = nullptr;
+	PlaneMechanic = nullptr;
+
+	Super::Shutdown(ShutdownType);
+}
+
+void URevolveSplineTool::GenerateAsset(const FDynamicMeshOpResult& OpResult)
+{
+	if (OpResult.Mesh.Get() == nullptr) return;
+
+	if (!ensure(OpResult.Mesh.Get() && OpResult.Mesh->TriangleCount() > 0))
+	{
+		return;
+	}
+
+	GetToolManager()->BeginUndoTransaction(LOCTEXT("RevolveSplineToolTransactionName", "Revolve Spline"));
+
+	FCreateMeshObjectParams NewMeshObjectParams;
+	NewMeshObjectParams.TargetWorld = GetTargetWorld();
+	NewMeshObjectParams.Transform = (FTransform)OpResult.Transform;
+	NewMeshObjectParams.BaseName = TEXT("RevolveSpline");
+	NewMeshObjectParams.Materials.Add(MaterialProperties->Material.Get());
+	NewMeshObjectParams.SetMesh(OpResult.Mesh.Get());
+	OutputTypeProperties->ConfigureCreateMeshObjectParams(NewMeshObjectParams);
+	FCreateMeshObjectResult Result = UE::Modeling::CreateMeshObject(GetToolManager(), MoveTemp(NewMeshObjectParams));
+	if (Result.IsOK() && Result.NewActor != nullptr)
+	{
+		ToolSelectionUtil::SetNewActorSelection(GetToolManager(), Result.NewActor);
+	}
+
+	GetToolManager()->EndUndoTransaction();
+}
+
+bool URevolveSplineTool::CanAccept() const
+{
+	return Preview->HaveValidNonEmptyResult();
+}
+
+TUniquePtr<FDynamicMeshOperator> URevolveSplineTool::MakeNewOperator()
+{
+	TUniquePtr<FCurveSweepOp> CurveSweepOp = MakeUnique<FCurveSweepOp>();
+
+	// Assemble profile curve
+	CurveSweepOp->ProfileCurve = ProfileCurve;
+	CurveSweepOp->bProfileCurveIsClosed = bProfileCurveIsClosed;
+
+	// If we are capping the top and bottom, we just add a couple extra vertices and mark the curve as being closed
+	if (!bProfileCurveIsClosed && Settings->bClosePathToAxis)
+	{
+		// Project first and last points onto the revolution axis.
+		FVector3d FirstPoint = CurveSweepOp->ProfileCurve[0];
+		FVector3d LastPoint = CurveSweepOp->ProfileCurve.Last();
+		double DistanceAlongAxis = RevolutionAxisDirection.Dot(LastPoint - RevolutionAxisOrigin);
+		FVector3d ProjectedPoint = RevolutionAxisOrigin + (RevolutionAxisDirection * DistanceAlongAxis);
+
+		CurveSweepOp->ProfileCurve.Add(ProjectedPoint);
+
+		DistanceAlongAxis = RevolutionAxisDirection.Dot(FirstPoint - RevolutionAxisOrigin);
+		ProjectedPoint = RevolutionAxisOrigin + (RevolutionAxisDirection * DistanceAlongAxis);
+
+		CurveSweepOp->ProfileCurve.Add(ProjectedPoint);
+		CurveSweepOp->bProfileCurveIsClosed = true;
+	}
+
+	Settings->ApplyToCurveSweepOp(*MaterialProperties,
+		RevolutionAxisOrigin, RevolutionAxisDirection, *CurveSweepOp);
+
+	return CurveSweepOp;
+}
+
+
+//Uses the settings stored in the properties object to update the revolution axis
+void URevolveSplineTool::UpdateRevolutionAxis()
+{
+	RevolutionAxisOrigin = (FVector3d)Settings->AxisOrigin;
+	
+	FRotator Rotator(Settings->AxisOrientation.X, Settings->AxisOrientation.Y, 0);
+	RevolutionAxisDirection = (FVector3d)Rotator.RotateVector(FVector(1, 0, 0));
+
+	PlaneMechanic->SetPlaneWithoutBroadcast(FFrame3d(Settings->AxisOrigin, Rotator.Quaternion()));
+
+	if (Preview)
+	{
+		Preview->InvalidateResult();
+	}
+}
+
+// To be called by builder
+void URevolveSplineTool::SetSpline(USplineComponent* SplineComponent)
+{
+	Spline = SplineComponent;
+	SplineOwningActor = Spline->GetOwner();
+
+	TArray<USplineComponent*> Splines;
+	SplineOwningActor->GetComponents<USplineComponent>(Splines);
+
+	SplineComponentIndex = Splines.IndexOfByKey(Spline.Get());
+}
+
+void URevolveSplineTool::RequestAction(ERevolveSplineToolAction Action)
+{
+	if (Action == ERevolveSplineToolAction::ResetAxis)
+	{
+		ResetAxis();
+	}
+}
+
+
+/// Tool builder:
+
+bool URevolveSplineToolBuilder::CanBuildTool(const FToolBuilderState& SceneState) const
+{
+	int32 NumSplines = ToolBuilderUtil::CountComponents(SceneState, [&](UActorComponent* Object) -> bool
+	{
+		return Object->IsA<USplineComponent>();
+	});
+	return NumSplines == 1;
+}
+
+UInteractiveTool* URevolveSplineToolBuilder::BuildTool(const FToolBuilderState& SceneState) const
+{
+	URevolveSplineTool* NewTool = NewObject<URevolveSplineTool>(SceneState.ToolManager);
+	USplineComponent* Spline = Cast<USplineComponent>(ToolBuilderUtil::FindFirstComponent(SceneState, [&](UActorComponent* Object)
+	{
+		return Object->IsA<USplineComponent>();
+	}));
+
+	ensure(Spline);
+	NewTool->SetSpline(Spline);
+
+	NewTool->SetWorld(SceneState.World);
+
+	return NewTool;
+}
+
+// Action set:
+
+void URevolveSplineToolActionPropertySet::PostAction(ERevolveSplineToolAction Action)
+{
+	if (ParentTool.IsValid())
+	{
+		ParentTool->RequestAction(Action);
+	}
+}
+
+
+#undef LOCTEXT_NAMESPACE
diff --git a/Engine/Plugins/Experimental/MeshModelingToolsetExp/Source/MeshModelingToolsExp/Public/RevolveSplineTool.h b/Engine/Plugins/Experimental/MeshModelingToolsetExp/Source/MeshModelingToolsExp/Public/RevolveSplineTool.h
new file mode 100644
--- /dev/null
+++ b/Engine/Plugins/Experimental/MeshModelingToolsetExp/Source/MeshModelingToolsExp/Public/RevolveSplineTool.h
@@ -0,0 +1,223 @@
+// Copyright Epic Games, Inc. All Rights Reserved.
+
+#pragma once
+
+#include "CoreMinimal.h"
+
+#include "InteractiveTool.h"
+#include "InteractiveToolBuilder.h"
+#include "InteractiveToolQueryInterfaces.h"
+#include "ModelingOperators.h" // IDynamicMeshOperatorFactory
+#include "Properties/RevolveProperties.h"
+
+#include "RevolveSplineTool.generated.h"
+
+class AActor;
+struct FDynamicMeshOpResult;
+class UConstructionPlaneMechanic;
+class UCreateMeshObjectTypeProperties;
+class UMeshOpPreviewWithBackgroundCompute;
+class UNewMeshMaterialProperties;
+class URevolveSplineTool;
+class USplineComponent;
+class UWorld;
+
+//~ TODO: Might want to have some shared enum for sampling splines in a util folder,
+//~ but hesitant to prescribe it until other tools want it.
+UENUM()
+enum class ERevolveSplineSampleMode : uint8
+{
+	ControlPointsOnly,
+
+	// Place points along the spline such that the resulting polyline has no more than 
+	// some maximum deviation from the curve.
+	PolyLineMaxError,
+
+	// Place points along spline that are an equal spacing apart, and so that the spacing
+	// is as close as possible to some max spacing.
+	UniformSpacingAlongCurve,
+};
+
+
+
+UCLASS()
+class MESHMODELINGTOOLSEXP_API URevolveSplineToolProperties : public URevolveProperties
+{
+	GENERATED_BODY()
+
+public:
+
+	/** Determines how points to revolve are actually picked from the spline. */
+	UPROPERTY(EditAnywhere, Category = Spline)
+	ERevolveSplineSampleMode SampleMode = ERevolveSplineSampleMode::ControlPointsOnly;
+
+	/** How far to allow the triangulation boundary can deviate from the spline curve before we add more vertices. */
+	UPROPERTY(EditAnywhere, Category = Spline, meta = (ClampMin = 0.001, EditConditionHides,
+		EditCondition = "SampleMode == ERevolveSplineSampleMode::PolyLineMaxError"))
+	double ErrorTolerance = 1.0;
+
+	/** The maximal distance that the spacing should be allowed to be. */
+	UPROPERTY(EditAnywhere, Category = Spline, meta = (ClampMin = 0.01, EditConditionHides,
+		EditCondition = "SampleMode == ERevolveSplineSampleMode::UniformSpacingAlongCurve"))
+	double MaxSampleDistance = 50.0;
+
+	/** Determines how end caps are created. This is not relevant if the end caps are not visible or if the path is not closed. */
+	UPROPERTY(EditAnywhere, Category = Revolve, AdvancedDisplay, meta = (DisplayAfter = "QuadSplitMode",
+		EditCondition = "HeightOffsetPerDegree != 0 || RevolveDegrees != 360"))
+	ERevolvePropertiesCapFillMode CapFillMode = ERevolvePropertiesCapFillMode::Delaunay;
+
+	/** Connect the ends of an open path to the axis to add caps to the top and bottom of the revolved result.
+	  * This is not relevant for paths that are already closed. */
+	UPROPERTY(EditAnywhere, Category = Revolve)
+	bool bClosePathToAxis = true;
+	
+	/** Sets the revolution axis origin. */
+	UPROPERTY(EditAnywhere, Category = RevolutionAxis, meta = (DisplayName = "Origin",
+		Delta = 5, LinearDeltaSensitivity = 1))
+	FVector AxisOrigin = FVector(0, 0, 0);
+
+	/** Sets the revolution axis pitch and yaw. */
+	UPROPERTY(EditAnywhere, Category = RevolutionAxis, meta = (DisplayName = "Orientation",
+		UIMin = -180, UIMax = 180, ClampMin = -180000, ClampMax = 180000))
+	FVector2D AxisOrientation;
+
+	/** 
+	 * If true, the revolution axis is re-fit to the input spline on each tool start. If false, the previous
+	 * revolution axis is kept.
+	 */
+	UPROPERTY(EditAnywhere, Category = RevolutionAxis, AdvancedDisplay)
+	bool bResetAxisOnStart = true;
+
+protected:
+	virtual ERevolvePropertiesCapFillMode GetCapFillMode() const override
+	{
+		return CapFillMode;
+	}
+};
+
+enum class ERevolveSplineToolAction
+{
+	ResetAxis
+};
+
+UCLASS()
+class MESHMODELINGTOOLSEXP_API URevolveSplineToolActionPropertySet : public UInteractiveToolPropertySet
+{
+	GENERATED_BODY()
+
+public:
+	TWeakObjectPtr<URevolveSplineTool> ParentTool;
+
+	void Initialize(URevolveSplineTool* ParentToolIn) { ParentTool = ParentToolIn; }
+	void PostAction(ERevolveSplineToolAction Action);
+
+	/** Fit the axis to the current curve(by aligning it to the startand end points) */
+	UFUNCTION(CallInEditor, Category = Actions, meta = (DisplayPriority = 1))
+	void ResetAxis() { PostAction(ERevolveSplineToolAction::ResetAxis); }
+};
+
+//~ We might someday decide to merge this tool as a separate path in the DrawAndRevolveTool...
+/**
+ * Revolves a selected spline to create a new mesh.
+ */
+UCLASS()
+class MESHMODELINGTOOLSEXP_API URevolveSplineTool : public UInteractiveTool, 
+	public IInteractiveToolEditorGizmoAPI, 
+	public UE::Geometry::IDynamicMeshOperatorFactory
+{
+	GENERATED_BODY()
+
+public:
+
+	virtual void SetSpline(USplineComponent* SplineComponent);
+
+	virtual void SetWorld(UWorld* World) { TargetWorld = World; }
+	virtual UWorld* GetTargetWorld() { return TargetWorld.Get(); }
+	virtual void RequestAction(ERevolveSplineToolAction ActionType);
+
+	// UInteractiveTool
+	virtual void Setup() override;
+	virtual void Shutdown(EToolShutdownType ShutdownType) override;
+	virtual void OnTick(float DeltaTime) override;
+	virtual bool HasCancel() const override { return true; }
+	virtual bool HasAccept() const override { return true; }
+	virtual bool CanAccept() const override;
+	virtual void Render(IToolsContextRenderAPI* RenderAPI) override;
+	virtual void OnPropertyModified(UObject* PropertySet, FProperty* Property) override;
+
+	// IDynamicMeshOperatorFactory
+	virtual TUniquePtr<UE::Geometry::FDynamicMeshOperator> MakeNewOperator();
+
+	// IInteractiveToolEditorGizmoAPI
+	virtual bool GetAllowStandardEditorGizmos() override { return true; }
+
+private:
+
+	UPROPERTY()
+	TObjectPtr<URevolveSplineToolProperties> Settings = nullptr;
+
+	UPROPERTY()
+	TObjectPtr<UNewMeshMaterialProperties> MaterialProperties = nullptr;
+
+	UPROPERTY()
+	TObjectPtr<UCreateMeshObjectTypeProperties> OutputTypeProperties = nullptr;
+
+	UPROPERTY()
+	TObjectPtr<URevolveSplineToolActionPropertySet> ToolActions = nullptr;
+
+	UPROPERTY()
+	TObjectPtr<UMeshOpPreviewWithBackgroundCompute> Preview = nullptr;
+
+	UPROPERTY()
+	TObjectPtr<UConstructionPlaneMechanic> PlaneMechanic = nullptr;
+
+	TWeakObjectPtr<UWorld> TargetWorld = nullptr;
+
+	TWeakObjectPtr<USplineComponent> Spline = nullptr;
+	
+	// Used to try to re-acquire the spline if it is inside a BP actor, where it will be destroyed and recreated
+	// whenever the user edits it.
+	TWeakObjectPtr<AActor> SplineOwningActor = nullptr;
+	int32 SplineComponentIndex;
+	// If failed to reacquire once, used to avoid trying to reaquire again.
+	bool bLostInputSpline = false;
+
+	// The actual points to be resolved, sampled from the spline
+	TArray<FVector3d> ProfileCurve;
+	bool bProfileCurveIsClosed;
+
+	// See if the spline has changed, and update our data if so
+	void PollSplineUpdates();
+	// Track the spline 'Version' integer, which is incremented when splines are changed
+	uint32 LastSplineVersion;
+	// Update the profile curve and fit plane from spline
+	void UpdatePointsFromSpline();
+
+	// Axis direction in vector form (since the user modifiable values are a pitch and yaw)
+	FVector3d RevolutionAxisDirection;
+	// This duplicates Settings->AxisOrigin, but kept for cleanliness since we do need RevolutionAxisDirection
+	FVector3d RevolutionAxisOrigin;
+	void UpdateRevolutionAxis();
+
+	void ResetAxis();
+
+	FVector3d SplineFitPlaneOrigin;
+	FVector3d SplineFitPlaneNormal;
+
+	virtual void GenerateAsset(const FDynamicMeshOpResult& OpResult);
+};
+
+
+
+UCLASS(Transient)
+class MESHMODELINGTOOLSEXP_API URevolveSplineToolBuilder : public UInteractiveToolBuilder
+{
+	GENERATED_BODY()
+
+public:
+	/** @return true if spline component sources can be found in the active selection */
+	virtual bool CanBuildTool(const FToolBuilderState& SceneState) const override;
+
+	/** @return new Tool instance initialized with selected spline source(s) */
+	virtual UInteractiveTool* BuildTool(const FToolBuilderState& SceneState) const override;
+};
diff --git a/Engine/Plugins/Runtime/MeshModelingToolset/Source/ModelingOperators/Private/CompositionOps/CurveSweepOp.cpp b/Engine/Plugins/Runtime/MeshModelingToolset/Source/ModelingOperators/Private/CompositionOps/CurveSweepOp.cpp
--- a/Engine/Plugins/Runtime/MeshModelingToolset/Source/ModelingOperators/Private/CompositionOps/CurveSweepOp.cpp
+++ b/Engine/Plugins/Runtime/MeshModelingToolset/Source/ModelingOperators/Private/CompositionOps/CurveSweepOp.cpp
@@ -112,6 +112,7 @@ void FCurveSweepOp::CalculateResult(FProgressCancel* Progress)
 					ConstrainedDelaunayTriangulate<double>,
 					PlaneOrigin,
 					PlaneNormal);
+				break;
 			default:
 				checkSlow(false);
 			}
