Fix includes in RevolveSplineTool.cpp to fix non-unity build.

diff --git a/Engine/Plugins/Experimental/MeshModelingToolsetExp/Source/MeshModelingToolsExp/Private/RevolveSplineTool.cpp b/Engine/Plugins/Experimental/MeshModelingToolsetExp/Source/MeshModelingToolsExp/Private/RevolveSplineTool.cpp
--- a/Engine/Plugins/Experimental/MeshModelingToolsetExp/Source/MeshModelingToolsExp/Private/RevolveSplineTool.cpp
+++ b/Engine/Plugins/Experimental/MeshModelingToolsetExp/Source/MeshModelingToolsExp/Private/RevolveSplineTool.cpp
@@ -2,21 +2,21 @@
 
 #include "RevolveSplineTool.h"
 
-#include "InteractiveToolManager.h"
-#include "ToolBuilderUtil.h"
-#include "ToolSetupUtil.h"
-#include "BaseBehaviors/SingleClickBehavior.h"
-#include "GameFramework/Actor.h"
-#include "Selection/ToolSelectionUtil.h"
-#include "Properties/MeshMaterialProperties.h"
 #include "CompGeom/PolygonTriangulation.h"
-
-#include "Mechanics/ConstructionPlaneMechanic.h"
+#include "Components/SplineComponent.h"
 #include "CompositionOps/CurveSweepOp.h"
 #include "DynamicMesh/MeshTransforms.h"
-#include "PropertySets/CreateMeshObjectTypeProperties.h"
-#include "Components/SplineComponent.h"
+#include "Engine/World.h" // UWorld (for WeakPtr)
+#include "GameFramework/Actor.h"
+#include "InteractiveToolManager.h"
+#include "Mechanics/ConstructionPlaneMechanic.h"
 #include "MeshOpPreviewHelpers.h" // UMeshOpPreviewWithBackgroundCompute
+#include "Properties/MeshMaterialProperties.h"
+#include "PropertySets/CreateMeshObjectTypeProperties.h"
+#include "SceneManagement.h" // FPrimitiveDrawInterface
+#include "Selection/ToolSelectionUtil.h"
+#include "ToolBuilderUtil.h"
+#include "ToolSetupUtil.h"
 
 #define LOCTEXT_NAMESPACE "URevolveSplineTool"
 
