Skeleton Modifier: update the skeletal mesh only once the skeleton is up to date

diff --git a/Engine/Plugins/Experimental/MeshModelingToolsetExp/Source/SkeletalMeshModifiers/Private/SkeletonModifier.cpp b/Engine/Plugins/Experimental/MeshModelingToolsetExp/Source/SkeletalMeshModifiers/Private/SkeletonModifier.cpp
--- a/Engine/Plugins/Experimental/MeshModelingToolsetExp/Source/SkeletalMeshModifiers/Private/SkeletonModifier.cpp
+++ b/Engine/Plugins/Experimental/MeshModelingToolsetExp/Source/SkeletalMeshModifiers/Private/SkeletonModifier.cpp
@@ -349,8 +349,6 @@ bool USkeletonModifier::CommitSkeletonToSkeletalMesh()
 	// update skeletal mesh LOD (cf. USkeletalMesh::CommitMeshDescription)
 	SkeletalMesh->CommitMeshDescription(USkeletonModifierLocals::LODIndex, *MeshDescription);
 
-	SkeletalMesh->PostEditChange();
-
 	// update skeleton
 	USkeleton* Skeleton = SkeletalMesh->GetSkeleton();
 	Skeleton->Modify();
@@ -358,6 +356,9 @@ bool USkeletonModifier::CommitSkeletonToSkeletalMesh()
 	{
 		Skeleton->MarkPackageDirty();	
 	}
+	
+	// must be done once the skeleton is up to date
+	SkeletalMesh->PostEditChange();
 
 	return true;
 #else
