CombineMeshInstances: change generation of Voxel LODs to have consistent UVs, by generating UVs at first voxel LOD and allowing them to carry through simplification.
- bring over DoSimplifyMesh function from geo script, really should move this to a centralized location
- add ComputeVoxWrapMeshAutoUV, uses XAtlas, again similar to geo script function
- refactor InitializeAttributes function into separate InitializeNormalsFromAngleThreshold and ProjectAttributes functions
- disable replacement of voxel LODs with approximation LODs if triangle count is larger. Currently no way to communicate this upwards and will break attempts to bake textures for voxel LODs at a higher level.

GeometryScript: remove restriction in CopyMeshToStaticMesh that only allows modifying materials when updating LOD0. There is no reason for this (howerver updating materials when writing hires LOD is still prevented)

diff --git a/Engine/Plugins/Experimental/GeometryScripting/Source/GeometryScriptingCore/Private/MeshAssetFunctions.cpp b/Engine/Plugins/Experimental/GeometryScripting/Source/GeometryScriptingCore/Private/MeshAssetFunctions.cpp
--- a/Engine/Plugins/Experimental/GeometryScripting/Source/GeometryScriptingCore/Private/MeshAssetFunctions.cpp
+++ b/Engine/Plugins/Experimental/GeometryScripting/Source/GeometryScriptingCore/Private/MeshAssetFunctions.cpp
@@ -233,11 +233,12 @@ UDynamicMesh*  UGeometryScriptLibrary_StaticMeshFunctions::CopyMeshToStaticMesh(
 
 #if WITH_EDITOR
 
-	int32 UseLODIndex = FMath::Clamp(TargetLOD.LODIndex, 0, 32);
+	int32 UseLODIndex = FMath::Clamp(TargetLOD.LODIndex, 0, MAX_STATIC_MESH_LODS);
 
-	if (Options.bReplaceMaterials && UseLODIndex != 0)
+	// currently material updates are only applied when writing LODs
+	if (Options.bReplaceMaterials && TargetLOD.bWriteHiResSource)
 	{
-		UE::Geometry::AppendError(Debug, EGeometryScriptErrorType::InvalidInputs, LOCTEXT("CopyMeshToStaticMesh_InvalidOptions1", "CopyMeshToStaticMesh: Can only Replace Materials when updating LOD0"));
+		UE::Geometry::AppendError(Debug, EGeometryScriptErrorType::InvalidInputs, LOCTEXT("CopyMeshToStaticMesh_InvalidOptions1", "CopyMeshToStaticMesh: Can only Replace Materials when updating LODs"));
 		return FromDynamicMesh;
 	}
 
diff --git a/Engine/Plugins/Experimental/MeshModelingToolsetExp/Source/GeometryProcessingAdapters/Private/CombineMeshInstancesImpl.cpp b/Engine/Plugins/Experimental/MeshModelingToolsetExp/Source/GeometryProcessingAdapters/Private/CombineMeshInstancesImpl.cpp
--- a/Engine/Plugins/Experimental/MeshModelingToolsetExp/Source/GeometryProcessingAdapters/Private/CombineMeshInstancesImpl.cpp
+++ b/Engine/Plugins/Experimental/MeshModelingToolsetExp/Source/GeometryProcessingAdapters/Private/CombineMeshInstancesImpl.cpp
@@ -16,6 +16,7 @@
 #include "DynamicMeshToMeshDescription.h"
 
 #include "DynamicMeshEditor.h"
+#include "Parameterization/DynamicMeshUVEditor.h"
 #include "DynamicMesh/DynamicMeshAttributeSet.h"
 #include "DynamicMesh/MeshNormals.h"
 #include "DynamicMesh/MeshTransforms.h"
@@ -43,6 +44,8 @@
 #include "Generators/FlatTriangulationMeshGenerator.h"
 #include "Operations/ExtrudeMesh.h"
 
+#include "XAtlasWrapper.h"
+
 #include "Physics/CollisionGeometryConversion.h"
 #include "Physics/PhysicsDataCollection.h"
 
@@ -1484,27 +1487,178 @@ static void ComputeSimplifiedVoxWrapMesh(
 
 
 
-static void InitializeAttributes(
+template<typename SimplificationType>
+void DoSimplifyMesh(
+	FDynamicMesh3& EditMesh,
+	int32 TargetTriCount,
+	FMeshProjectionTarget* ProjectionTarget = nullptr,
+	double GeometricTolerance = 0)
+{
+	SimplificationType Simplifier(&EditMesh);
+
+	Simplifier.ProjectionMode = SimplificationType::ETargetProjectionMode::NoProjection;
+	if (ProjectionTarget != nullptr)
+	{
+		Simplifier.SetProjectionTarget(ProjectionTarget);
+	}
+
+	Simplifier.DEBUG_CHECK_LEVEL = 0;
+	Simplifier.bRetainQuadricMemory = true;
+	Simplifier.bAllowSeamCollapse = true;
+	//if (Options.bAllowSeamCollapse)		// always true
+	{
+		Simplifier.SetEdgeFlipTolerance(1.e-5);
+		if (EditMesh.HasAttributes())
+		{
+			EditMesh.Attributes()->SplitAllBowties();	// eliminate any bowties that might have formed on attribute seams.
+		}
+	}
+
+	// do these flags matter here since we are not flipping??
+	EEdgeRefineFlags MeshBoundaryConstraints = EEdgeRefineFlags::NoFlip;
+	EEdgeRefineFlags GroupBorderConstraints = EEdgeRefineFlags::NoConstraint;
+	EEdgeRefineFlags MaterialBorderConstraints = EEdgeRefineFlags::NoConstraint;
+
+	FMeshConstraints Constraints;
+	FMeshConstraintsUtil::ConstrainAllBoundariesAndSeams(Constraints, EditMesh,
+		MeshBoundaryConstraints, GroupBorderConstraints, MaterialBorderConstraints,
+		true, false, true);
+	Simplifier.SetExternalConstraints(MoveTemp(Constraints));
+
+	if (ProjectionTarget != nullptr && GeometricTolerance > 0)
+	{
+		Simplifier.GeometricErrorConstraint = SimplificationType::EGeometricErrorCriteria::PredictedPointToProjectionTarget;
+		Simplifier.GeometricErrorTolerance = GeometricTolerance;
+	}
+
+	Simplifier.SimplifyToTriangleCount(FMath::Max(1, TargetTriCount));
+
+	EditMesh.CompactInPlace();
+}
+
+
+
+
+static void ComputeVoxWrapMeshAutoUV(FDynamicMesh3& EditMesh)
+{
+	check(EditMesh.IsCompact());
+	check(EditMesh.HasAttributes());
+
+	FDynamicMeshUVEditor UVEditor(&EditMesh, 0, true);
+	FDynamicMeshUVOverlay* UVOverlay = UVEditor.GetOverlay();
+
+	const bool bFixOrientation = false;
+	//const bool bFixOrientation = true;
+	//FDynamicMesh3 FlippedMesh(EMeshComponents::FaceGroups);
+	//FlippedMesh.Copy(Mesh, false, false, false, false);
+	//if (bFixOrientation)
+	//{
+	//	FlippedMesh.ReverseOrientation(false);
+	//}
+
+	int32 NumVertices = EditMesh.VertexCount();
+	TArray<FVector3f> VertexBuffer;
+	VertexBuffer.SetNum(NumVertices);
+	for (int32 k = 0; k < NumVertices; ++k)
+	{
+		VertexBuffer[k] = (FVector3f)EditMesh.GetVertex(k);
+	}
+
+	TArray<int32> IndexBuffer;
+	IndexBuffer.Reserve(EditMesh.TriangleCount() * 3);
+	for (FIndex3i Triangle : EditMesh.TrianglesItr())
+	{
+		IndexBuffer.Add(Triangle.A);
+		IndexBuffer.Add(Triangle.B);
+		IndexBuffer.Add(Triangle.C);
+	}
+
+	TArray<FVector2D> UVVertexBuffer;
+	TArray<int32>     UVIndexBuffer;
+	TArray<int32>     VertexRemapArray; // This maps the UV vertices to the original position vertices.  Note multiple UV vertices might share the same positional vertex (due to UV boundaries)
+	XAtlasWrapper::XAtlasChartOptions ChartOptions;
+	ChartOptions.MaxIterations = 1;
+	XAtlasWrapper::XAtlasPackOptions PackOptions;
+	bool bSuccess = XAtlasWrapper::ComputeUVs(IndexBuffer, VertexBuffer, ChartOptions, PackOptions,
+		UVVertexBuffer, UVIndexBuffer, VertexRemapArray);
+	if (bSuccess == false)
+	{
+		return;
+	}
+
+	UVOverlay->ClearElements();
+
+	int32 NumUVs = UVVertexBuffer.Num();
+	TArray<int32> UVOffsetToElID;  UVOffsetToElID.Reserve(NumUVs);
+	for (int32 i = 0; i < NumUVs; ++i)
+	{
+		FVector2D UV = UVVertexBuffer[i];
+		const int32 VertOffset = VertexRemapArray[i];		// The associated VertID in the dynamic mesh
+		const int32 NewID = UVOverlay->AppendElement(FVector2f(UV));		// add the UV to the mesh overlay
+		UVOffsetToElID.Add(NewID);
+	}
+
+	int32 NumUVTris = UVIndexBuffer.Num() / 3;
+	for (int32 i = 0; i < NumUVTris; ++i)
+	{
+		int32 t = i * 3;
+		FIndex3i UVTri(UVIndexBuffer[t], UVIndexBuffer[t + 1], UVIndexBuffer[t + 2]);	// The triangle in UV space
+		FIndex3i TriVertIDs;				// the triangle in terms of the VertIDs in the DynamicMesh
+		for (int c = 0; c < 3; ++c)
+		{
+			int32 Offset = VertexRemapArray[UVTri[c]];		// the offset for this vertex in the LinearMesh
+			TriVertIDs[c] = Offset;
+		}
+
+		// NB: this could be slow.. 
+		int32 TriID = EditMesh.FindTriangle(TriVertIDs[0], TriVertIDs[1], TriVertIDs[2]);
+		if (TriID != IndexConstants::InvalidID)
+		{
+			FIndex3i ElTri = (bFixOrientation) ?
+				FIndex3i(UVOffsetToElID[UVTri[1]], UVOffsetToElID[UVTri[0]], UVOffsetToElID[UVTri[2]])
+				: FIndex3i(UVOffsetToElID[UVTri[0]], UVOffsetToElID[UVTri[1]], UVOffsetToElID[UVTri[2]]);
+			UVOverlay->SetTriangle(TriID, ElTri);
+		}
+	}
+}
+
+
+
+static void InitializeNormalsFromAngleThreshold(
 	FDynamicMesh3& TargetMesh,
-	double NormalAngleThreshDeg,
-	bool bProjectAttributes,
-	const FDynamicMesh3* SourceMesh,
-	FDynamicMeshAABBTree3* SourceMeshSpatial)
+	double NormalAngleThreshDeg)
 {
-	TargetMesh.EnableTriangleGroups();
-	TargetMesh.EnableAttributes();
+	if (TargetMesh.HasAttributes() == false)
+	{
+		TargetMesh.EnableAttributes();
+	}
+
 	// recompute normals
 	FMeshNormals::InitializeOverlayTopologyFromOpeningAngle(&TargetMesh, TargetMesh.Attributes()->PrimaryNormals(), NormalAngleThreshDeg);
 	FMeshNormals::QuickRecomputeOverlayNormals(TargetMesh);
+}
+
 
-	if (bProjectAttributes == false || SourceMesh == nullptr || SourceMeshSpatial == nullptr)
+
+static void ProjectAttributes(
+	FDynamicMesh3& TargetMesh,
+	const FDynamicMesh3* SourceMesh,
+	FDynamicMeshAABBTree3* SourceMeshSpatial)
+{
+	if (SourceMesh == nullptr || SourceMeshSpatial == nullptr)
 	{
 		return;
 	}
 
+	TargetMesh.EnableTriangleGroups();
+	if (TargetMesh.HasAttributes() == false)
+	{
+		TargetMesh.EnableAttributes();
+	}
+
 	const FDynamicMeshColorOverlay* SourceColors = nullptr;
 	FDynamicMeshColorOverlay* TargetColors = nullptr;
-	if ( SourceMesh->HasAttributes() && SourceMesh->Attributes()->HasPrimaryColors() )
+	if (SourceMesh->HasAttributes() && SourceMesh->Attributes()->HasPrimaryColors())
 	{
 		SourceColors = SourceMesh->Attributes()->PrimaryColors();
 		TargetMesh.Attributes()->EnablePrimaryColors();
@@ -1513,7 +1667,7 @@ static void InitializeAttributes(
 
 	const FDynamicMeshMaterialAttribute* SourceMaterialID = nullptr;
 	FDynamicMeshMaterialAttribute* TargetMaterialID = nullptr;
-	if ( SourceMesh->HasAttributes() && SourceMesh->Attributes()->HasMaterialID() )
+	if (SourceMesh->HasAttributes() && SourceMesh->Attributes()->HasMaterialID())
 	{
 		SourceMaterialID = SourceMesh->Attributes()->GetMaterialID();
 		TargetMesh.Attributes()->EnableMaterialID();
@@ -1542,7 +1696,7 @@ static void InitializeAttributes(
 			int A = TargetColors->AppendElement(Color);
 			int B = TargetColors->AppendElement(Color);
 			int C = TargetColors->AppendElement(Color);
-			TargetColors->SetTriangle(tid, FIndex3i(A,B,C));
+			TargetColors->SetTriangle(tid, FIndex3i(A, B, C));
 		}
 	}
 }
@@ -1550,6 +1704,7 @@ static void InitializeAttributes(
 
 
 
+
 struct FCombinedMeshLOD
 {
 	FDynamicMesh3 Mesh;
@@ -1707,7 +1862,7 @@ void BuildCombinedMesh(
 	using namespace UE::Geometry;
 	bool bVerbose = CVarGeometryCombineMeshInstancesVerbose.GetValueOnGameThread();
 
-	bool bAppendMinimalLOD = false;
+	const bool bAppendMinimalLOD = false;
 
 	int32 NumLODs = CombineOptions.NumLODs;
 	int NumExtraLODs = (bAppendMinimalLOD) ? 1 : 0;
@@ -1911,38 +2066,48 @@ void BuildCombinedMesh(
 
 		int32 MaxTriCount = CombineOptions.VoxWrapMaxTriCountBase;
 		double SimplifyTolerance = CombineOptions.VoxWrapBaseTolerance;
-		for (int32 LODIndex = FirstVoxWrappedIndex; LODIndex < NumLODs; ++LODIndex)
+
+		// Current state of TempBaseVoxWrapMesh is our initial voxel LOD. To ensure
+		// that voxel LODs have compatible UVs (to allow baking), we compute UVs on
+		// the first LOD and allow them to propagate (and currently normals as well)
+		TempBaseVoxWrapMesh.DiscardAttributes();
 		{
 			TRACE_CPUPROFILER_EVENT_SCOPE(SimplifyVoxWrap);
-
-			// using previous-simplified mesh for next level may not be ideal...
-			MeshLODs[LODIndex].Mesh = (LODIndex == FirstVoxWrappedIndex) ?
-				MoveTemp(TempBaseVoxWrapMesh) : MeshLODs[LODIndex-1].Mesh;
-			// need to do this because we projected attributes in previous loop
-			MeshLODs[LODIndex].Mesh.DiscardAttributes();
-
-			ComputeSimplifiedVoxWrapMesh(MeshLODs[LODIndex].Mesh, &SourceVoxWrapMesh, &Spatial, 
+			ComputeSimplifiedVoxWrapMesh(TempBaseVoxWrapMesh, &SourceVoxWrapMesh, &Spatial,
 				SimplifyTolerance, MaxTriCount);
+		}
+		TempBaseVoxWrapMesh.EnableAttributes();
+		InitializeNormalsFromAngleThreshold(TempBaseVoxWrapMesh, CombineOptions.HardNormalAngleDeg);
+		ComputeVoxWrapMeshAutoUV(TempBaseVoxWrapMesh);
+		MeshLODs[FirstVoxWrappedIndex].Mesh = MoveTemp(TempBaseVoxWrapMesh);
 
-			if ( bVerbose )
-			{
-				UE_LOG(LogGeometry, Log, TEXT("         LOD %2d               - Tris %8d Verts %8d - MaxTriCount %4d"), 
-					LODIndex, MeshLODs[LODIndex].Mesh.TriangleCount(), MeshLODs[LODIndex].Mesh.VertexCount(), MaxTriCount);
-			}
+		// iterate simplification criteria to next level
+		SimplifyTolerance *= 1.5;
+		MaxTriCount /= 2;
 
-			InitializeAttributes(MeshLODs[LODIndex].Mesh, CombineOptions.HardNormalAngleDeg,
-				/*bProjectAttributes*/true, &SourceVoxWrapMesh, &Spatial);
+		for (int32 LODIndex = FirstVoxWrappedIndex+1; LODIndex < NumLODs; ++LODIndex)
+		{
+			// need to simplify from previous level to preserve UVs/etc
+			MeshLODs[LODIndex].Mesh = MeshLODs[LODIndex-1].Mesh;
+
+			DoSimplifyMesh<FAttrMeshSimplification>(MeshLODs[LODIndex].Mesh, MaxTriCount, nullptr, SimplifyTolerance);
 
 			SimplifyTolerance *= 1.5;
 			MaxTriCount /= 2;
 		}
-	}
 
+		// Project colors and materials after mesh simplification to avoid constraining it.
+		// If they /should/ constrain simplification, then they should be projected onto the first
+		// mesh (TempBaseVoxWrapMesh above) and they will automatically transfer
+		for (int32 LODIndex = FirstVoxWrappedIndex; LODIndex < NumLODs; ++LODIndex)
+		{
+			ProjectAttributes(MeshLODs[LODIndex].Mesh, &SourceVoxWrapMesh, &Spatial);
+		}
+	}
 
 	// wait...
 	UE::Tasks::Wait(PendingRemoveHiddenTasks);
 
-
 	if (bAppendMinimalLOD)
 	{
 		check(false);		// for future use
@@ -1968,14 +2133,17 @@ void BuildCombinedMesh(
 
 
 
-
+	// can't replace voxel LODs if we are generating UVs for them!
+	// (no way to communicate this upwards...)
+	//int MaxReplaceLOD = MeshLODs.Num();
+	int MaxReplaceLOD = FirstVoxWrappedIndex;
 	for (int32 LODLevel = 0; LODLevel < MeshLODs.Num(); ++LODLevel)
 	{
 		FDynamicMesh3 LODMesh = MoveTemp(MeshLODs[LODLevel].Mesh);
 
 		// If we ended up larger than the mesh in the previous LOD, we should use that instead!
 		// This can happen particular with VoxWrap LODs
-		if (LODLevel > 0)
+		if (LODLevel > 0 && LODLevel < MaxReplaceLOD)
 		{
 			if (LODMesh.TriangleCount() > CombinedMeshLODs.Last().TriangleCount())
 			{
