Add Geometry Selection visualization to Normals tool

diff --git a/Engine/Plugins/Experimental/MeshModelingToolsetExp/Source/MeshModelingToolsExp/Private/EditNormalsTool.cpp b/Engine/Plugins/Experimental/MeshModelingToolsetExp/Source/MeshModelingToolsExp/Private/EditNormalsTool.cpp
--- a/Engine/Plugins/Experimental/MeshModelingToolsetExp/Source/MeshModelingToolsExp/Private/EditNormalsTool.cpp
+++ b/Engine/Plugins/Experimental/MeshModelingToolsetExp/Source/MeshModelingToolsExp/Private/EditNormalsTool.cpp
@@ -19,6 +19,7 @@
 #include "Components/StaticMeshComponent.h"
 
 #include "GroupTopology.h"
+#include "Selection/ToolSelectionUtil.h"
 #include "Selection/StoredMeshSelectionUtil.h"
 #include "Selections/GeometrySelectionUtil.h"
 #include "SelectionSet.h"
@@ -108,10 +109,14 @@ void UEditNormalsTool::Setup()
 
 	BasicProperties = NewObject<UEditNormalsToolProperties>(this, TEXT("Mesh Normals Settings"));
 	BasicProperties->RestoreProperties(this);
+	BasicProperties->bToolHasSelection = !InputGeometrySelection.IsEmpty();
 	AddToolPropertySource(BasicProperties);
 
+	// Note: Changes to this property set currently do NOT trigger recomputation. See UEditNormalsTool::OnPropertyModified
 	AdvancedProperties = NewObject<UEditNormalsAdvancedProperties>(this, TEXT("Advanced Settings"));
 	AdvancedProperties->RestoreProperties(this);
+	AdvancedProperties->bToolHasSelection = !InputGeometrySelection.IsEmpty();
+	AdvancedProperties->bToolHasEdgeSelection = (!InputGeometrySelection.IsEmpty() && InputGeometrySelection.ElementType == EGeometryElementType::Edge);
 	AddToolPropertySource(AdvancedProperties);
 
 	// initialize the PreviewMesh+BackgroundCompute object
@@ -133,36 +138,51 @@ void UEditNormalsTool::Setup()
 		});
 		SetToolPropertySourceEnabled(PolygroupLayerProperties, BasicProperties->SplitNormalMethod == ESplitNormalMethod::FaceGroupID);
 
-		if (!InputGeometrySelection.IsEmpty() &&
-			(InputGeometrySelection.ElementType == EGeometryElementType::Vertex ||
-			 InputGeometrySelection.ElementType == EGeometryElementType::Edge))
+		if (!InputGeometrySelection.IsEmpty())
 		{
-			// Convert to a Triangle+Vertex selection and operate on that. If the input selection has Edge elements then
-			// we do this because users will expect this to behave similarly to a vertex selection containing the
-			// touched vertices. If the input selection has Vertex elements then we do this to simplify the
-			// implementation (the reason being that we'd want to use EnumeratePolygroupSelectionVertices but this
-			// visits all vertices in any group touching the corner, which isn't what we want for building the overlay
-			// selection from a Polygroup Topology)
-
-			TriangleVertexGeometrySelection.InitializeTypes(EGeometryElementType::Vertex, EGeometryTopologyType::Triangle);
-
-			if (InputGeometrySelection.TopologyType == EGeometryTopologyType::Polygroup)
-			{
-				const FGroupTopology GroupTopology(OriginalDynamicMeshes[0].Get(),
-					ActiveGroupSet.IsValid() ? ActiveGroupSet->PolygroupAttrib : nullptr,
-					true);
-
-				ensure(ConvertSelection(*OriginalDynamicMeshes[0], &GroupTopology, InputGeometrySelection, TriangleVertexGeometrySelection));
-			}
-			else
+			// Setup input geometry selection visualization
+			FTransform ApplyTransform = (FTransform)UE::ToolTarget::GetLocalToWorldTransform(Targets[0]);
+
+			// Compute group topology if the selection has Polygroup topology, and do nothing otherwise
+			FGroupTopology GroupTopology(OriginalDynamicMeshes[0].Get(),
+				ActiveGroupSet.IsValid() ? ActiveGroupSet->PolygroupAttrib : nullptr,
+				InputGeometrySelection.TopologyType == EGeometryTopologyType::Polygroup);
+
+			InputGeometrySelectionRenderer.Initialize(InputGeometrySelection,
+				*OriginalDynamicMeshes[0],
+				&GroupTopology,
+				&ApplyTransform);
+
+			// Special case handling for edge and vertex element selections
+			if (InputGeometrySelection.ElementType == EGeometryElementType::Vertex ||
+				InputGeometrySelection.ElementType == EGeometryElementType::Edge)
 			{
-				ensure(ConvertSelection(*OriginalDynamicMeshes[0], nullptr, InputGeometrySelection, TriangleVertexGeometrySelection));
+				// Convert to a Triangle+Vertex selection and operate on that. If the input selection has Edge elements then
+				// we do this because users will expect this to behave similarly to a vertex selection containing the
+				// touched vertices. If the input selection has Vertex elements then we do this to simplify the
+				// implementation (the reason being that we'd want to use EnumeratePolygroupSelectionVertices but this
+				// visits all vertices in any group touching the corner, which isn't what we want for building the overlay
+				// selection from a Polygroup Topology)
+
+				TriangleVertexGeometrySelection.InitializeTypes(EGeometryElementType::Vertex, EGeometryTopologyType::Triangle);
+				bool bSuccess = ConvertSelection(
+					*OriginalDynamicMeshes[0],
+					&GroupTopology, // This argument is ignored if InputGeometrySelection has Triangle topology
+					InputGeometrySelection,
+					TriangleVertexGeometrySelection);
+				ensure(bSuccess == true);
+				ensure(!TriangleVertexGeometrySelection.IsEmpty());
+
+				// Setup debug visualization
+				TriangleVertexGeometrySelectionRenderer.Initialize(
+					TriangleVertexGeometrySelection,
+					*OriginalDynamicMeshes[0],
+					nullptr,
+					&ApplyTransform);
 			}
-
-			ensure(!TriangleVertexGeometrySelection.IsEmpty());
 		}
 	}
-	
+
 	for (UMeshOpPreviewWithBackgroundCompute* Preview : Previews)
 	{
 		Preview->InvalidateResult();
@@ -318,6 +338,18 @@ TUniquePtr<FDynamicMeshOperator> UEditNormalsOperatorFactory::MakeNewOperator()
 
 void UEditNormalsTool::Render(IToolsContextRenderAPI* RenderAPI)
 {
+	if (InputGeometrySelection.IsEmpty() == false && AdvancedProperties->bShowSelection)
+	{
+		if (InputGeometrySelection.ElementType == EGeometryElementType::Edge &&
+			AdvancedProperties->bShowEdgeSelectionAsVertexSelection)
+		{
+			TriangleVertexGeometrySelectionRenderer.Render(RenderAPI);
+		}
+		else
+		{
+			InputGeometrySelectionRenderer.Render(RenderAPI);
+		}
+	}
 }
 
 void UEditNormalsTool::OnTick(float DeltaTime)
@@ -343,6 +375,11 @@ void UEditNormalsTool::PostEditChangeProperty(FPropertyChangedEvent& PropertyCha
 
 void UEditNormalsTool::OnPropertyModified(UObject* PropertySet, FProperty* Property)
 {
+	if (PropertySet == AdvancedProperties)
+	{
+		return;
+	}
+
 	UpdateNumPreviews();
 	for (UMeshOpPreviewWithBackgroundCompute* Preview : Previews)
 	{
diff --git a/Engine/Plugins/Experimental/MeshModelingToolsetExp/Source/MeshModelingToolsExp/Public/EditNormalsTool.h b/Engine/Plugins/Experimental/MeshModelingToolsetExp/Source/MeshModelingToolsExp/Public/EditNormalsTool.h
--- a/Engine/Plugins/Experimental/MeshModelingToolsetExp/Source/MeshModelingToolsExp/Public/EditNormalsTool.h
+++ b/Engine/Plugins/Experimental/MeshModelingToolsetExp/Source/MeshModelingToolsExp/Public/EditNormalsTool.h
@@ -10,6 +10,7 @@
 #include "DynamicMesh/DynamicMesh3.h"
 #include "PropertySets/PolygroupLayersProperties.h"
 #include "Polygroups/PolygroupSet.h"
+#include "Selection/ToolSelectionUtil.h"
 #include "Selections/GeometrySelection.h"
 #include "EditNormalsTool.generated.h"
 
@@ -54,17 +55,17 @@ public:
 	}
 
 	/** Recompute all mesh normals */
-	UPROPERTY(EditAnywhere, Category = NormalsCalculation, meta = (EditCondition = "SplitNormalMethod == ESplitNormalMethod::UseExistingTopology"))
+	UPROPERTY(EditAnywhere, Category = NormalsCalculation,
+		meta = (EditCondition = "SplitNormalMethod == ESplitNormalMethod::UseExistingTopology && !bToolHasSelection", HideEditConditionToggle))
 	bool bRecomputeNormals;
 
 	/** Choose the method for computing vertex normals */
 	UPROPERTY(EditAnywhere, Category = NormalsCalculation)
 	ENormalCalculationMethod NormalCalculationMethod;
 
-	// TODO Disable bFixInconsistentNormals when we have a selection
-
 	/** For meshes with inconsistent triangle orientations/normals, flip as needed to make the normals consistent */
-	UPROPERTY(EditAnywhere, Category = NormalsCalculation)
+	UPROPERTY(EditAnywhere, Category = NormalsCalculation,
+		meta = (EditCondition = "!bToolHasSelection", HideEditConditionToggle))
 	bool bFixInconsistentNormals;
 
 	/** Invert (flip) all mesh normals and associated triangle orientations */
@@ -82,6 +83,13 @@ public:
 	/** Assign separate normals at 'sharp' vertices, for example, at the tip of a cone */
 	UPROPERTY(EditAnywhere, Category = NormalsTopology, meta = (EditCondition = "SplitNormalMethod == ESplitNormalMethod::FaceNormalThreshold"))
 	bool bAllowSharpVertices;
+
+	//
+	// The following are not user visible
+	//
+
+	UPROPERTY(meta = (TransientToolProperty))
+	bool bToolHasSelection;
 };
 
 
@@ -97,6 +105,26 @@ class MESHMODELINGTOOLSEXP_API UEditNormalsAdvancedProperties : public UInteract
 
 public:
 	UEditNormalsAdvancedProperties();
+
+	/** Render the geometry selection */
+	UPROPERTY(EditAnywhere, Category = "Selection Visualization",
+		meta = (EditCondition = "bToolHasSelection", HideEditConditionToggle, EditConditionHides))
+	bool bShowSelection = false;
+
+	/** This tool treats edge selections as vertex selections. Enable this to show the edited vertices */
+	UPROPERTY(EditAnywhere, Category="Selection Visualization", AdvancedDisplay, DisplayName="Show Vertex Selection",
+		meta = (EditCondition = "bToolHasSelection && bToolHasEdgeSelection", HideEditConditionToggle, EditConditionHides))
+	bool bShowEdgeSelectionAsVertexSelection = false;
+
+	//
+	// The following are not user visible
+	//
+
+	UPROPERTY(meta = (TransientToolProperty))
+	bool bToolHasSelection;
+
+	UPROPERTY(meta = (TransientToolProperty))
+	bool bToolHasEdgeSelection;
 };
 
 
@@ -189,6 +217,7 @@ protected:
 	// The geometry selection that the user started the tool with. If the selection is empty we operate on the whole
 	// mesh, if its not empty we only edit the overlay elements implied by the selection.
 	UE::Geometry::FGeometrySelection InputGeometrySelection;
+	UE::Geometry::FSelectionRenderHelper InputGeometrySelectionRenderer; // Default selection visualization
 
 	// Cache the input polygroup set which was used to start the tool. We do this because users can change the
 	// polygroup referenced by the operator while using the tool.
@@ -196,6 +225,6 @@ protected:
 
 	// If the user starts the tool with an edge selection we convert it to a vertex selection with triangle topology
 	// and store it here, we do this since we expect users to want vertex and edge selections to behave similarly.
-	// TODO Add a debug visualization option that shows this selection to make this detail clear to advanced users
 	UE::Geometry::FGeometrySelection TriangleVertexGeometrySelection;
-};
\ No newline at end of file
+	UE::Geometry::FSelectionRenderHelper TriangleVertexGeometrySelectionRenderer; // Debug selection visualization, off by default
+};
diff --git a/Engine/Plugins/Runtime/GeometryProcessing/Source/DynamicMesh/Private/Selections/GeometrySelectionUtil.cpp b/Engine/Plugins/Runtime/GeometryProcessing/Source/DynamicMesh/Private/Selections/GeometrySelectionUtil.cpp
--- a/Engine/Plugins/Runtime/GeometryProcessing/Source/DynamicMesh/Private/Selections/GeometrySelectionUtil.cpp
+++ b/Engine/Plugins/Runtime/GeometryProcessing/Source/DynamicMesh/Private/Selections/GeometrySelectionUtil.cpp
@@ -1055,13 +1055,11 @@ bool UE::Geometry::ConvertSelection(
 
 		// TODO Add a function which gets the Vids only, to remove the matrix-vector multiplication we just ignore
 		const FTransform Transform = FTransform::Identity;
-		EnumerateTriangleSelectionVertices(FromSelectionIn, Mesh, Transform,
+		return EnumerateTriangleSelectionVertices(FromSelectionIn, Mesh, Transform,
 			[&ToSelectionOut](uint64 Vid, const FVector3d& Unused)
 			{
 				ToSelectionOut.Selection.Add( FGeoSelectionID::MeshVertex((int32)Vid).Encoded() );
 			});
-
-		return true;
 	};
 
 
@@ -1079,13 +1077,11 @@ bool UE::Geometry::ConvertSelection(
 
 		// TODO Add a function which gets the Vids only, to remove the matrix-vector multiplication we just ignore
 		const FTransform Transform = FTransform::Identity;
-		EnumeratePolygroupSelectionVertices(FromSelectionIn, Mesh, GroupTopology, Transform,
+		return EnumeratePolygroupSelectionVertices(FromSelectionIn, Mesh, GroupTopology, Transform,
 			[&ToSelectionOut](uint64 Vid, const FVector3d& Unused)
 			{
 				ToSelectionOut.Selection.Add( FGeoSelectionID::MeshVertex((int32)Vid).Encoded() );
 			});
-
-		return true;
 	};
 
 
diff --git a/Engine/Plugins/Runtime/GeometryProcessing/Source/DynamicMesh/Public/Selections/GeometrySelectionUtil.h b/Engine/Plugins/Runtime/GeometryProcessing/Source/DynamicMesh/Public/Selections/GeometrySelectionUtil.h
--- a/Engine/Plugins/Runtime/GeometryProcessing/Source/DynamicMesh/Public/Selections/GeometrySelectionUtil.h
+++ b/Engine/Plugins/Runtime/GeometryProcessing/Source/DynamicMesh/Public/Selections/GeometrySelectionUtil.h
@@ -213,13 +213,14 @@ DYNAMICMESH_API bool InitializeSelectionFromTriangles(
  *   ----------------------------------------------------------------
  *   Triangle Vertex      1       .       .      .       .       .   
  *   Triangle Edge        1       1       .      .       .       .   
- *   Triangle Face        1       1       1      4       3       2   
- *   Polygroup Vertex     6       .       .      1       .       .   
- *   Polygroup Edge       5       .       .      .       1       .   
+ *   Triangle Face        1       1       1      4#      3#      2#  
+ *   Polygroup Vertex     6#      .       .      1       .       .   
+ *   Polygroup Edge       5#      .       .      .       1       .   
  *   Polygroup Face       .       .       .      .       .       1   
  *   ================================================================
  *
- *   .  These conversions are not implemented... yet?
+ *   Key:
+ *   .  These conversions are not implemented... yet? GroupTopology is ignored
  *   1  supported. The implementation is obvious/unambiguous
  *   2  supported. Polygroup faces containing any input triangle are selected
  *   3  supported. Polygroup edges containing any input triangle edge are selected, but
@@ -227,6 +228,7 @@ DYNAMICMESH_API bool InitializeSelectionFromTriangles(
  *   4  supported. Polygroup corners coinciding with any input triangle vertex are selected
  *   5  supported. All mesh vertices along the polygroup edge are selected
  *   6  supported. All mesh vertices coinciding with polygroup corners are selected
+ *   #  indicates GroupTopology must not be null for this combination. If this symbol is missing GroupTopology is ignored
  *
  * @return true if conversion is supported and was computed successfully, return false otherwise
  */
diff --git a/Engine/Plugins/Runtime/MeshModelingToolset/Source/ModelingComponents/Private/Selection/DynamicMeshSelector.cpp b/Engine/Plugins/Runtime/MeshModelingToolset/Source/ModelingComponents/Private/Selection/DynamicMeshSelector.cpp
--- a/Engine/Plugins/Runtime/MeshModelingToolset/Source/ModelingComponents/Private/Selection/DynamicMeshSelector.cpp
+++ b/Engine/Plugins/Runtime/MeshModelingToolset/Source/ModelingComponents/Private/Selection/DynamicMeshSelector.cpp
@@ -13,6 +13,7 @@
 #include "Spatial/SegmentTree3.h"
 #include "ToolSceneQueriesUtil.h"
 #include "Selection/DynamicMeshPolygroupTransformer.h"
+#include "Selection/ToolSelectionUtil.h"
 #include "ToolDataVisualizer.h"
 
 using namespace UE::Geometry;
@@ -601,35 +602,17 @@ void FBaseDynamicMeshSelector::AccumulateSelectionBounds(const FGeometrySelectio
 
 void FBaseDynamicMeshSelector::AccumulateSelectionElements(const FGeometrySelection& Selection, FGeometrySelectionElements& Elements, bool bTransformToWorld, bool bIsForPreview)
 {
-	FTransform UseWorldTransform = GetLocalToWorldTransform();
+	const FTransform UseWorldTransform = GetLocalToWorldTransform();
 	const FTransform* ApplyTransform = (bTransformToWorld) ? &UseWorldTransform : nullptr;
-
-	bool bMapFacesToEdges = bIsForPreview;
-	if (Selection.TopologyType == EGeometryTopologyType::Polygroup)
-	{
-		const FGroupTopology* Topology = GetGroupTopology();
-		TargetMesh->ProcessMesh([&](const UE::Geometry::FDynamicMesh3& SourceMesh)
-		{
-			UE::Geometry::EnumeratePolygroupSelectionElements(Selection, SourceMesh, Topology,
-				[&](uint32 VertexID, const FVector3d& Point) { Elements.Points.Add(Point); },
-				[&](uint32 EdgeID, const FSegment3d& Segment) { Elements.Segments.Add(Segment); },
-				[&](uint32 TriangleID, const FTriangle3d& Triangle) { Elements.Triangles.Add(Triangle); },
-				ApplyTransform, bMapFacesToEdges
-			);
-		});
-	}
-	else
-	{
-		TargetMesh->ProcessMesh([&](const UE::Geometry::FDynamicMesh3& SourceMesh)
-		{
-			UE::Geometry::EnumerateTriangleSelectionElements(Selection, SourceMesh, 
-				[&](uint32 VertexID, const FVector3d& Point) { Elements.Points.Add(Point); },
-				[&](uint32 EdgeID, const FSegment3d& Segment) { Elements.Segments.Add(Segment); },
-				[&](uint32 TriangleID, const FTriangle3d& Triangle) { Elements.Triangles.Add(Triangle); },
-				ApplyTransform, bMapFacesToEdges
-			);
-		});
-	}
+	const bool bMapFacesToEdges = bIsForPreview;
+
+	ToolSelectionUtil::AccumulateSelectionElements(
+		Elements,
+		Selection,
+		TargetMesh->GetMeshRef(),
+		Selection.TopologyType == EGeometryTopologyType::Polygroup ? GetGroupTopology() : nullptr,
+		ApplyTransform,
+		bMapFacesToEdges);
 }
 
 
diff --git a/Engine/Plugins/Runtime/MeshModelingToolset/Source/ModelingComponents/Private/Selection/GeometrySelectionManager.cpp b/Engine/Plugins/Runtime/MeshModelingToolset/Source/ModelingComponents/Private/Selection/GeometrySelectionManager.cpp
--- a/Engine/Plugins/Runtime/MeshModelingToolset/Source/ModelingComponents/Private/Selection/GeometrySelectionManager.cpp
+++ b/Engine/Plugins/Runtime/MeshModelingToolset/Source/ModelingComponents/Private/Selection/GeometrySelectionManager.cpp
@@ -3,6 +3,7 @@
 #include "Selection/GeometrySelectionManager.h"
 #include "Engine/Engine.h"
 #include "Selection/DynamicMeshSelector.h"
+#include "Selection/ToolSelectionUtil.h"
 #include "Selection/SelectionEditInteractiveCommand.h"
 #include "InteractiveToolsContext.h"
 #include "InteractiveToolManager.h"
@@ -10,12 +11,6 @@
 #include "ToolDataVisualizer.h"
 #include "Selections/GeometrySelectionUtil.h"
 
-// for debug drawing
-#include "SceneManagement.h"
-#include "SceneView.h"
-#include "DynamicMeshBuilder.h"
-#include "ToolSetupUtil.h"
-
 #include UE_INLINE_GENERATED_CPP_BY_NAME(GeometrySelectionManager)
 
 static TAutoConsoleVariable<int32> CVarGeometrySelectionManager_FullSelectionHoverHighlights(
@@ -1278,87 +1273,11 @@ void UGeometrySelectionManager::DebugRender(IToolsContextRenderAPI* RenderAPI)
 	//const UMaterialInterface* TriangleMaterial = ToolSetupUtil::GetSelectionMaterial(FLinearColor(1.0f, 0, 0, 0.5f), nullptr, 0.5f);
 
 	RebuildSelectionRenderCaches();
-
 	for ( const FGeometrySelectionElements& Elements : CachedSelectionRenderElements )
 	{
-		// batch render all the triangles, vastly more efficient than drawing one by one!
-
-		FPrimitiveDrawInterface* CurrentPDI = RenderAPI->GetPrimitiveDrawInterface();
-		FDynamicMeshBuilder MeshBuilder(CurrentPDI->View->GetFeatureLevel());
-		int32 DepthPriority = SDPG_World; // SDPG_Foreground;  // SDPG_World
-		FVector2f UVs[3] = { FVector2f(0,0), FVector2f(0,1), FVector2f(1,1)	};
-		FVector3f Normal = FVector3f(0, 0, 1);
-		FVector3f Tangent = FVector3f(1, 0, 0);	
-		for (const FTriangle3d& Triangle : Elements.Triangles)
-		{
-			int32 V0 = MeshBuilder.AddVertex(FDynamicMeshVertex((FVector3f)Triangle.V[0], Tangent, Normal, UVs[0], FColor::White));
-			int32 V1 = MeshBuilder.AddVertex(FDynamicMeshVertex((FVector3f)Triangle.V[1], Tangent, Normal, UVs[1], FColor::White));
-			int32 V2 = MeshBuilder.AddVertex(FDynamicMeshVertex((FVector3f)Triangle.V[2], Tangent, Normal, UVs[2], FColor::White));
-			MeshBuilder.AddTriangle(V0, V1, V2);
-		}
-		//FMaterialRenderProxy* MaterialRenderProxy = TriangleMaterial->GetRenderProxy();		// currently does not work, material does not render
-		FMaterialRenderProxy* MaterialRenderProxy = GEngine->ConstraintLimitMaterialX->GetRenderProxy();
-		MeshBuilder.Draw(CurrentPDI, FMatrix::Identity, MaterialRenderProxy, DepthPriority, false, false);
-
-
-		FToolDataVisualizer Visualizer;
-		Visualizer.bDepthTested = false;
-		Visualizer.BeginFrame(RenderAPI);
-
-		Visualizer.SetLineParameters(FLinearColor(0,0.3f,0.95f,1), 3.0f);
-		for (const FSegment3d& Segment : Elements.Segments)
-		{
-			Visualizer.DrawLine(Segment.StartPoint(), Segment.EndPoint());
-		}
-
-		Visualizer.SetPointParameters(FLinearColor(0,0.3f,0.95f,1), 10.0f);
-		for (const FVector3d& Point : Elements.Points)
-		{
-			Visualizer.DrawPoint(Point);
-		}
-
-		Visualizer.EndFrame();
+		ToolSelectionUtil::DebugRenderGeometrySelectionElements(RenderAPI, Elements, false);
 	}
-
-	// draw selection preview
-	{
-		FPrimitiveDrawInterface* CurrentPDI = RenderAPI->GetPrimitiveDrawInterface();
-		FDynamicMeshBuilder MeshBuilder(CurrentPDI->View->GetFeatureLevel());
-		int32 DepthPriority = SDPG_World; // SDPG_Foreground;  // SDPG_World
-		FVector2f UVs[3] = { FVector2f(0,0), FVector2f(0,1), FVector2f(1,1) };
-		FVector3f Normal = FVector3f(0, 0, 1);
-		FVector3f Tangent = FVector3f(1, 0, 0);
-		for (const FTriangle3d& Triangle : CachedPreviewRenderElements.Triangles)
-		{
-			int32 V0 = MeshBuilder.AddVertex(FDynamicMeshVertex((FVector3f)Triangle.V[0], Tangent, Normal, UVs[0], FColor::White));
-			int32 V1 = MeshBuilder.AddVertex(FDynamicMeshVertex((FVector3f)Triangle.V[1], Tangent, Normal, UVs[1], FColor::White));
-			int32 V2 = MeshBuilder.AddVertex(FDynamicMeshVertex((FVector3f)Triangle.V[2], Tangent, Normal, UVs[2], FColor::White));
-			MeshBuilder.AddTriangle(V0, V1, V2);
-		}
-		//FMaterialRenderProxy* MaterialRenderProxy = TriangleMaterial->GetRenderProxy();		// currently does not work, material does not render
-		FMaterialRenderProxy* MaterialRenderProxy = GEngine->ConstraintLimitMaterialX->GetRenderProxy();
-		MeshBuilder.Draw(CurrentPDI, FMatrix::Identity, MaterialRenderProxy, DepthPriority, false, false);
-
-
-		FToolDataVisualizer Visualizer;
-		Visualizer.bDepthTested = false;
-		Visualizer.BeginFrame(RenderAPI);
-
-		Visualizer.SetLineParameters(FLinearColor(1, 1, 0, 1), 1.0f);
-		for (const FSegment3d& Segment : CachedPreviewRenderElements.Segments)
-		{
-			Visualizer.DrawLine(Segment.StartPoint(), Segment.EndPoint());
-		}
-
-		Visualizer.SetPointParameters(FLinearColor(1, 1, 0, 1), 5.0f);
-		for (const FVector3d& Point : CachedPreviewRenderElements.Points)
-		{
-			Visualizer.DrawPoint(Point);
-		}
-
-		Visualizer.EndFrame();
-	}
-
+	ToolSelectionUtil::DebugRenderGeometrySelectionElements(RenderAPI, CachedPreviewRenderElements, true);
 }
 
 
diff --git a/Engine/Plugins/Runtime/MeshModelingToolset/Source/ModelingComponents/Private/Selection/ToolSelectionUtil.cpp b/Engine/Plugins/Runtime/MeshModelingToolset/Source/ModelingComponents/Private/Selection/ToolSelectionUtil.cpp
--- a/Engine/Plugins/Runtime/MeshModelingToolset/Source/ModelingComponents/Private/Selection/ToolSelectionUtil.cpp
+++ b/Engine/Plugins/Runtime/MeshModelingToolset/Source/ModelingComponents/Private/Selection/ToolSelectionUtil.cpp
@@ -5,6 +5,23 @@
 #include "InteractiveToolManager.h"
 #include "GameFramework/Actor.h"
 
+#include "TriangleTypes.h"
+#include "SegmentTypes.h"
+#include "GroupTopology.h"
+#include "ToolDataVisualizer.h"
+#include "DynamicMeshBuilder.h"
+#include "DynamicMesh/DynamicMesh3.h"
+#include "Selections/GeometrySelection.h"
+#include "Selections/GeometrySelectionUtil.h"
+
+// For debug drawing
+#include "Engine/Engine.h"
+#include "Materials/MaterialInstanceDynamic.h"
+#include "SceneManagement.h"
+#include "SceneView.h"
+
+using namespace UE::Geometry;
+
 void ToolSelectionUtil::SetNewActorSelection(UInteractiveToolManager* ToolManager, AActor* Actor)
 {
 	FSelectedOjectsChangeList NewSelection;
@@ -23,4 +40,102 @@ void ToolSelectionUtil::SetNewActorSelection(UInteractiveToolManager* ToolManage
 		NewSelection.Actors.Add(Actor);
 	}
 	ToolManager->RequestSelectionChange(NewSelection);
-}
\ No newline at end of file
+}
+
+MODELINGCOMPONENTS_API bool ToolSelectionUtil::AccumulateSelectionElements(
+		FGeometrySelectionElements& Elements,
+		const FGeometrySelection& Selection,
+		const FDynamicMesh3& SourceMesh,
+		const FGroupTopology* GroupTopology,
+		const FTransform* ApplyTransform,
+		bool bMapFacesToEdges)
+{
+	auto AddPoint = [&Elements](uint32 Vid, const FVector3d& Point) { Elements.Points.Add(Point); };
+	auto AddSegment = [&Elements](uint32 Eid, const FSegment3d& Segment) { Elements.Segments.Add(Segment); };
+	auto AddTriangle = [&Elements](uint32 Tid, const FTriangle3d& Triangle) { Elements.Triangles.Add(Triangle); };
+
+	if (Selection.TopologyType == EGeometryTopologyType::Polygroup)
+	{
+		if (GroupTopology)
+		{
+			return EnumeratePolygroupSelectionElements(Selection, SourceMesh, GroupTopology,
+				AddPoint, AddSegment, AddTriangle, ApplyTransform, bMapFacesToEdges);
+		}
+
+		const FGroupTopology ComputedGroupTopology(&SourceMesh, true);
+		return EnumeratePolygroupSelectionElements(Selection, SourceMesh, &ComputedGroupTopology,
+			AddPoint, AddSegment, AddTriangle, ApplyTransform, bMapFacesToEdges);
+	}
+
+	return EnumerateTriangleSelectionElements(Selection, SourceMesh, 
+			AddPoint, AddSegment, AddTriangle, ApplyTransform, bMapFacesToEdges);
+}
+
+void ToolSelectionUtil::DebugRenderGeometrySelectionElements(
+	IToolsContextRenderAPI* RenderAPI,
+	const FGeometrySelectionElements& Elements,
+	bool bIsPreview)
+{
+	FPrimitiveDrawInterface* CurrentPDI = RenderAPI->GetPrimitiveDrawInterface();
+
+	// batch render all the triangles, vastly more efficient than drawing one by one!
+	FDynamicMeshBuilder MeshBuilder(CurrentPDI->View->GetFeatureLevel());
+	int32 DepthPriority = SDPG_World; // SDPG_Foreground;  // SDPG_World
+	FVector2f UVs[3] = { FVector2f(0,0), FVector2f(0,1), FVector2f(1,1)	};
+	FVector3f Normal = FVector3f(0, 0, 1);
+	FVector3f Tangent = FVector3f(1, 0, 0);	
+	for (const FTriangle3d& Triangle : Elements.Triangles)
+	{
+		int32 V0 = MeshBuilder.AddVertex(FDynamicMeshVertex((FVector3f)Triangle.V[0], Tangent, Normal, UVs[0], FColor::White));
+		int32 V1 = MeshBuilder.AddVertex(FDynamicMeshVertex((FVector3f)Triangle.V[1], Tangent, Normal, UVs[1], FColor::White));
+		int32 V2 = MeshBuilder.AddVertex(FDynamicMeshVertex((FVector3f)Triangle.V[2], Tangent, Normal, UVs[2], FColor::White));
+		MeshBuilder.AddTriangle(V0, V1, V2);
+	}
+	//FMaterialRenderProxy* MaterialRenderProxy = TriangleMaterial->GetRenderProxy();		// currently does not work, material does not render
+	FMaterialRenderProxy* MaterialRenderProxy = GEngine->ConstraintLimitMaterialX->GetRenderProxy();
+	MeshBuilder.Draw(CurrentPDI, FMatrix::Identity, MaterialRenderProxy, DepthPriority, false, false);
+
+
+	FToolDataVisualizer Visualizer;
+	Visualizer.bDepthTested = false;
+	Visualizer.BeginFrame(RenderAPI);
+
+	if (bIsPreview)
+	{
+		Visualizer.SetLineParameters(FLinearColor(1, 1, 0, 1), 1.0f);
+		Visualizer.SetPointParameters(FLinearColor(1, 1, 0, 1), 5.0f);
+	}
+	else
+	{
+		Visualizer.SetLineParameters(FLinearColor(0,0.3f,0.95f,1), 3.0f);
+		Visualizer.SetPointParameters(FLinearColor(0,0.3f,0.95f,1), 10.0f);
+	}
+
+	for (const FSegment3d& Segment : Elements.Segments)
+	{
+		Visualizer.DrawLine(Segment.StartPoint(), Segment.EndPoint());
+	}
+	for (const FVector3d& Point : Elements.Points)
+	{
+		Visualizer.DrawPoint(Point);
+	}
+
+	Visualizer.EndFrame();
+}
+
+void FSelectionRenderHelper::Initialize(const FGeometrySelection& Selection,
+	const FDynamicMesh3& SourceMesh,
+	const FGroupTopology* Topology,
+	const FTransform* ApplyTransform)
+{
+	bool bSuccess = ToolSelectionUtil::AccumulateSelectionElements(Elements, Selection, SourceMesh, Topology, ApplyTransform, false);
+	ensure(bSuccess);
+}
+
+void FSelectionRenderHelper::Render(IToolsContextRenderAPI* RenderAPI) const
+{
+	// TODO Using PDI here does not scale well for large selections. We are going to have to reimplement this with
+	// something that uses (eg) a custom LineSetComponent / PointSetComponent / TriangleSetComponent depending on
+	// the selection type
+	ToolSelectionUtil::DebugRenderGeometrySelectionElements(RenderAPI, Elements, false);
+}
diff --git a/Engine/Plugins/Runtime/MeshModelingToolset/Source/ModelingComponents/Public/Selection/ToolSelectionUtil.h b/Engine/Plugins/Runtime/MeshModelingToolset/Source/ModelingComponents/Public/Selection/ToolSelectionUtil.h
--- a/Engine/Plugins/Runtime/MeshModelingToolset/Source/ModelingComponents/Public/Selection/ToolSelectionUtil.h
+++ b/Engine/Plugins/Runtime/MeshModelingToolset/Source/ModelingComponents/Public/Selection/ToolSelectionUtil.h
@@ -2,11 +2,22 @@
 
 #pragma once
 
+#include "Selections/GeometrySelection.h"
 #include "CoreMinimal.h"
 
 class UInteractiveToolManager;
+class IToolsContextRenderAPI;
 class AActor;
 
+namespace UE::Geometry
+{
+
+class FDynamicMesh3;
+class FGroupTopology;
+struct FGeometrySelection;
+
+}
+
 /**
  * Utility functions for Tool implementations to use when doing selection
  */
@@ -22,4 +33,46 @@ namespace ToolSelectionUtil
 	 * Change the active selection to the given Actors, via given ToolManager. Replaces existing selection.
 	 */
 	MODELINGCOMPONENTS_API void SetNewActorSelection(UInteractiveToolManager* ToolManager, const TArray<AActor*>& Actors);
-}
\ No newline at end of file
+
+	/**
+	 * Add the geometry selection elements corresponding to the given Selection to Elements. This function does not
+	 * reset Elements before adding elements. If the Selection has Polygroup topology then use the given Topology to
+	 * accumulate elements if it isn't null, otherwise compute a FGroupTopology from SourceMesh.TriangleGroups and use
+	 * that to accumulate elements. If the Selection has Triangle topology then the Topology argument is ignored.
+	 * Return false if there was an error and true otherwise
+	 */
+	MODELINGCOMPONENTS_API bool AccumulateSelectionElements(
+			UE::Geometry::FGeometrySelectionElements& Elements,
+			const UE::Geometry::FGeometrySelection& Selection,
+			const UE::Geometry::FDynamicMesh3& SourceMesh,
+			const UE::Geometry::FGroupTopology* Topology = nullptr,
+			const FTransform* ApplyTransform = nullptr,
+			bool bMapFacesToEdges = false);
+
+	/**
+	 * Render the given Elements using FPrimitiveDrawInterface
+	 */
+	MODELINGCOMPONENTS_API void DebugRenderGeometrySelectionElements(
+		IToolsContextRenderAPI* RenderAPI,
+		const UE::Geometry::FGeometrySelectionElements& Elements,
+		bool bIsPreview = false);
+}
+
+namespace UE::Geometry
+{
+
+struct MODELINGCOMPONENTS_API FSelectionRenderHelper
+{
+	void Initialize(const FGeometrySelection& Selection,
+		const FDynamicMesh3& SourceMesh,
+		const FGroupTopology* Topology = nullptr,
+		const FTransform* ApplyTransform = nullptr);
+
+	void Render(IToolsContextRenderAPI* RenderAPI) const;
+
+protected:
+
+	FGeometrySelectionElements Elements;
+};
+
+}
