Fixes Skin Paint tool not working correctly on deformed mesh.

diff --git a/Engine/Plugins/Experimental/MeshModelingToolsetExp/Source/MeshModelingToolsExp/Private/SkinWeightsPaintTool.cpp b/Engine/Plugins/Experimental/MeshModelingToolsetExp/Source/MeshModelingToolsExp/Private/SkinWeightsPaintTool.cpp
--- a/Engine/Plugins/Experimental/MeshModelingToolsetExp/Source/MeshModelingToolsExp/Private/SkinWeightsPaintTool.cpp
+++ b/Engine/Plugins/Experimental/MeshModelingToolsetExp/Source/MeshModelingToolsExp/Private/SkinWeightsPaintTool.cpp
@@ -126,7 +126,7 @@ void FSkinToolDeformer::Initialize(const USkeletalMeshComponent* SkeletalMeshCom
 void FSkinToolDeformer::UpdateVertexDeformation(USkinWeightsPaintTool* Tool)
 {
 	TRACE_CPUPROFILER_EVENT_SCOPE(SkinTool::UpdateDeformationTotal);
-	
+
 	if (VerticesWithModifiedWeights.IsEmpty())
 	{
 		return;
@@ -171,9 +171,27 @@ void FSkinToolDeformer::UpdateVertexDeformation(USkinWeightsPaintTool* Tool)
 		// ensure previous async update is finished before queuing the next one...
 		Tool->TriangleOctreeFuture.Wait();
 
-		TArray<int32>& TrianglesToReinsert = Tool->TriangleToReinsert;
-		const UE::Geometry::FAxisAlignedBox3d QueryBox(Tool->StampLocalPos, Tool->CurrentBrushRadius);
-		Tool->TrianglesOctree.RangeQuery(QueryBox, TrianglesToReinsert);
+		// reusable buffer of triangles to reinsert into the octree
+		TArray<int32>& TrianglesToReinsert = Tool->TrianglesToReinsert;
+		TrianglesToReinsert.Reset();
+
+		// reinsert all triangles containing an updated vertex
+		const FDynamicMesh3* DynamicMesh = PreviewMesh->GetMesh();
+		for (const int32 TriangleID : DynamicMesh->TriangleIndicesItr())
+		{
+			UE::Geometry::FIndex3i TriVerts = DynamicMesh->GetTriangle(TriangleID);
+			bool bIsTriangleAffected = VerticesWithModifiedWeights.Contains(TriVerts[0]);
+			bIsTriangleAffected = VerticesWithModifiedWeights.Contains(TriVerts[1]) ? true : bIsTriangleAffected;
+			bIsTriangleAffected = VerticesWithModifiedWeights.Contains(TriVerts[2]) ? true : bIsTriangleAffected;
+			if (bIsTriangleAffected)
+			{
+				TrianglesToReinsert.Add(TriangleID);
+			}
+		}
+
+		// asynchronously update the octree, this normally finishes well before the next update
+		// but in the unlikely event that it does not, it would result in a frame where the paint brush
+		// is not perfectly aligned with the mesh; not a deal breaker.
 		UE::Geometry::FDynamicMeshOctree3& OctreeToUpdate = Tool->TrianglesOctree;
 		Tool->TriangleOctreeFuture = Async(SkinPaintToolAsyncExecTarget, [&OctreeToUpdate, &TrianglesToReinsert]()
 		{
diff --git a/Engine/Plugins/Experimental/MeshModelingToolsetExp/Source/MeshModelingToolsExp/Public/SkinWeightsPaintTool.h b/Engine/Plugins/Experimental/MeshModelingToolsetExp/Source/MeshModelingToolsExp/Public/SkinWeightsPaintTool.h
--- a/Engine/Plugins/Experimental/MeshModelingToolsetExp/Source/MeshModelingToolsExp/Public/SkinWeightsPaintTool.h
+++ b/Engine/Plugins/Experimental/MeshModelingToolsetExp/Source/MeshModelingToolsExp/Public/SkinWeightsPaintTool.h
@@ -47,7 +47,7 @@ namespace SkinPaintTool
 		void UpdateVertexDeformation(USkinWeightsPaintTool* Tool);
 
 		void SetVertexNeedsUpdated(int32 VertexIndex);
-	
+		
 		// which vertices require updating (partially re-calculated skinning deformation while painting)
 		TSet<int32> VerticesWithModifiedWeights;
 		// position of all vertices in the reference pose
@@ -283,7 +283,7 @@ protected:
 	UE::Geometry::TDynamicVerticesOctree3<FDynamicMesh3> VerticesOctree;
 	UE::Geometry::FDynamicMeshOctree3 TrianglesOctree;
 	TFuture<void> TriangleOctreeFuture;
-	TArray<int32> TriangleToReinsert;
+	TArray<int32> TrianglesToReinsert;
 
 	// tool properties
 	UPROPERTY()
