MeshModelingTools: Fix CubeGrid having ctrl+q hotkey on mac (which closes editor).

diff --git a/Engine/Plugins/Experimental/MeshModelingToolsetExp/Source/MeshModelingToolsExp/Private/CubeGridTool.cpp b/Engine/Plugins/Experimental/MeshModelingToolsetExp/Source/MeshModelingToolsExp/Private/CubeGridTool.cpp
--- a/Engine/Plugins/Experimental/MeshModelingToolsetExp/Source/MeshModelingToolsExp/Private/CubeGridTool.cpp
+++ b/Engine/Plugins/Experimental/MeshModelingToolsetExp/Source/MeshModelingToolsExp/Private/CubeGridTool.cpp
@@ -2317,13 +2317,22 @@ void UCubeGridTool::RegisterActions(FInteractiveToolActionSet& ActionSet)
 		TEXT("DecreaseGridPower"),
 		LOCTEXT("DecreaseGridPowerAction", "Decrease Grid Power"),
 		LOCTEXT("DecreaseGridPowerTooltip", ""),
+		// Note that we can't use Ctrl+Q on Mac because that is mapped to Cmd+Q which kills the editor.
+#if PLATFORM_MAC
+		EModifierKey::Alt, EKeys::Q,
+#else
 		EModifierKey::Control, EKeys::Q,
+#endif
 		[this]() { RequestAction(ECubeGridToolAction::DecreaseGridPower); });
 	ActionSet.RegisterAction(this, ActionID++,
 		TEXT("IncreaseGridPower"),
 		LOCTEXT("IncreaseGridPowerAction", "Increase Grid Power"),
 		LOCTEXT("IncreaseGridPowerTooltip", ""),
+#if PLATFORM_MAC
+		EModifierKey::Alt, EKeys::E,
+#else
 		EModifierKey::Control, EKeys::E,
+#endif
 		[this]() { RequestAction(ECubeGridToolAction::IncreaseGridPower); });
 
 	ActionSet.RegisterAction(this, ActionID++,
diff --git a/Engine/Plugins/Experimental/MeshModelingToolsetExp/Source/MeshModelingToolsExp/Public/CubeGridTool.h b/Engine/Plugins/Experimental/MeshModelingToolsetExp/Source/MeshModelingToolsExp/Public/CubeGridTool.h
--- a/Engine/Plugins/Experimental/MeshModelingToolsetExp/Source/MeshModelingToolsExp/Public/CubeGridTool.h
+++ b/Engine/Plugins/Experimental/MeshModelingToolsetExp/Source/MeshModelingToolsExp/Public/CubeGridTool.h
@@ -95,7 +95,9 @@ public:
 	const uint8 DEFAULT_GRID_POWER = 5;
 	const double DEFAULT_CURRENT_BLOCK_SIZE = 100;
 
-	/** Determines cube grid scale. Can also be adjusted with Ctrl + E/Q. */
+	//~ Unfortunately it seems that there isn't a way to make the automatic tooltip
+	//~ here platform-specific, so we can't be specific about the hotkeys
+	/** Determines cube grid scale. Can also be adjusted with hotkeys. */
 	UPROPERTY(EditAnywhere, Category = Options, meta = (
 		EditCondition = "bAllowedToEditGrid", HideEditConditionToggle,
 		UIMin = "0", UIMax = "10", ClampMin = "0", ClampMax = "31"))
@@ -177,7 +179,12 @@ public:
 	FString PushPull = TEXT("E/Q to pull/push, or use Ctrl+drag.");
 
 	UPROPERTY(VisibleAnywhere, Category = ShortcutInfo)
-	FString ResizeGrid = TEXT("Ctrl + E/Q to increase/decrease grid in powers of two.");
+	FString ResizeGrid =
+#if PLATFORM_MAC
+		TEXT("Option + E/Q to increase/decrease grid size.");
+#else
+		TEXT("Ctrl + E/Q to increase/decrease grid size.");
+#endif
 
 	UPROPERTY(VisibleAnywhere, Category = ShortcutInfo)
 	FString SlideSelection = TEXT("Middle mouse drag to slide selection in plane. "
