CombineMeshInstances: constrain voxwrap-LOD simplification to always target a triangle count lower than last non-voxwrap LOD

diff --git a/Engine/Plugins/Experimental/MeshModelingToolsetExp/Source/GeometryProcessingAdapters/Private/CombineMeshInstancesImpl.cpp b/Engine/Plugins/Experimental/MeshModelingToolsetExp/Source/GeometryProcessingAdapters/Private/CombineMeshInstancesImpl.cpp
--- a/Engine/Plugins/Experimental/MeshModelingToolsetExp/Source/GeometryProcessingAdapters/Private/CombineMeshInstancesImpl.cpp
+++ b/Engine/Plugins/Experimental/MeshModelingToolsetExp/Source/GeometryProcessingAdapters/Private/CombineMeshInstancesImpl.cpp
@@ -2062,9 +2062,19 @@ void BuildCombinedMesh(
 			UE_LOG(LogGeometry, Log, TEXT("         FastCollapse         - Tris %8d Verts %8d"), TempBaseVoxWrapMesh.TriangleCount(), TempBaseVoxWrapMesh.VertexCount());
 		}
 
-		const FDynamicMesh3& LastApproxLOD = MeshLODs[FirstVoxWrappedIndex-1].Mesh;
-
-		int32 MaxTriCount = CombineOptions.VoxWrapMaxTriCountBase;
+		// need to ensure that the triangle count of the first voxel LOD ends up smaller
+		// than the triangle count of the last approximate/etc LOD. We don't know this until
+		// the hidden-removal tasks finish, so wait for them here. We cannot guarantee
+		// that those LODs necessarily reduce in triangle count, that gets sorted out later,
+		// so find the min count here
+		UE::Tasks::Wait(PendingRemoveHiddenTasks);
+		int32 MinNonVoxWrapLODTriCount = MeshLODs[0].Mesh.TriangleCount();
+		for (int32 k = 1; k < FirstVoxWrappedIndex; ++k)
+		{
+			MinNonVoxWrapLODTriCount = FMath::Min(MinNonVoxWrapLODTriCount, MeshLODs[k].Mesh.TriangleCount());
+		}
+		// half is maybe a bit aggressive...
+		int32 MaxTriCount = FMath::Min(CombineOptions.VoxWrapMaxTriCountBase, MinNonVoxWrapLODTriCount / 2);
 		double SimplifyTolerance = CombineOptions.VoxWrapBaseTolerance;
 
 		// Current state of TempBaseVoxWrapMesh is our initial voxel LOD. To ensure
