ApproximateActors: fix issue where ground-plane occlusion could remove all triangles, and an empty StaticMesh would be returned. Now this case is detected and an error code is returned.

diff --git a/Engine/Plugins/Experimental/MeshModelingToolsetExp/Source/GeometryProcessingAdapters/Private/ApproximateActorsImpl.cpp b/Engine/Plugins/Experimental/MeshModelingToolsetExp/Source/GeometryProcessingAdapters/Private/ApproximateActorsImpl.cpp
--- a/Engine/Plugins/Experimental/MeshModelingToolsetExp/Source/GeometryProcessingAdapters/Private/ApproximateActorsImpl.cpp
+++ b/Engine/Plugins/Experimental/MeshModelingToolsetExp/Source/GeometryProcessingAdapters/Private/ApproximateActorsImpl.cpp
@@ -677,6 +677,12 @@ static TSharedPtr<FApproximationMeshData> GenerateApproximationMesh(
 		}
 	}
 
+	// if occlusion removed all the triangles, fail
+	if (CurResultMesh->TriangleCount() == 0)
+	{
+		Result->ResultCode = IGeometryProcessing_ApproximateActors::EResultCode::MeshGenerationFailed;
+		return Result;
+	}
 
 	TRACE_BOOKMARK(TEXT("ApproximateActors-Clip Ground"));
 
@@ -698,6 +704,13 @@ static TSharedPtr<FApproximationMeshData> GenerateApproximationMesh(
 		}
 	}
 
+	// if clipping removed all the triangles, fail
+	if (CurResultMesh->TriangleCount() == 0)
+	{
+		Result->ResultCode = IGeometryProcessing_ApproximateActors::EResultCode::MeshGenerationFailed;
+		return Result;
+	}
+
 	TRACE_BOOKMARK(TEXT("ApproximateActors-Normals and UVs"));
 
 	// re-enable attributes
