fix ensure on trying to copy tangents from a mesh that doesn't have any + tell the user that the mesh didn't have tangents

diff --git a/Engine/Plugins/Experimental/MeshModelingToolsetExp/Source/MeshModelingToolsEditorOnlyExp/Private/MeshTangentsTool.cpp b/Engine/Plugins/Experimental/MeshModelingToolsetExp/Source/MeshModelingToolsEditorOnlyExp/Private/MeshTangentsTool.cpp
--- a/Engine/Plugins/Experimental/MeshModelingToolsetExp/Source/MeshModelingToolsEditorOnlyExp/Private/MeshTangentsTool.cpp
+++ b/Engine/Plugins/Experimental/MeshModelingToolsetExp/Source/MeshModelingToolsEditorOnlyExp/Private/MeshTangentsTool.cpp
@@ -110,6 +110,23 @@ void UMeshTangentsTool::Setup()
 
 	Compute = MakeUnique<TGenericDataBackgroundCompute<FMeshTangentsd>>();
 	Compute->Setup(this);
+	Compute->OnOpCompleted.AddLambda([this](const UE::Geometry::TGenericDataOperator<UE::Geometry::FMeshTangentsd>* Op)
+		{
+			if (((FCalculateTangentsOp*)(Op))->bNoAttributesError)
+			{
+				bHasDisplayedNoAttributeError = true;
+				GetToolManager()->DisplayMessage(
+					LOCTEXT("TangentsNoAttributesError", "Error: Source mesh did not have tangents."),
+					EToolMessageLevel::UserWarning);
+			}
+			else if (bHasDisplayedNoAttributeError)
+			{
+				bHasDisplayedNoAttributeError = false;
+				GetToolManager()->DisplayMessage(
+					FText(),
+					EToolMessageLevel::UserWarning);
+			}
+		});
 	Compute->OnResultUpdated.AddLambda( [this](const TUniquePtr<FMeshTangentsd>& NewResult) { OnTangentsUpdated(NewResult); } );
 	Compute->InvalidateResult();
 
diff --git a/Engine/Plugins/Experimental/MeshModelingToolsetExp/Source/MeshModelingToolsEditorOnlyExp/Public/MeshTangentsTool.h b/Engine/Plugins/Experimental/MeshModelingToolsetExp/Source/MeshModelingToolsEditorOnlyExp/Public/MeshTangentsTool.h
--- a/Engine/Plugins/Experimental/MeshModelingToolsetExp/Source/MeshModelingToolsEditorOnlyExp/Public/MeshTangentsTool.h
+++ b/Engine/Plugins/Experimental/MeshModelingToolsetExp/Source/MeshModelingToolsEditorOnlyExp/Public/MeshTangentsTool.h
@@ -148,4 +148,6 @@ protected:
 	};
 	TArray<FMikktDeviation> Deviations;
 
+private:
+	bool bHasDisplayedNoAttributeError = false;
 };
diff --git a/Engine/Plugins/Runtime/MeshModelingToolset/Source/ModelingOperatorsEditorOnly/Private/ParameterizationOps/CalculateTangentsOp.cpp b/Engine/Plugins/Runtime/MeshModelingToolset/Source/ModelingOperatorsEditorOnly/Private/ParameterizationOps/CalculateTangentsOp.cpp
--- a/Engine/Plugins/Runtime/MeshModelingToolset/Source/ModelingOperatorsEditorOnly/Private/ParameterizationOps/CalculateTangentsOp.cpp
+++ b/Engine/Plugins/Runtime/MeshModelingToolset/Source/ModelingOperatorsEditorOnly/Private/ParameterizationOps/CalculateTangentsOp.cpp
@@ -78,8 +78,9 @@ void FCalculateTangentsOp::CalculateStandard(FProgressCancel* Progress, TUniqueP
 void FCalculateTangentsOp::CopyFromSource(FProgressCancel* Progress, TUniquePtr<FMeshTangentsd>& Tangents)
 {
 	int32 NumTangents = SourceTangents->GetTangents().Num();
-	if ( ! ensure(NumTangents == SourceMesh->MaxTriangleID() * 3) )
+	if (NumTangents != SourceMesh->MaxTriangleID() * 3)
 	{
+		ensure(NumTangents == 0); // Expect this case to happen only when there aren't any tangents
 		bNoAttributesError = true;
 		Tangents->InitializeTriVertexTangents(true);
 		return;
