ModelingTools: enable negative angles for Bend warp. Despite the comment in the header, bending a negative angle does exactly what I would expect, and seems to work fine...

diff --git a/Engine/Plugins/Experimental/MeshModelingToolsetExp/Source/MeshModelingToolsExp/Public/MeshSpaceDeformerTool.h b/Engine/Plugins/Experimental/MeshModelingToolsetExp/Source/MeshModelingToolsExp/Public/MeshSpaceDeformerTool.h
--- a/Engine/Plugins/Experimental/MeshModelingToolsetExp/Source/MeshModelingToolsExp/Public/MeshSpaceDeformerTool.h
+++ b/Engine/Plugins/Experimental/MeshModelingToolsetExp/Source/MeshModelingToolsExp/Public/MeshSpaceDeformerTool.h
@@ -91,14 +91,12 @@ public:
 	UPROPERTY(EditAnywhere, Category = Options, meta = (DisplayName = "Lower Bound", UIMax = "0", ClampMax = "0"))
 	float LowerBoundsInterval = -100.0;
 
-	//~ The two degrees properties (bend vs twist) are separate because they have different clamp values. 
-	//~ Bending a "negative" amount probably won't do what the user wants, whereas twisting in the opposite
-	//~ direction makes sense.
+	//~ The two degrees properties (bend vs twist) are separate because they have different default values
 	/** 
 	 * A line along the Z axis of the gizmo from lower bound to upper bound will be bent into a perfect arc of this
 	 * many degrees in the direction of the Y axis without changing length.
 	 */
-	UPROPERTY(EditAnywhere, Category = Options, meta = (UIMin = "0", UIMax = "360", ClampMin = "0", ClampMax = "3600", 
+	UPROPERTY(EditAnywhere, Category = Options, meta = (UIMin = "-360", UIMax = "360", ClampMin = "-3600", ClampMax = "3600", 
 		EditCondition = "SelectedOperationType == ENonlinearOperationType::Bend", EditConditionHides))
 	float BendDegrees = 90;
 
