Skeletal Mesh Modifier: avoid removing a bone twice

diff --git a/Engine/Plugins/Experimental/MeshModelingToolsetExp/Source/MeshModelingToolsExp/Private/SkeletalMesh/SkeletonEditingTool.cpp b/Engine/Plugins/Experimental/MeshModelingToolsetExp/Source/MeshModelingToolsExp/Private/SkeletalMesh/SkeletonEditingTool.cpp
--- a/Engine/Plugins/Experimental/MeshModelingToolsetExp/Source/MeshModelingToolsExp/Private/SkeletalMesh/SkeletonEditingTool.cpp
+++ b/Engine/Plugins/Experimental/MeshModelingToolsetExp/Source/MeshModelingToolsExp/Private/SkeletalMesh/SkeletonEditingTool.cpp
@@ -994,7 +994,7 @@ FBox USkeletonEditingTool::GetWorldSpaceFocusBox()
 
 			// get direct children
 			TArray<int32> Children;
-			RefSkeleton.GetDirectChildBones(BoneIndex, Children);
+			RefSkeleton.GetRawDirectChildBones(BoneIndex, Children);
 			AllChildren.Append(Children);
 		}
 
diff --git a/Engine/Plugins/Experimental/MeshModelingToolsetExp/Source/SkeletalMeshModifiers/Private/SkeletonModifier.cpp b/Engine/Plugins/Experimental/MeshModelingToolsetExp/Source/SkeletalMeshModifiers/Private/SkeletonModifier.cpp
--- a/Engine/Plugins/Experimental/MeshModelingToolsetExp/Source/SkeletalMeshModifiers/Private/SkeletonModifier.cpp
+++ b/Engine/Plugins/Experimental/MeshModelingToolsetExp/Source/SkeletalMeshModifiers/Private/SkeletonModifier.cpp
@@ -544,7 +544,7 @@ void USkeletonModifier::GetBonesToMirror(
 
 		if (InOptions.bMirrorChildren)
 		{
-			TArray<int32> Children; ReferenceSkeleton->GetDirectChildBones(BoneIndex, Children);
+			TArray<int32> Children; ReferenceSkeleton->GetRawDirectChildBones(BoneIndex, Children);
 			for (int32 ChildIndex: Children)
 			{
 				GetBonesToMirror2(ChildIndex, GetBonesToMirror2);
@@ -807,7 +807,7 @@ bool USkeletonModifier::SetBonesTransforms(
 		for (int32 Index = 0; Index < BoneIndices.Num(); Index++)
 		{
 			TArray<int32> Children;
-			ReferenceSkeleton->GetDirectChildBones(BoneIndices[Index], Children);
+			ReferenceSkeleton->GetRawDirectChildBones(BoneIndices[Index], Children);
 			for (int32 ChildIndex: Children)
 			{
 				if (!BoneIndices.Contains(ChildIndex))
@@ -888,10 +888,45 @@ bool USkeletonModifier::RemoveBones(const TArray<FName>& InBoneNames, const bool
 	// store initial data
 	const TArray<FMeshBoneInfo> InfosBeforeRemoval = ReferenceSkeleton->GetRawRefBoneInfo();
 
+	TArray<FName> BonesToRemove;
+	if (bRemoveChildren)
+	{
+		const TArray<FMeshBoneInfo>& BoneInfos = ReferenceSkeleton->GetRawRefBoneInfo();
+		auto IsParentToBeRemoved = [&](const FName BoneName)
+		{
+			const int32 BoneIndex = ReferenceSkeleton->FindRawBoneIndex(BoneName);
+			if (BoneIndex != INDEX_NONE)
+			{
+				int32 ParentIndex = BoneInfos[BoneIndex].ParentIndex;
+				while (ParentIndex != INDEX_NONE)
+				{
+					if (InBoneNames.Contains(BoneInfos[ParentIndex].Name))
+					{
+						return true;
+					}
+					ParentIndex = BoneInfos[ParentIndex].ParentIndex;
+				}
+			}
+			return false;
+		};
+		
+		for (const FName& BoneName: InBoneNames)
+		{
+			if (!IsParentToBeRemoved(BoneName))
+			{
+				BonesToRemove.Add(BoneName);
+			}
+		}
+	}
+	else
+	{
+		BonesToRemove = InBoneNames;
+	}
+
 	// update reference skeleton
 	{
 		FReferenceSkeletonModifier Modifier(*ReferenceSkeleton, nullptr);
-		for (const FName& BoneName: InBoneNames)
+		for (const FName& BoneName: BonesToRemove)
 		{
 			Modifier.Remove(BoneName, bRemoveChildren);
 		}
@@ -1070,7 +1105,7 @@ bool USkeletonModifier::OrientBones(const TArray<FName>& InBoneNames, const FOri
 		
 		const int32 ParentIndex = BoneIndex != INDEX_NONE ? ReferenceSkeleton->GetRawParentIndex(BoneIndex) : INDEX_NONE;
 
-		TArray<int32> Children; ReferenceSkeleton->GetDirectChildBones(BoneIndex, Children);
+		TArray<int32> Children; ReferenceSkeleton->GetRawDirectChildBones(BoneIndex, Children);
 		const int32 NumChildren = Children.Num();
 		if (NumChildren > 1)
 		{ // we can't align if there are more than one children
@@ -1145,7 +1180,7 @@ void USkeletonModifier::GetBonesToOrient(
 
 		if (InOptions.bOrientChildren)
 		{
-			TArray<int32> Children; ReferenceSkeleton->GetDirectChildBones(BoneIndex, Children);
+			TArray<int32> Children; ReferenceSkeleton->GetRawDirectChildBones(BoneIndex, Children);
 			for (int32 ChildIndex: Children)
 			{
 				GetBonesToOrient2(ChildIndex, GetBonesToOrient2);
diff --git a/Engine/Plugins/Experimental/MeshModelingToolsetExp/Source/SkeletalMeshModifiers/Public/SkeletonModifier.h b/Engine/Plugins/Experimental/MeshModelingToolsetExp/Source/SkeletalMeshModifiers/Public/SkeletonModifier.h
--- a/Engine/Plugins/Experimental/MeshModelingToolsetExp/Source/SkeletalMeshModifiers/Public/SkeletonModifier.h
+++ b/Engine/Plugins/Experimental/MeshModelingToolsetExp/Source/SkeletalMeshModifiers/Public/SkeletonModifier.h
@@ -251,7 +251,7 @@ struct FTransformComposer
 		Transforms[Index] = RefSkeleton.GetRawRefBonePose()[Index];
 		TransformCached[Index] = false;
 
-		TArray<int32> Children; RefSkeleton.GetDirectChildBones(Index, Children);
+		TArray<int32> Children; RefSkeleton.GetRawDirectChildBones(Index, Children);
 		for (const int32 ChildIndex: Children)
 		{
 			Invalidate(ChildIndex);
diff --git a/Engine/Source/Runtime/Engine/Private/Animation/ReferenceSkeleton.cpp b/Engine/Source/Runtime/Engine/Private/Animation/ReferenceSkeleton.cpp
--- a/Engine/Source/Runtime/Engine/Private/Animation/ReferenceSkeleton.cpp
+++ b/Engine/Source/Runtime/Engine/Private/Animation/ReferenceSkeleton.cpp
@@ -118,7 +118,7 @@ void FReferenceSkeleton::Remove(const FName InBoneName, const bool bRemoveChildr
 
 	// store children indices and sort them from greatest to lowest
 	TArray<int32> Children;
-	GetDirectChildBones(RawBoneIndex, Children);
+	GetRawDirectChildBones(RawBoneIndex, Children);
 	Children.Sort([](const int32 Index0, const int32 Index1) {return Index0 > Index1;} );
 
 	// reindex function
@@ -623,20 +623,31 @@ void FReferenceSkeleton::EnsureParentsExistAndSort(TArray<FBoneIndexType>& InOut
 	InOutBoneUnsortedArray.Sort();
 }
 
-int32 FReferenceSkeleton::GetDirectChildBones(int32 ParentBoneIndex, TArray<int32> & Children) const
+int32 FReferenceSkeleton::GetChildrenInternal(int32 InParentBoneIndex, TArray<int32>& OutChildren, const bool bRaw) const
 {
-	Children.Reset();
+	OutChildren.Reset();
 
-	const int32 NumBones = GetNum();
-	for (int32 ChildIndex = ParentBoneIndex + 1; ChildIndex < NumBones; ChildIndex++)
+	const int32 NumBones = bRaw ? GetRawBoneNum() : GetNum();
+	for (int32 ChildIndex = InParentBoneIndex + 1; ChildIndex < NumBones; ChildIndex++)
 	{
-		if (ParentBoneIndex == GetParentIndex(ChildIndex))
+		const int32 ParentIndex = bRaw ? GetRawParentIndex(ChildIndex) : GetParentIndex(ChildIndex);
+		if (ParentIndex == InParentBoneIndex)
 		{
-			Children.Add(ChildIndex);
+			OutChildren.Add(ChildIndex);
 		}
 	}
 
-	return Children.Num();
+	return OutChildren.Num();	
+}
+
+int32 FReferenceSkeleton::GetDirectChildBones(int32 ParentBoneIndex, TArray<int32> & Children) const
+{
+	return GetChildrenInternal(ParentBoneIndex, Children, false);
+}
+
+int32 FReferenceSkeleton::GetRawDirectChildBones(int32 ParentBoneIndex, TArray<int32> & Children) const
+{
+	return GetChildrenInternal(ParentBoneIndex, Children, true);
 }
 
 FArchive & operator<<(FArchive & Ar, FReferenceSkeleton & F)
diff --git a/Engine/Source/Runtime/Engine/Public/ReferenceSkeleton.h b/Engine/Source/Runtime/Engine/Public/ReferenceSkeleton.h
--- a/Engine/Source/Runtime/Engine/Public/ReferenceSkeleton.h
+++ b/Engine/Source/Runtime/Engine/Public/ReferenceSkeleton.h
@@ -225,6 +225,9 @@ private:
 	// Help us translate a virtual bone source into a raw bone source (for evaluating virtual bone transform)
 	int32 GetRawSourceBoneIndex(const USkeleton* Skeleton, const FName& SourceBoneName) const;
 
+	// very slow search function for all children (raw or final)
+	int32 GetChildrenInternal(int32 InParentBoneIndex, TArray<int32>& OutChildren, const bool bRaw) const;
+	
 public:
 	ENGINE_API void RebuildRefSkeleton(const USkeleton* Skeleton, bool bRebuildNameMap);
 
@@ -428,6 +431,7 @@ public:
 
 	// very slow search function for all children
 	ENGINE_API int32 GetDirectChildBones(int32 ParentBoneIndex, TArray<int32> & Children) const;
+	ENGINE_API int32 GetRawDirectChildBones(int32 ParentBoneIndex, TArray<int32> & Children) const;
 	ENGINE_API friend FArchive & operator<<(FArchive & Ar, FReferenceSkeleton & F);
 	friend FReferenceSkeletonModifier;
 };
