Fixed bone mirroring
   - ensure the mirrored bone names are available

diff --git a/Engine/Plugins/Experimental/MeshModelingToolsetExp/Source/SkeletalMeshModifiers/Private/SkeletonModifier.cpp b/Engine/Plugins/Experimental/MeshModelingToolsetExp/Source/SkeletalMeshModifiers/Private/SkeletonModifier.cpp
--- a/Engine/Plugins/Experimental/MeshModelingToolsetExp/Source/SkeletalMeshModifiers/Private/SkeletonModifier.cpp
+++ b/Engine/Plugins/Experimental/MeshModelingToolsetExp/Source/SkeletalMeshModifiers/Private/SkeletonModifier.cpp
@@ -592,7 +592,7 @@ void USkeletonModifier::GetMirroredNames(
 		}
 		else
 		{
-			return GetUniqueName(BoneInfos[BoneIndex].Name);
+			return GetUniqueName(BoneInfos[BoneIndex].Name, OutBonesName);
 		}
 		return FName(*BoneStr);
 	});
@@ -1217,7 +1217,7 @@ void USkeletonModifier::UpdateBoneTracker(const TArray<FMeshBoneInfo>& InOtherIn
 	}
 }
 
-FName USkeletonModifier::GetUniqueName(const FName InBoneName) const
+FName USkeletonModifier::GetUniqueName(const FName InBoneName, const TArray<FName>& InBoneNames) const
 {
 	if (!ReferenceSkeleton || InBoneName == NAME_None)
 	{
@@ -1267,6 +1267,12 @@ FName USkeletonModifier::GetUniqueName(const FName InBoneName) const
 		{
 			return false;
 		}
+
+		if (InBoneNames.Contains(Name))
+		{
+			return false;
+		}
+		
 		return true;
 	};
 
diff --git a/Engine/Plugins/Experimental/MeshModelingToolsetExp/Source/SkeletalMeshModifiers/Public/SkeletonModifier.h b/Engine/Plugins/Experimental/MeshModelingToolsetExp/Source/SkeletalMeshModifiers/Public/SkeletonModifier.h
--- a/Engine/Plugins/Experimental/MeshModelingToolsetExp/Source/SkeletalMeshModifiers/Public/SkeletonModifier.h
+++ b/Engine/Plugins/Experimental/MeshModelingToolsetExp/Source/SkeletalMeshModifiers/Public/SkeletonModifier.h
@@ -178,7 +178,7 @@ public:
 
 	void ExternalUpdate(const FReferenceSkeleton& InRefSkeleton, const TArray<int32>& InIndexTracker);
 	
-	FName GetUniqueName(const FName InBoneName) const;
+	FName GetUniqueName(const FName InBoneName, const TArray<FName>& InBoneNames = {}) const;
 	
 	const FReferenceSkeleton& GetReferenceSkeleton() const;
 	const TArray<int32>& GetBoneIndexTracker() const;
