MeshModelingTools: Fix the CubeGrid "Reset Grid from Actor" button using the input actor's scale for the grid visualization. Change the Mac bindings for grid size changes to option+A/D because option+E is consumed by Mac.

diff --git a/Engine/Plugins/Experimental/MeshModelingToolsetExp/Source/MeshModelingToolsExp/Private/CubeGridTool.cpp b/Engine/Plugins/Experimental/MeshModelingToolsetExp/Source/MeshModelingToolsExp/Private/CubeGridTool.cpp
--- a/Engine/Plugins/Experimental/MeshModelingToolsetExp/Source/MeshModelingToolsExp/Private/CubeGridTool.cpp
+++ b/Engine/Plugins/Experimental/MeshModelingToolsetExp/Source/MeshModelingToolsExp/Private/CubeGridTool.cpp
@@ -2277,7 +2277,9 @@ void UCubeGridTool::ApplyAction(ECubeGridToolAction ActionType)
 	case ECubeGridToolAction::ResetFromActor:
 		if (ToolActions->GridSourceActor)
 		{
-			UpdateGridGizmo(ToolActions->GridSourceActor->GetTransform());
+			FTransform TransformToUse = ToolActions->GridSourceActor->GetTransform();
+			TransformToUse.SetScale3D(FVector::OneVector);
+			UpdateGridGizmo(TransformToUse);
 		}
 		break;
 	}
@@ -2318,8 +2320,9 @@ void UCubeGridTool::RegisterActions(FInteractiveToolActionSet& ActionSet)
 		LOCTEXT("DecreaseGridPowerAction", "Decrease Grid Power"),
 		LOCTEXT("DecreaseGridPowerTooltip", ""),
 		// Note that we can't use Ctrl+Q on Mac because that is mapped to Cmd+Q which kills the editor.
+		// At the same time we can't use Option+E because Mac consumes that for typing accented letters
 #if PLATFORM_MAC
-		EModifierKey::Alt, EKeys::Q,
+		EModifierKey::Alt, EKeys::A,
 #else
 		EModifierKey::Control, EKeys::Q,
 #endif
@@ -2329,7 +2332,7 @@ void UCubeGridTool::RegisterActions(FInteractiveToolActionSet& ActionSet)
 		LOCTEXT("IncreaseGridPowerAction", "Increase Grid Power"),
 		LOCTEXT("IncreaseGridPowerTooltip", ""),
 #if PLATFORM_MAC
-		EModifierKey::Alt, EKeys::E,
+		EModifierKey::Alt, EKeys::D,
 #else
 		EModifierKey::Control, EKeys::E,
 #endif
diff --git a/Engine/Plugins/Experimental/MeshModelingToolsetExp/Source/MeshModelingToolsExp/Public/CubeGridTool.h b/Engine/Plugins/Experimental/MeshModelingToolsetExp/Source/MeshModelingToolsExp/Public/CubeGridTool.h
--- a/Engine/Plugins/Experimental/MeshModelingToolsetExp/Source/MeshModelingToolsExp/Public/CubeGridTool.h
+++ b/Engine/Plugins/Experimental/MeshModelingToolsetExp/Source/MeshModelingToolsExp/Public/CubeGridTool.h
@@ -181,7 +181,7 @@ public:
 	UPROPERTY(VisibleAnywhere, Category = ShortcutInfo)
 	FString ResizeGrid =
 #if PLATFORM_MAC
-		TEXT("Option + E/Q to increase/decrease grid size.");
+		TEXT("Option + D/A to increase/decrease grid size.");
 #else
 		TEXT("Ctrl + E/Q to increase/decrease grid size.");
 #endif
