Fixed transaction order resulting in skeleton and skeletal mesh being unsynchronized when re-registering skinned mesh components

diff --git a/Engine/Plugins/Experimental/MeshModelingToolsetExp/Source/SkeletalMeshModifiers/Private/SkeletonModifier.cpp b/Engine/Plugins/Experimental/MeshModelingToolsetExp/Source/SkeletalMeshModifiers/Private/SkeletonModifier.cpp
--- a/Engine/Plugins/Experimental/MeshModelingToolsetExp/Source/SkeletalMeshModifiers/Private/SkeletonModifier.cpp
+++ b/Engine/Plugins/Experimental/MeshModelingToolsetExp/Source/SkeletalMeshModifiers/Private/SkeletonModifier.cpp
@@ -337,7 +337,12 @@ bool USkeletonModifier::CommitSkeletonToSkeletalMesh()
 
 	// update skeletal mesh
 	FlushRenderingCommands();
-	
+
+	// call modify on the skeleton first as post undo will re-register components so it must be done once both
+	// skeletal mesh and skeleton are up to date, so it must be done after the skeletal mesh has been undone 
+	USkeleton* Skeleton = SkeletalMesh->GetSkeleton();
+	Skeleton->Modify();
+
 	SkeletalMesh->SetFlags(RF_Transactional);
 	SkeletalMesh->Modify();
 
@@ -350,8 +355,6 @@ bool USkeletonModifier::CommitSkeletonToSkeletalMesh()
 	SkeletalMesh->CommitMeshDescription(USkeletonModifierLocals::LODIndex, *MeshDescription);
 
 	// update skeleton
-	USkeleton* Skeleton = SkeletalMesh->GetSkeleton();
-	Skeleton->Modify();
 	if (Skeleton->RecreateBoneTree(SkeletalMesh))
 	{
 		Skeleton->MarkPackageDirty();	
