Ensure bone index is safe to use

diff --git a/Engine/Plugins/Experimental/MeshModelingToolsetExp/Source/MeshModelingToolsExp/Private/SkeletalMesh/SkeletonEditingTool.cpp b/Engine/Plugins/Experimental/MeshModelingToolsetExp/Source/MeshModelingToolsExp/Private/SkeletalMesh/SkeletonEditingTool.cpp
--- a/Engine/Plugins/Experimental/MeshModelingToolsetExp/Source/MeshModelingToolsExp/Private/SkeletalMesh/SkeletonEditingTool.cpp
+++ b/Engine/Plugins/Experimental/MeshModelingToolsetExp/Source/MeshModelingToolsExp/Private/SkeletalMesh/SkeletonEditingTool.cpp
@@ -262,24 +262,27 @@ void USkeletonEditingTool::UpdateGizmo() const
 	const TArray<FMeshBoneInfo>& BoneInfos = ReferenceSkeleton.GetRawRefBoneInfo();
 	for (const int32 BoneIndex: BoneIndexes)
 	{
-		const FName BoneName = BoneInfos[BoneIndex].Name;
-		const int32 ParentBoneIndex = BoneInfos[BoneIndex].ParentIndex;
-		const FTransform& ParentGlobal = Modifier->GetTransform(ParentBoneIndex, true);
+		if (BoneInfos.IsValidIndex(BoneIndex))
+		{
+			const FName BoneName = BoneInfos[BoneIndex].Name;
+			const int32 ParentBoneIndex = BoneInfos[BoneIndex].ParentIndex;
+			const FTransform& ParentGlobal = Modifier->GetTransform(ParentBoneIndex, true);
 		
-		GizmoWrapper->HandleBoneTransform(
-			[this, BoneIndex, bWorld]()
-			{
-				return Modifier->GetTransform(BoneIndex, bWorld);
-			},
-			[this, BoneName, ParentGlobal, bWorld](const FTransform& InNewTransform)
-			{
-				if (bWorld)
+			GizmoWrapper->HandleBoneTransform(
+				[this, BoneIndex, bWorld]()
 				{
-					const FTransform NewLocal = InNewTransform.GetRelativeTransform(ParentGlobal);
-					return Modifier->SetBoneTransform(BoneName, NewLocal, Properties->bUpdateChildren);
-				}
-				return Modifier->SetBoneTransform(BoneName, InNewTransform, Properties->bUpdateChildren);
-			});
+					return Modifier->GetTransform(BoneIndex, bWorld);
+				},
+				[this, BoneName, ParentGlobal, bWorld](const FTransform& InNewTransform)
+				{
+					if (bWorld)
+					{
+						const FTransform NewLocal = InNewTransform.GetRelativeTransform(ParentGlobal);
+						return Modifier->SetBoneTransform(BoneName, NewLocal, Properties->bUpdateChildren);
+					}
+					return Modifier->SetBoneTransform(BoneName, InNewTransform, Properties->bUpdateChildren);
+				});
+		}
 	}
 }
 
