Ensure we are actually interacting with the gizmo before calling begin/end update functions

diff --git a/Engine/Plugins/Experimental/Animation/SkeletalMeshModelingTools/Source/SkeletalMeshModelingTools/Private/SkeletalMeshGizmoUtils.cpp b/Engine/Plugins/Experimental/Animation/SkeletalMeshModelingTools/Source/SkeletalMeshModelingTools/Private/SkeletalMeshGizmoUtils.cpp
--- a/Engine/Plugins/Experimental/Animation/SkeletalMeshModelingTools/Source/SkeletalMeshModelingTools/Private/SkeletalMeshGizmoUtils.cpp
+++ b/Engine/Plugins/Experimental/Animation/SkeletalMeshModelingTools/Source/SkeletalMeshModelingTools/Private/SkeletalMeshGizmoUtils.cpp
@@ -145,6 +145,15 @@ USkeletalMeshGizmoWrapperBase* USkeletalMeshGizmoContextObject::GetNewWrapper(UI
 	USkeletalMeshGizmoWrapper* GizmoWrapper = NewObject<USkeletalMeshGizmoWrapper>(Outer);
 	GizmoWrapper->TransformGizmo = Gizmo;
 	GizmoWrapper->TransformProxy = SkeletonTransformProxy;
+
+	if (Gizmo->HitTarget)
+	{
+		TWeakObjectPtr<USkeletalMeshGizmoWrapper> WeakWrapper = MakeWeakObjectPtr(GizmoWrapper);
+		Gizmo->HitTarget->Condition = [WeakWrapper](const FInputDeviceRay&)
+		{
+			return WeakWrapper.IsValid() ? WeakWrapper->CanInteract() : false;
+		};
+	}
 	
 	return GizmoWrapper;
 }
@@ -163,6 +172,15 @@ const FGizmoCustomization& USkeletalMeshGizmoContextObject::GetGizmoCustomizatio
 	return Customization;
 }
 
+bool USkeletalMeshGizmoWrapper::CanInteract() const
+{
+	if (!TransformProxy || !TransformGizmo)
+	{
+		return false;
+	}
+	return TransformProxy->IsValid() && TransformGizmo->bVisible; 
+}
+
 bool USkeletalMeshGizmoWrapper::IsGizmoHit(const FInputDeviceRay& PressPos) const
 {
 	if (!TransformGizmo || !TransformGizmo->bVisible)
diff --git a/Engine/Plugins/Experimental/Animation/SkeletalMeshModelingTools/Source/SkeletalMeshModelingTools/Private/SkeletalMeshGizmoUtils.h b/Engine/Plugins/Experimental/Animation/SkeletalMeshModelingTools/Source/SkeletalMeshModelingTools/Private/SkeletalMeshGizmoUtils.h
--- a/Engine/Plugins/Experimental/Animation/SkeletalMeshModelingTools/Source/SkeletalMeshModelingTools/Private/SkeletalMeshGizmoUtils.h
+++ b/Engine/Plugins/Experimental/Animation/SkeletalMeshModelingTools/Source/SkeletalMeshModelingTools/Private/SkeletalMeshGizmoUtils.h
@@ -59,6 +59,7 @@ public:
 	virtual void HandleBoneTransform(FGetTransform GetTransformFunc, FSetTransform SetTransformFunc) override;
 	virtual void Clear() override;
 
+	virtual bool CanInteract() const override;
 	virtual bool IsGizmoHit(const FInputDeviceRay& PressPos) const override;
 
 	UPROPERTY()
diff --git a/Engine/Plugins/Experimental/MeshModelingToolsetExp/Source/MeshModelingToolsExp/Private/SkeletalMesh/SkeletonEditingTool.cpp b/Engine/Plugins/Experimental/MeshModelingToolsetExp/Source/MeshModelingToolsExp/Private/SkeletalMesh/SkeletonEditingTool.cpp
--- a/Engine/Plugins/Experimental/MeshModelingToolsetExp/Source/MeshModelingToolsExp/Private/SkeletalMesh/SkeletonEditingTool.cpp
+++ b/Engine/Plugins/Experimental/MeshModelingToolsetExp/Source/MeshModelingToolsExp/Private/SkeletalMesh/SkeletonEditingTool.cpp
@@ -178,19 +178,25 @@ void USkeletonEditingTool::Setup()
 		UGizmoLambdaStateTarget* NewTarget = NewObject<UGizmoLambdaStateTarget>(this);
 		NewTarget->BeginUpdateFunction = [this]()
 		{
-			TGuardValue OperationGuard(Operation, EEditingOperation::Transform);
-			BeginChange();
+			if (GizmoWrapper && GizmoWrapper->CanInteract())
+			{
+				TGuardValue OperationGuard(Operation, EEditingOperation::Transform);
+				BeginChange();
+			}
 		};
-		NewTarget->EndUpdateFunction = [&]()
+		NewTarget->EndUpdateFunction = [this]()
 		{
-			TGuardValue OperationGuard(Operation, EEditingOperation::Transform);
-			EndChange();
-
-			const IToolsContextQueriesAPI* ToolsContextQueries = GetToolManager()->GetPairedGizmoManager()->GetContextQueriesAPI();
-			check(ToolsContextQueries);
-			if (ToolsContextQueries->GetCurrentCoordinateSystem() == EToolContextCoordinateSystem::World)
+			if (GizmoWrapper && GizmoWrapper->CanInteract())
 			{
-				UpdateGizmo();
+				TGuardValue OperationGuard(Operation, EEditingOperation::Transform);
+				EndChange();
+
+				const IToolsContextQueriesAPI* ToolsContextQueries = GetToolManager()->GetPairedGizmoManager()->GetContextQueriesAPI();
+				check(ToolsContextQueries);
+				if (ToolsContextQueries->GetCurrentCoordinateSystem() == EToolContextCoordinateSystem::World)
+				{
+					UpdateGizmo();
+				}
 			}
 		};
 		
@@ -367,8 +373,12 @@ void USkeletonEditingTool::RegisterActions(FInteractiveToolActionSet& ActionSet)
 		EModifierKey::None, EKeys::Escape,
 		[this]()
 		{
-			Operation = EEditingOperation::Select;
-			GetToolManager()->DisplayMessage(LOCTEXT("Select", "Click on a bone to select it."), EToolMessageLevel::UserNotification);
+			if (Operation != EEditingOperation::Select)
+			{
+				Operation = EEditingOperation::Select;
+				UpdateGizmo();
+				GetToolManager()->DisplayMessage(LOCTEXT("Select", "Click on a bone to select it."), EToolMessageLevel::UserNotification);
+			}
 		});
 
 	// register UnParent key
diff --git a/Engine/Plugins/Experimental/MeshModelingToolsetExp/Source/MeshModelingToolsExp/Public/SkeletalMesh/SkeletalMeshEditionInterface.h b/Engine/Plugins/Experimental/MeshModelingToolsetExp/Source/MeshModelingToolsExp/Public/SkeletalMesh/SkeletalMeshEditionInterface.h
--- a/Engine/Plugins/Experimental/MeshModelingToolsetExp/Source/MeshModelingToolsExp/Public/SkeletalMesh/SkeletalMeshEditionInterface.h
+++ b/Engine/Plugins/Experimental/MeshModelingToolsetExp/Source/MeshModelingToolsExp/Public/SkeletalMesh/SkeletalMeshEditionInterface.h
@@ -126,6 +126,9 @@ public:
 	PURE_VIRTUAL(USkeletalMeshGizmoWrapperBase::HandleBoneTransform, return;)
 
 	virtual void Clear() PURE_VIRTUAL(USkeletalMeshGizmoWrapperBase::Clear, return;)
+
+	virtual bool CanInteract() const
+	PURE_VIRTUAL(USkeletalMeshGizmoWrapperBase::CanInteract, return false;)
 	
 	virtual bool IsGizmoHit(const FInputDeviceRay& PressPos) const PURE_VIRTUAL(USkeletalMeshGizmoWrapperBase::IsGizmoHit, return false;)
 
