CombineMeshInstances: add FOptions::bRetriangulateSourceLODs and start LODIndex setting, to improve simplification of LOD1 geometry. When enabled this finds nearly-planar polygroups and then uses FPolygroupRemesh to retriangulate, this generally is more effective than relying on the SimplifyToMinimalPlanar done during face merging after hidden removal.

Also added FOptions::VertexColorMappingMode which currently only supports a mapping from tri count (relative to a box) to a white-to-red vertex color scale, basically intended for use as a debug visualization.

#rnx

diff --git a/Engine/Plugins/Experimental/MeshModelingToolsetExp/Source/GeometryProcessingAdapters/Private/CombineMeshInstancesImpl.cpp b/Engine/Plugins/Experimental/MeshModelingToolsetExp/Source/GeometryProcessingAdapters/Private/CombineMeshInstancesImpl.cpp
--- a/Engine/Plugins/Experimental/MeshModelingToolsetExp/Source/GeometryProcessingAdapters/Private/CombineMeshInstancesImpl.cpp
+++ b/Engine/Plugins/Experimental/MeshModelingToolsetExp/Source/GeometryProcessingAdapters/Private/CombineMeshInstancesImpl.cpp
@@ -28,6 +28,7 @@
 
 #include "Polygroups/PolygroupsGenerator.h"
 #include "GroupTopology.h"
+#include "Operations/PolygroupRemesh.h"
 
 #include "MeshSimplification.h"
 #include "DynamicMesh/ColliderMesh.h"
@@ -659,6 +660,60 @@ static void SimplifyPartMesh(
 
 
 
+
+
+
+/**
+ * This function uses FPolygroupRemesh to try to completely retriangulate planar faces
+ */
+static void PlanarRetriangulatePartMesh(
+	FDynamicMesh3& EditMesh,
+	double Tolerance,
+	double RecomputeNormalsAngleThreshold)
+{
+	// currently bowties need to be split for Welder
+	FDynamicMeshEditor MeshEditor(&EditMesh);
+	FDynamicMeshEditResult EditResult;
+	MeshEditor.SplitBowties(EditResult);
+
+	// weld edges in case input was unwelded...
+	FMergeCoincidentMeshEdges Welder(&EditMesh);
+	Welder.MergeVertexTolerance = Tolerance * 0.001;
+	Welder.OnlyUniquePairs = false;
+	Welder.Apply();
+
+	// Skip out for very low-poly parts, they are unlikely to simplify very nicely.
+	if (EditMesh.VertexCount() < 16)
+	{
+		return;
+	}
+
+	double AngleToleranceDeg = 2.0;
+
+	// generate polygroups for planar areas of the mesh
+	FPolygroupsGenerator Generator(&EditMesh);
+	const bool bUVSeams = false, bNormalSeams = false;
+	double DotTolerance = 1.0 - FMathd::Cos(AngleToleranceDeg * FMathd::DegToRad);
+	Generator.FindPolygroupsFromFaceNormals(DotTolerance, bUVSeams, bNormalSeams);
+	Generator.CopyPolygroupsToMesh();
+
+	FGroupTopology UseTopology(&EditMesh, true);
+
+	FPolygroupRemesh Simplifier(&EditMesh, &UseTopology, ConstrainedDelaunayTriangulate<double>);
+	Simplifier.SimplificationAngleTolerance = AngleToleranceDeg;
+	Simplifier.Compute();
+
+	// compact result
+	EditMesh.CompactInPlace();
+
+	// recompute normals
+	FMeshNormals::InitializeOverlayTopologyFromOpeningAngle(&EditMesh, EditMesh.Attributes()->PrimaryNormals(), RecomputeNormalsAngleThreshold);
+	FMeshNormals::QuickRecomputeOverlayNormals(EditMesh);
+}
+
+
+
+
 static void ComputeBoxApproximation(
 	const FDynamicMesh3& SourceMesh,
 	FDynamicMesh3& OutputMesh )
@@ -1005,8 +1060,9 @@ void ComputeMeshApproximations(
 	ParallelFor(NumSets, [&](int32 Index)
 	{
 		TUniquePtr<FMeshInstanceSet>& InstanceSet = Assembly.InstanceSets[Index];
-		const FSourceGeometry& SourceGeo = Assembly.SourceMeshGeometry[Index];
-		const FDynamicMesh3& OptimizationSourceMesh = (CombineOptions.ApproximationSourceLOD < SourceGeo.SourceMeshLODs.Num()) ?
+		FSourceGeometry& SourceGeo = Assembly.SourceMeshGeometry[Index];
+		int32 NumSourceLODs = SourceGeo.SourceMeshLODs.Num();
+		const FDynamicMesh3& OptimizationSourceMesh = (CombineOptions.ApproximationSourceLOD < NumSourceLODs) ?
 			SourceGeo.SourceMeshLODs[CombineOptions.ApproximationSourceLOD] : SourceGeo.SourceMeshLODs.Last();
 		FOptimizedGeometry& ApproxGeo = Assembly.OptimizedMeshGeometry[Index];
 
@@ -1040,11 +1096,21 @@ void ComputeMeshApproximations(
 			FMeshNormals::QuickRecomputeOverlayNormals(ApproxGeo.ApproximateMeshLODs[k]);
 		}
 
+		// test remeshing last source LOD
+		if (CombineOptions.bRetriangulateSourceLODs)
+		{
+			for (int32 SourceLODIndex = CombineOptions.bRetriangulateSourceLODs; SourceLODIndex < NumSourceLODs; ++SourceLODIndex)
+			{
+				PlanarRetriangulatePartMesh(
+					SourceGeo.SourceMeshLODs[SourceLODIndex],
+					CombineOptions.SimplifyBaseTolerance, AngleThresholdDeg);
+			}
+		}
+
 	}, (bVerbose) ? EParallelForFlags::ForceSingleThread : EParallelForFlags::None );
 
 
 	// try to filter out simplifications that did bad things
-	// argh crashing!
 	ReplaceBadSimplifiedLODs(Assembly);
 }
 
@@ -2421,15 +2487,28 @@ void FCombineMeshInstancesImpl::CombineMeshInstances(
 
 
 	
-	InstanceAssembly.PreProcessInstanceMeshFunc = [&InstanceAssembly, &MeshInstances](FDynamicMesh3& AppendMesh, const FMeshInstance& Instance)
+	InstanceAssembly.PreProcessInstanceMeshFunc = [&InstanceAssembly, &MeshInstances, &Options](FDynamicMesh3& AppendMesh, const FMeshInstance& Instance)
 	{
 		int32 SourceInstance = Instance.ExternalInstanceIndex[0];
 		int GroupDataIdx = MeshInstances.StaticMeshInstances[SourceInstance].GroupDataIndex;
 		if (MeshInstances.InstanceGroupDatas[GroupDataIdx].bHasConstantOverrideVertexColor)
 		{
-			FColor VertexColorSRGB = MeshInstances.InstanceGroupDatas[GroupDataIdx].OverrideVertexColor;
-			//FLinearColor VertexColorLinear(VertexColorSRGB);
-			FLinearColor VertexColorLinear = VertexColorSRGB.ReinterpretAsLinear();
+			FLinearColor VertexColorLinear(0,0,0,1);
+			if (Options.VertexColorMappingMode == EVertexColorMappingMode::TriangleCountMetric)
+			{
+				const double UseMax = 25.0;
+				double TriCountRelToBox = FMath::Clamp((double)AppendMesh.TriangleCount() / (double)12, 1.0, UseMax);		// 12 is num tris in a bounding box
+				double T = (TriCountRelToBox) / (UseMax);
+				T = FMathd::Sqrt(T);	// improve color mapping somewhat (try better options?)
+				VertexColorLinear = FLinearColor::LerpUsingHSV(
+					FLinearColor::White, FLinearColor::Red, FMath::Clamp(T, 0.0, 1.0));
+			}
+			else
+			{
+				FColor VertexColorSRGB = MeshInstances.InstanceGroupDatas[GroupDataIdx].OverrideVertexColor;
+				VertexColorLinear = VertexColorSRGB.ReinterpretAsLinear();
+			}
+
 			SetConstantVertexColor(AppendMesh, VertexColorLinear);
 		}
 	};
diff --git a/Engine/Source/Developer/GeometryProcessingInterfaces/Public/GeometryProcessingInterfaces/CombineMeshInstances.h b/Engine/Source/Developer/GeometryProcessingInterfaces/Public/GeometryProcessingInterfaces/CombineMeshInstances.h
--- a/Engine/Source/Developer/GeometryProcessingInterfaces/Public/GeometryProcessingInterfaces/CombineMeshInstances.h
+++ b/Engine/Source/Developer/GeometryProcessingInterfaces/Public/GeometryProcessingInterfaces/CombineMeshInstances.h
@@ -76,6 +76,13 @@ public:
 		OcclusionBased = 6
 	};
 
+
+	enum class EVertexColorMappingMode
+	{
+		None = 0,
+		TriangleCountMetric = 1
+	};
+
 	struct FOptions
 	{
 		// number of requested LODs
@@ -123,6 +130,14 @@ public:
 		bool bMergeCoplanarFaces = true;
 		int32 MergeCoplanarFacesStartLOD = 1;
 		TFunction<UE::Geometry::FIndex3i(const UE::Geometry::FDynamicMesh3& Mesh, int32 TriangleID)> TriangleGroupingIDFunc;
+
+		// If enabled, attempt to retriangulate planar areas of Source LODs to remove redundant coplanar geometry
+		bool bRetriangulateSourceLODs = true;
+		// which Source LOD to start retriangulating at
+		int32 StartRetriangulateSourceLOD = 1;
+
+		// Color mapping modes for vertex colors, primarily used for debugging
+		EVertexColorMappingMode VertexColorMappingMode = EVertexColorMappingMode::None;
 	};
 
 
