Setting up weight editing UI in preparation for weight editing on selections.

diff --git a/Engine/Plugins/Editor/ModelingToolsEditorMode/Source/ModelingToolsEditorMode/Private/ModelingToolsManagerActions.cpp b/Engine/Plugins/Editor/ModelingToolsEditorMode/Source/ModelingToolsEditorMode/Private/ModelingToolsManagerActions.cpp
--- a/Engine/Plugins/Editor/ModelingToolsEditorMode/Source/ModelingToolsEditorMode/Private/ModelingToolsManagerActions.cpp
+++ b/Engine/Plugins/Editor/ModelingToolsEditorMode/Source/ModelingToolsEditorMode/Private/ModelingToolsManagerActions.cpp
@@ -331,10 +331,9 @@ void FModelingToolsManagerCommands::RegisterCommands()
 	REGISTER_MODELING_TOOL_COMMAND(BeginAddPatchTool, "Patch", "Add Patch");
 	REGISTER_MODELING_TOOL_COMMAND(BeginShapeSprayTool, "Spray", "Shape Spray");
 	REGISTER_MODELING_TOOL_COMMAND(BeginEditCollisionGeometryTool, "EditPhys", "Edit Simple Collision Geometry for selected Mesh");
-
-	// why are these ::Button ?
-	UI_COMMAND(BeginSkinWeightsPaintTool, "SkinWts", "Start the Paint Skin Weights Tool", EUserInterfaceActionType::Button, FInputChord());
-	UI_COMMAND(BeginSkinWeightsBindingTool, "SkinBind", "Start the Skin Weights Binding Tool", EUserInterfaceActionType::Button, FInputChord());
+	
+	REGISTER_MODELING_TOOL_COMMAND(BeginSkinWeightsPaintTool, "Weights", "Start the Skin Weights Editing Tool");
+	REGISTER_MODELING_TOOL_COMMAND(BeginSkinWeightsBindingTool, "Bind", "Start the Skin Weights Binding Tool");
 
 	REGISTER_MODELING_TOOL_COMMAND(BeginSkeletonEditingTool, "SkeletonEditing", "Start the Skeleton Editing Tool");
 
diff --git a/Engine/Plugins/Experimental/Animation/SkeletalMeshModelingTools/Source/SkeletalMeshModelingTools/Private/SkeletalMeshModelingToolsEditorModeToolkit.cpp b/Engine/Plugins/Experimental/Animation/SkeletalMeshModelingTools/Source/SkeletalMeshModelingTools/Private/SkeletalMeshModelingToolsEditorModeToolkit.cpp
--- a/Engine/Plugins/Experimental/Animation/SkeletalMeshModelingTools/Source/SkeletalMeshModelingTools/Private/SkeletalMeshModelingToolsEditorModeToolkit.cpp
+++ b/Engine/Plugins/Experimental/Animation/SkeletalMeshModelingTools/Source/SkeletalMeshModelingTools/Private/SkeletalMeshModelingToolsEditorModeToolkit.cpp
@@ -246,7 +246,7 @@ FText FSkeletalMeshModelingToolsEditorModeToolkit::GetActiveToolMessage() const
 	return ActiveToolMessage;
 }
 
-static const FName EditTabName(TEXT("Edit"));
+static const FName EditTabName(TEXT("Edit Mesh"));
 static const FName ProcessingTabName(TEXT("MeshOps"));
 static const FName DeformTabName(TEXT("Deform"));
 static const FName SkinWeightsTabName(TEXT("Skin"));
diff --git a/Engine/Plugins/Experimental/MeshModelingToolsetExp/Source/MeshModelingToolsExp/Private/MeshModelingToolsExp.cpp b/Engine/Plugins/Experimental/MeshModelingToolsetExp/Source/MeshModelingToolsExp/Private/MeshModelingToolsExp.cpp
--- a/Engine/Plugins/Experimental/MeshModelingToolsetExp/Source/MeshModelingToolsExp/Private/MeshModelingToolsExp.cpp
+++ b/Engine/Plugins/Experimental/MeshModelingToolsetExp/Source/MeshModelingToolsExp/Private/MeshModelingToolsExp.cpp
@@ -4,6 +4,8 @@
 
 #define LOCTEXT_NAMESPACE "FMeshModelingToolsExpModule"
 
+DEFINE_LOG_CATEGORY(LogMeshModelingTools);
+
 void FMeshModelingToolsExpModule::StartupModule()
 {
 	// This code will execute after your module is loaded into memory; the exact timing is specified in the .uplugin file per-module
diff --git a/Engine/Plugins/Experimental/MeshModelingToolsetExp/Source/MeshModelingToolsExp/Private/SkinWeightsPaintTool.cpp b/Engine/Plugins/Experimental/MeshModelingToolsetExp/Source/MeshModelingToolsExp/Private/SkinWeightsPaintTool.cpp
--- a/Engine/Plugins/Experimental/MeshModelingToolsetExp/Source/MeshModelingToolsExp/Private/SkinWeightsPaintTool.cpp
+++ b/Engine/Plugins/Experimental/MeshModelingToolsetExp/Source/MeshModelingToolsExp/Private/SkinWeightsPaintTool.cpp
@@ -4,19 +4,22 @@
 #include "Engine/SkeletalMesh.h"
 #include "InteractiveToolManager.h"
 #include "SkeletalMeshAttributes.h"
-#include "SkeletalDebugRendering.h"
 #include "Math/UnrealMathUtility.h"
 #include "Components/SkeletalMeshComponent.h"
 
 #include "TargetInterfaces/PrimitiveComponentBackedTarget.h"
 #include "ModelingToolTargetUtil.h"
-#include "ToolBuilderUtil.h"
 
 #include "MeshDescription.h"
+#include "MeshModelingToolsExp.h"
+#include "PointSetAdapter.h"
+#include "Animation/MirrorDataTable.h"
 #include "DynamicMesh/NonManifoldMappingSupport.h"
 #include "Parameterization/MeshLocalParam.h"
 #include "Spatial/FastWinding.h"
 #include "Async/Async.h"
+#include "DynamicMesh/MeshAdapterUtil.h"
+#include "Spatial/PointSetHashTable.h"
 
 #include UE_INLINE_GENERATED_CPP_BY_NAME(SkinWeightsPaintTool)
 
@@ -30,8 +33,7 @@ static EAsyncExecution SkinPaintToolAsyncExecTarget = EAsyncExecution::ThreadPoo
 // any weight below this value is ignored, since it won't be representable in unsigned 16-bit precision
 constexpr float MinimumWeightThreshold = 1.0f / 65535.0f;
 
-USkinWeightsPaintToolProperties::USkinWeightsPaintToolProperties(const FObjectInitializer& ObjectInitializer)
-		: Super(ObjectInitializer)
+USkinWeightsPaintToolProperties::USkinWeightsPaintToolProperties()
 {
 	ColorRamp.Add(FLinearColor::Blue);
 	ColorRamp.Add(FLinearColor::Yellow);
@@ -239,19 +241,20 @@ void FSkinToolWeights::InitializeSkinWeights(
 	}
 	
 	// maintain duplicate weight map
-	PreStrokeWeights = CurrentWeights;
+	PreChangeWeights = CurrentWeights;
 
 	// maintain relax-per stroke map
 	MaxFalloffPerVertexThisStroke.SetNumZeroed(NumVertices);
 }
 
 void FSkinToolWeights::EditVertexWeightAndNormalize(
-	const FName BoneToHoldConstant,
+	const int32 BoneToHoldIndex,
 	const int32 VertexID,
-	const float NewWeightValue,
+	float NewWeightValue,
 	FMultiBoneWeightEdits& WeightEdits)
 {
-	const int32 BoneToHoldIndex = Deformer.BoneNameToIndexMap[BoneToHoldConstant];
+	// clamp new weight
+	NewWeightValue = FMath::Clamp(NewWeightValue, 0.0f, 1.0f);
 	
 	// calculate the sum of all the weights on this vertex (not including the one we currently applied)
 	TArray<int32> BonesAffectingVertex;
@@ -275,9 +278,13 @@ void FSkinToolWeights::EditVertexWeightAndNormalize(
 		Total += VertexBoneData.Weight;
 	}
 
-	// if user applied FULL weight to this vertex OR there's no other weights of any significance,
-	// then simply set everything else to zero and return
-	if (NewWeightValue >= (1.0f - MinimumWeightThreshold) || Total <= MinimumWeightThreshold)
+	// are there no OTHER influences on this vertex?
+	const bool bVertexHasNoOtherInfluences = Total <= MinimumWeightThreshold;
+	
+	// if user applied FULL weight to this vertex AND there's no other weights of any significance,
+	// then simply apply full weight to this vertex, set all other influences to zero and return
+	const bool bApplyFullWeightToThisVertex = NewWeightValue >= (1.0f - MinimumWeightThreshold);
+	if (bApplyFullWeightToThisVertex && bVertexHasNoOtherInfluences)
 	{
 		// set all other influences to 0.0f
 		for (int32 i=0; i<ValuesToNormalize.Num(); ++i)
@@ -289,7 +296,7 @@ void FSkinToolWeights::EditVertexWeightAndNormalize(
 		}
 
 		// set current bone value to 1.0f
-		const float PrevWeight = GetWeightOfBoneOnVertex(BoneToHoldIndex, VertexID, PreStrokeWeights);
+		const float PrevWeight = GetWeightOfBoneOnVertex(BoneToHoldIndex, VertexID, PreChangeWeights);
 		WeightEdits.MergeSingleEdit(
 			BoneToHoldIndex,
 			VertexID,
@@ -299,6 +306,44 @@ void FSkinToolWeights::EditVertexWeightAndNormalize(
 		return;
 	}
 
+	// is the user trying to prune ALL weight from this vertex, AND all other weights are equal to zero?
+	// in this case, we have two options:
+	// 1. if there are other influences recorded for this vertex, then split remaining influence among them
+	// 2. if there are NO other influences recorded for this vertex, then move the weight to the root as last ditch effort
+	const bool bApplyZeroWeightToThisVertex = NewWeightValue <= MinimumWeightThreshold;
+	if (bApplyZeroWeightToThisVertex && bVertexHasNoOtherInfluences)
+	{
+		if (ValuesToNormalize.IsEmpty())
+		{
+			// assign all weight to the root
+			constexpr int32 BoneIndex = 0;
+			constexpr float OldWeight = 0.f;
+			constexpr float NewWeight = 1.f;
+			WeightEdits.MergeSingleEdit(BoneIndex, VertexID, OldWeight, NewWeight);
+		}
+		else
+		{
+			// evenly distribute weight among other influences
+			const float NewWeight = 1.f / ValuesToNormalize.Num();
+			for (int32 i=0; i<ValuesToNormalize.Num(); ++i)
+			{
+				const int32 BoneIndex = BonesAffectingVertex[i];
+				const float OldWeight = ValuesToNormalize[i];
+				WeightEdits.MergeSingleEdit(BoneIndex, VertexID, OldWeight, NewWeight);
+			}
+		}
+
+		// set current bone value to 0.0f
+		const float PrevWeight = GetWeightOfBoneOnVertex(BoneToHoldIndex, VertexID, PreChangeWeights);
+		WeightEdits.MergeSingleEdit(
+			BoneToHoldIndex,
+			VertexID,
+			PrevWeight,
+			0.f);
+		
+		return;
+	}
+
 	// calculate amount we have to spread across the other bones affecting this vertex
 	const float AvailableTotal = 1.0f - NewWeightValue;
 
@@ -317,7 +362,7 @@ void FSkinToolWeights::EditVertexWeightAndNormalize(
 	}
 
 	// record current bone edit
-	const float PrevWeight = GetWeightOfBoneOnVertex(BoneToHoldIndex, VertexID, PreStrokeWeights);
+	const float PrevWeight = GetWeightOfBoneOnVertex(BoneToHoldIndex, VertexID, PreChangeWeights);
 	WeightEdits.MergeSingleEdit(
 		BoneToHoldIndex,
 		VertexID,
@@ -433,9 +478,9 @@ void FSkinToolWeights::SetWeightOfBoneOnVertex(
 	}
 }
 
-void FSkinToolWeights::ResetAfterStroke()
+void FSkinToolWeights::ResetAfterChange()
 {
-	PreStrokeWeights = CurrentWeights;
+	PreChangeWeights = CurrentWeights;
 
 	for (int32 i=0; i<MaxFalloffPerVertexThisStroke.Num(); ++i)
 	{
@@ -519,18 +564,17 @@ void USkinWeightsPaintTool::Setup()
 	EditedMesh = MakeUnique<FMeshDescription>();
 	*EditedMesh = *UE::ToolTarget::GetMeshDescription(Target);
 
-	// initialize the tool properties
-	BrushProperties->RestoreProperties(this); // hides strength and falloff
-	
-	ToolProps = NewObject<USkinWeightsPaintToolProperties>(this);
-	ToolProps->RestoreProperties(this);
-	ToolProps->SkeletalMesh = Component->GetSkeletalMeshAsset();
-	AddToolPropertySource(ToolProps);
-	// attach callback to be informed when tool properties are modified
-	ToolProps->GetOnModified().AddUObject(this, &USkinWeightsPaintTool::OnToolPropertiesModified);
+	// create a custom set of properties inheriting from the base tool properties
+	WeightToolProperties = NewObject<USkinWeightsPaintToolProperties>(this);
+	WeightToolProperties->RestoreProperties(this);
+	WeightToolProperties->WeightTool = this;
+	WeightToolProperties->bSpecifyRadius = true;
+	// replace the base brush properties
+	ReplaceToolPropertySource(BrushProperties, WeightToolProperties);
+	BrushProperties = WeightToolProperties;
 
 	// default to the root bone as current bone
-	PendingCurrentBone = CurrentBone = ToolProps->SkeletalMesh->GetRefSkeleton().GetBoneName(0);
+	PendingCurrentBone = CurrentBone = Component->GetSkeletalMeshAsset()->GetRefSkeleton().GetBoneName(0);
 
 	// configure preview mesh
 	PreviewMesh->SetTangentsMode(EDynamicMeshComponentTangentsMode::AutoCalculated);
@@ -597,11 +641,11 @@ void USkinWeightsPaintTool::OnTick(float DeltaTime)
 		PendingCurrentBone.Reset();
 	}
 
-	if (bVisibleWeightsValid == false || ToolProps->bColorModeChanged)
+	if (bVisibleWeightsValid == false || WeightToolProperties->bColorModeChanged)
 	{
 		UpdateCurrentBoneVertexColors();
 		bVisibleWeightsValid = true;
-		ToolProps->bColorModeChanged = false;
+		WeightToolProperties->bColorModeChanged = false;
 	}
 
 	// sparsely updates vertex positions (only on vertices with modified weights)
@@ -691,15 +735,10 @@ void USkinWeightsPaintTool::OnEndDrag(const FRay& Ray)
 	bSmoothStroke = false;
 	bStampPending = false;
 
-	Weights.ResetAfterStroke();
-
 	// close change record
 	TUniquePtr<FMeshSkinWeightsChange> Change = EndChange();
-
-	GetToolManager()->BeginUndoTransaction(LOCTEXT("BoneWeightValuesChange", "Paint"));
-
-	GetToolManager()->EmitObjectChange(this, MoveTemp(Change), LOCTEXT("BoneWeightValuesChange", "Paint"));
-
+	GetToolManager()->BeginUndoTransaction(LOCTEXT("PaintWeightChange", "Paint skin weights."));
+	GetToolManager()->EmitObjectChange(this, MoveTemp(Change), LOCTEXT("PaintWeightChange", "Paint skin weights."));
 	GetToolManager()->EndUndoTransaction();
 }
 
@@ -711,14 +750,21 @@ bool USkinWeightsPaintTool::OnUpdateHover(const FInputDeviceRay& DevicePos)
 
 void USkinWeightsPaintTool::CalculateVertexROI(
 	const FBrushStampData& Stamp,
-	TArray<int32>& VertexROI,
-	TArray<float>& VertexSqDistances) const
+	TArray<VertexIndex>& VertexIDs,
+	TArray<float>& VertexFalloffs)
 {
 	using namespace UE::Geometry;
 
 	TRACE_CPUPROFILER_EVENT_SCOPE(SkinTool::CalculateVertexROI);
+
+	auto DistanceToFalloff = [this](int32 InVertexID, float InDistanceSq)-> float
+	{
+		const float CurrentFalloff = CalculateBrushFalloff(FMath::Sqrt(InDistanceSq));
+		const float UseFalloff = Weights.SetCurrentFalloffAndGetMaxFalloffThisStroke(InVertexID, CurrentFalloff);
+		return UseFalloff;
+	};
 	
-	if (ToolProps->FalloffMode == EWeightBrushFalloffMode::Volume)
+	if (WeightToolProperties->FalloffMode == EWeightBrushFalloffMode::Volume)
 	{
 		const IPrimitiveComponentBackedTarget* TargetComponent = Cast<IPrimitiveComponentBackedTarget>(Target);
 		const FTransform3d Transform(TargetComponent->GetWorldTransform());
@@ -728,17 +774,18 @@ void USkinWeightsPaintTool::CalculateVertexROI(
 		const FAxisAlignedBox3d QueryBox(StampPosLocal, CurrentBrushRadius);
 		VerticesOctree.RangeQuery(QueryBox,
 			[&](int32 VertexID) { return FVector3d::DistSquared(Mesh->GetVertex(VertexID), StampPosLocal) < RadiusSqr; },
-			VertexROI);
+			VertexIDs);
 
-		for (const int32 VertexID : VertexROI)
+		for (const int32 VertexID : VertexIDs)
 		{
-			VertexSqDistances.Add(FVector3d::DistSquared(Mesh->GetVertex(VertexID), StampPosLocal));
+			const float DistSq = FVector3d::DistSquared(Mesh->GetVertex(VertexID), StampPosLocal);
+			VertexFalloffs.Add(DistanceToFalloff(VertexID, DistSq));
 		}
 		
 		return;
 	}
 
-	if (ToolProps->FalloffMode == EWeightBrushFalloffMode::Surface)
+	if (WeightToolProperties->FalloffMode == EWeightBrushFalloffMode::Surface)
 	{
 		// get coordinate frame from stamp
 		auto GetFrameFromStamp = [](const FBrushStampData& InStamp) -> FFrame3d
@@ -786,9 +833,10 @@ void USkinWeightsPaintTool::CalculateVertexROI(
 				{
 					continue;
 				}
+
 				
-				VertexSqDistances.Add(DistSq);
-				VertexROI.Add(VertexID);
+				VertexFalloffs.Add(DistanceToFalloff(VertexID, DistSq));
+				VertexIDs.Add(VertexID);
 			}
 		}
 		
@@ -801,7 +849,7 @@ void USkinWeightsPaintTool::CalculateVertexROI(
 FVector4f USkinWeightsPaintTool::WeightToColor(float Value) const
 {
 	// optional greyscale mode
-	if (ToolProps->ColorMode == EWeightColorMode::Greyscale)
+	if (WeightToolProperties->ColorMode == EWeightColorMode::Greyscale)
 	{
 		return FLinearColor::LerpUsingHSV(FLinearColor::Black, FLinearColor::White, Value);	
 	}
@@ -809,22 +857,22 @@ FVector4f USkinWeightsPaintTool::WeightToColor(float Value) const
 	// early out zero weights to min color
 	if (Value <= MinimumWeightThreshold)
 	{
-		return ToolProps->MinColor;
+		return WeightToolProperties->MinColor;
 	}
 
 	// early out full weights to max color
 	if (FMath::IsNearlyEqual(Value, 1.0f))
 	{
-		return ToolProps->MaxColor;
+		return WeightToolProperties->MaxColor;
 	}
 
 	// get user-specified color ramp for intermediate colors
-	const TArray<FLinearColor>& Colors = ToolProps->ColorRamp;
+	const TArray<FLinearColor>& Colors = WeightToolProperties->ColorRamp;
 	
 	// revert back to simple Lerp(min,max) if user supplied color ramp doesn't have enough colors
 	if (Colors.Num() < 2)
 	{
-		const FLinearColor FinalColor = FLinearColor::LerpUsingHSV(ToolProps->MinColor, ToolProps->MaxColor, Value);
+		const FLinearColor FinalColor = FLinearColor::LerpUsingHSV(WeightToolProperties->MinColor, WeightToolProperties->MaxColor, Value);
 		return UE::Geometry::ToVector4<float>(FinalColor);
 	}
 
@@ -844,7 +892,7 @@ FVector4f USkinWeightsPaintTool::WeightToColor(float Value) const
 
 void USkinWeightsPaintTool::UpdateCurrentBoneVertexColors()
 {
-	const int32 CurrentBoneIndex = Weights.Deformer.BoneNameToIndexMap[CurrentBone];
+	const int32 CurrentBoneIndex = CurrentBone == NAME_None ? INDEX_NONE : Weights.Deformer.BoneNameToIndexMap[CurrentBone];
 	
 	// update mesh with new value colors
 	PreviewMesh->DeferredEditMesh([this, &CurrentBoneIndex](FDynamicMesh3& Mesh)
@@ -853,6 +901,13 @@ void USkinWeightsPaintTool::UpdateCurrentBoneVertexColors()
 		UE::Geometry::FDynamicMeshColorOverlay* ColorOverlay = Mesh.Attributes()->PrimaryColors();
 		for (const int32 ElementId : ColorOverlay->ElementIndicesItr())
 		{
+			// with no bone selected, all vertices are drawn black
+			if (CurrentBoneIndex == INDEX_NONE)
+			{
+				ColorOverlay->SetElement(ElementId, FVector4f(FLinearColor::Black));
+				continue;
+			}
+			
 			const int32 VertexID = ColorOverlay->GetParentVertex(ElementId);	
 			const int32 SrcVertexID = NonManifoldMappingSupport.GetOriginalNonManifoldVertexID(VertexID);
 			const float Value = Weights.GetWeightOfBoneOnVertex(CurrentBoneIndex, SrcVertexID, Weights.CurrentWeights);
@@ -863,10 +918,10 @@ void USkinWeightsPaintTool::UpdateCurrentBoneVertexColors()
 	PreviewMesh->NotifyDeferredEditCompleted(UPreviewMesh::ERenderUpdateMode::FastUpdate, EMeshRenderAttributeFlags::VertexColors, false);
 }
 
-double USkinWeightsPaintTool::CalculateBrushFalloff(double Distance) const
+float USkinWeightsPaintTool::CalculateBrushFalloff(float Distance) const
 {
-	const double f = FMathd::Clamp(1.0 - BrushProperties->BrushFalloffAmount, 0.0, 1.0);
-	double d = Distance / CurrentBrushRadius;
+	const float f = FMathd::Clamp(1.f - BrushProperties->BrushFalloffAmount, 0.f, 1.f);
+	float d = Distance / CurrentBrushRadius;
 	double w = 1;
 	if (d > f)
 	{
@@ -880,13 +935,19 @@ double USkinWeightsPaintTool::CalculateBrushFalloff(double Distance) const
 void USkinWeightsPaintTool::ApplyStamp(const FBrushStampData& Stamp)
 {
 	TRACE_CPUPROFILER_EVENT_SCOPE(SkinTool::ApplyStamp);
+
+	// must select a bone to paint
+	if (CurrentBone == NAME_None)
+	{
+		return;
+	}
 	
 	// get the vertices under the brush, and their squared distances to the brush center
 	// when using "Volume" brush, distances are straight line
 	// when using "Surface" brush, distances are geodesics
 	TArray<int32> VerticesInStamp;
-	TArray<float> VertexSqDistances;
-	CalculateVertexROI(Stamp, VerticesInStamp, VertexSqDistances);
+	TArray<float> VertexFalloffs;
+	CalculateVertexROI(Stamp, VerticesInStamp, VertexFalloffs);
 
 	// gather sparse set of modifications made from this stamp, these edits are merged throughout
 	// the lifetime of a single brush stroke in the "ActiveChange" allowing for undo/redo
@@ -894,18 +955,21 @@ void USkinWeightsPaintTool::ApplyStamp(const FBrushStampData& Stamp)
 	{
 		TRACE_CPUPROFILER_EVENT_SCOPE(SkinTool::EditWeightOfVerticesInStamp);
 		// generate a weight edit from this stamp (includes modifications caused by normalization)
-		if (bSmoothStroke || ToolProps->BrushMode == EBrushBehaviorMode::Relax)
+		if (bSmoothStroke || WeightToolProperties->BrushMode == EWeightEditOperation::Relax)
 		{
 			// use mesh topology to iteratively smooth weights across neighboring vertices
-			RelaxWeightOnVertices(VerticesInStamp, VertexSqDistances, WeightEditsFromStamp);
+			const float UseStrength = CalculateBrushStrengthToUse(EWeightEditOperation::Relax);
+			RelaxWeightOnVertices(VerticesInStamp, VertexFalloffs, UseStrength, WeightEditsFromStamp);
 		}
 		else
 		{
 			// edit weight; either by "Add", "Remove", "Replace", "Multiply"
-			EditWeightOfVerticesInStamp(
-				ToolProps->BrushMode,
+			const float UseStrength = CalculateBrushStrengthToUse(WeightToolProperties->BrushMode);
+			EditWeightOnVertices(
+				WeightToolProperties->BrushMode,
 				VerticesInStamp,
-				VertexSqDistances,
+				VertexFalloffs,
+				UseStrength,
 				WeightEditsFromStamp);
 		}
 	}
@@ -951,10 +1015,97 @@ void USkinWeightsPaintTool::ApplyStamp(const FBrushStampData& Stamp)
 	}
 }
 
+float USkinWeightsPaintTool::CalculateBrushStrengthToUse(EWeightEditOperation EditMode) const
+{
+	float UseStrength = BrushProperties->BrushStrength;
+
+	// invert brush strength differently depending on brush mode
+	switch (EditMode)
+	{
+	case EWeightEditOperation::Add:
+		{
+			UseStrength *= bInvertStroke ? -1.0f : 1.0f;
+			break;
+		}
+	case EWeightEditOperation::Replace:
+		{
+			UseStrength = bInvertStroke ? 1.0f - UseStrength : UseStrength;
+			break;
+		}
+	case EWeightEditOperation::Multiply:
+		{
+			UseStrength = bInvertStroke ? 1.0f + UseStrength : UseStrength;
+			break;
+		}
+	case EWeightEditOperation::Relax:
+		{
+			UseStrength = bInvertStroke ? 1.0f - UseStrength : UseStrength;
+			break;
+		}
+	default:
+		checkNoEntry();
+	}
+
+	return UseStrength;
+}
+
+void USkinWeightsPaintTool::EditWeightOnVertices(
+	EWeightEditOperation EditOperation,
+	const TArray<int32>& VerticesToEdit,
+	const TArray<float>& VertexFalloffs,
+	const float UseStrength,
+	FMultiBoneWeightEdits& InOutWeightEdits)
+{
+	// spin through the vertices in the stamp and store new weight values in NewValuesFromStamp
+	// afterwards, these values are normalized while taking into consideration the user's desired changes
+	const int32 CurrentBoneIndex = Weights.Deformer.BoneNameToIndexMap[CurrentBone];
+	const int32 NumVerticesInStamp = VerticesToEdit.Num();
+	for (int32 Index = 0; Index < NumVerticesInStamp; ++Index)
+	{
+		const int32 VertexID = VerticesToEdit[Index];
+		const float UseFalloff = VertexFalloffs.IsValidIndex(Index) ? VertexFalloffs[Index] : 1.f;
+		const float ValueBeforeStroke = Weights.GetWeightOfBoneOnVertex(CurrentBoneIndex, VertexID, Weights.PreChangeWeights);
+
+		// calculate new weight value
+		float NewValueAfterStamp = ValueBeforeStroke;
+		switch (EditOperation)
+		{
+		case EWeightEditOperation::Add:
+			{
+				NewValueAfterStamp = ValueBeforeStroke + (UseStrength * UseFalloff);
+				break;
+			}
+		case EWeightEditOperation::Replace:
+			{
+				NewValueAfterStamp = FMath::Lerp(ValueBeforeStroke, UseStrength, UseFalloff);
+				break;
+			}
+		case EWeightEditOperation::Multiply:
+			{
+				const float DeltaFromThisStamp = ((ValueBeforeStroke * UseStrength) - ValueBeforeStroke) * UseFalloff;
+				NewValueAfterStamp = ValueBeforeStroke + DeltaFromThisStamp;
+				break;
+			}
+		default:
+			// relax operation not supported by this function, use RelaxWeightOnVertices()
+			checkNoEntry();
+		}
+
+		// normalize the values across all bones affecting the vertices in the stamp, and record the bone edits
+		// normalization is done while holding all weights on the current bone constant so that user edits are not overwritten
+		Weights.EditVertexWeightAndNormalize(
+			CurrentBoneIndex,
+			VertexID,
+			NewValueAfterStamp,
+			InOutWeightEdits);
+	}
+}
+
 void USkinWeightsPaintTool::RelaxWeightOnVertices(
-	TArray<int32> VerticesInStamp,
-	TArray<float> VertexSqDistances,
-	FMultiBoneWeightEdits& AllBoneWeightEditsFromStamp)
+	TArray<int32> VerticesToEdit,
+	TArray<float> VertexFalloffs,
+	const float UseStrength,
+	FMultiBoneWeightEdits& InOutWeightEdits)
 {
 	const FDynamicMesh3* CurrentMesh = PreviewMesh->GetMesh();
 	const UE::Geometry::FNonManifoldMappingSupport NonManifoldMappingSupport(*CurrentMesh);
@@ -978,9 +1129,9 @@ void USkinWeightsPaintTool::RelaxWeightOnVertices(
 	TArray<int32> AllNeighborVertices;
 	TMap<BoneIndex, VertexNeighborWeights> WeightsOnAllNeighbors;
 	TMap<BoneIndex, float> FinalWeights;
-	for (int32 Index = 0; Index < VerticesInStamp.Num(); ++Index)
+	for (int32 Index = 0; Index < VerticesToEdit.Num(); ++Index)
 	{
-		const int32 VertexID = VerticesInStamp[Index];
+		const int32 VertexID = VerticesToEdit[Index];
 		const int32 SrcVertexID = NonManifoldMappingSupport.GetOriginalNonManifoldVertexID(VertexID);
 
 		// get list of all neighboring vertices, AND this vertex
@@ -995,7 +1146,7 @@ void USkinWeightsPaintTool::RelaxWeightOnVertices(
 		WeightsOnAllNeighbors.Reset();
 		for (const int32 NeighborVertexID : AllNeighborVertices)
 		{
-			for (const FVertexBoneWeight& BoneWeight : Weights.PreStrokeWeights[NeighborVertexID])
+			for (const FVertexBoneWeight& BoneWeight : Weights.PreChangeWeights[NeighborVertexID])
 			{
 				if (BoneWeight.Weight > MinimumWeightThreshold)
 				{
@@ -1021,15 +1172,13 @@ void USkinWeightsPaintTool::RelaxWeightOnVertices(
 		NormalizeWeights(FinalWeights);
 
 		// lerp weights from previous values, to fully relaxed values by brush strength scaled by falloff
-		const float CurrentFalloff = static_cast<float>(CalculateBrushFalloff(FMath::Sqrt(VertexSqDistances[Index])));
-		const float UseFalloff = Weights.SetCurrentFalloffAndGetMaxFalloffThisStroke(VertexID, CurrentFalloff);
-		const float UseStrength = BrushProperties->BrushStrength * UseFalloff;
+		const float UseFalloff = VertexFalloffs.IsValidIndex(Index) ? VertexFalloffs[Index] * UseStrength : UseStrength;
 		for (TTuple<BoneIndex, float>& FinalWeight : FinalWeights)
 		{
 			const int32 BoneIndex = FinalWeight.Key;
 			float NewWeight = FinalWeight.Value;
-			float OldWeight = Weights.GetWeightOfBoneOnVertex(BoneIndex, VertexID, Weights.PreStrokeWeights);
-			FinalWeight.Value = FMath::Lerp(OldWeight, NewWeight, UseStrength);
+			float OldWeight = Weights.GetWeightOfBoneOnVertex(BoneIndex, VertexID, Weights.PreChangeWeights);
+			FinalWeight.Value = FMath::Lerp(OldWeight, NewWeight, UseFalloff);
 		}
 
 		// normalize again
@@ -1041,99 +1190,88 @@ void USkinWeightsPaintTool::RelaxWeightOnVertices(
 			// record an edit for this vertex, for this bone
 			const int32 BoneIndex = FinalWeight.Key;
 			const float NewWeight = FinalWeight.Value;
-			const float OldWeight = Weights.GetWeightOfBoneOnVertex(BoneIndex, VertexID, Weights.PreStrokeWeights);
-			AllBoneWeightEditsFromStamp.MergeSingleEdit(BoneIndex, VertexID, OldWeight, NewWeight);
+			const float OldWeight = Weights.GetWeightOfBoneOnVertex(BoneIndex, VertexID, Weights.PreChangeWeights);
+			InOutWeightEdits.MergeSingleEdit(BoneIndex, VertexID, OldWeight, NewWeight);
 		}
 	}
 }
 
-void USkinWeightsPaintTool::EditWeightOfVerticesInStamp(
-	EBrushBehaviorMode EditMode,
-	const TArray<int32>& VerticesInStamp,
-	const TArray<float>& VertexSqDistances,
-	FMultiBoneWeightEdits& AllBoneWeightEditsFromStamp)
+void USkinWeightsPaintTool::ApplyWeightEditsToMesh(
+	const FText& TransactionLabel,
+	const SkinPaintTool::FMultiBoneWeightEdits& WeightEdits)
 {
-	const UE::Geometry::FNonManifoldMappingSupport NonManifoldMappingSupport(*PreviewMesh->GetMesh());
+	// clear the active change
+	BeginChange();
 
-	// invert brush strength differently depending on brush mode
-	float UseStrength = BrushProperties->BrushStrength;
-	switch (EditMode)
+	// store weight edits in the active change
+	for (const TTuple<int32, FSingleBoneWeightEdits>& BoneWeightEdits : WeightEdits.PerBoneWeightEdits)
 	{
-	case EBrushBehaviorMode::Add:
-		{
-			UseStrength *= bInvertStroke ? -1.0f : 1.0f;
-			break;
-		}
-	case EBrushBehaviorMode::Replace:
-		{
-			UseStrength = bInvertStroke ? 1.0f - UseStrength : UseStrength;
-			break;
-		}
-	case EBrushBehaviorMode::Multiply:
-		{
-			UseStrength = bInvertStroke ? 1.0f + UseStrength : UseStrength;
-			break;
-		}		
-	default:
-		checkNoEntry();
+		ActiveChange->AddBoneWeightEdit(BoneWeightEdits.Value);
 	}
-	
-	// spin through the vertices in the stamp and store new weight values in NewValuesFromStamp
-	// afterwards, these values are normalized while taking into consideration the user's desired changes
-	const int32 CurrentBoneIndex = Weights.Deformer.BoneNameToIndexMap[CurrentBone];
-	const int32 NumVerticesInStamp = VerticesInStamp.Num();
-	for (int32 Index = 0; Index < NumVerticesInStamp; ++Index)
-	{
-		const int32 VertexID = VerticesInStamp[Index];
-		const int32 SrcVertexID = NonManifoldMappingSupport.GetOriginalNonManifoldVertexID(VertexID);
-		const float CurrentFalloff = static_cast<float>(CalculateBrushFalloff(FMath::Sqrt(VertexSqDistances[Index])));
-		const float UseFalloff = Weights.SetCurrentFalloffAndGetMaxFalloffThisStroke(SrcVertexID, CurrentFalloff);
-		const float ValueBeforeStroke = Weights.GetWeightOfBoneOnVertex(CurrentBoneIndex, VertexID, Weights.PreStrokeWeights);
 
-		// calculate new weight value
-		float NewValueAfterStamp = ValueBeforeStroke;
-		switch (EditMode)
+	// apply the weight edits to the actual mesh
+	// (copies weight modifications to the tool's weight data structure and updates the vertex colors
+	// in brush mode, this is normally called by the tool itself after a stroke)
+	ActiveChange->Apply(this);
+
+	// store active change in the transaction buffer
+	UInteractiveToolManager* ToolManager = GetToolManager();
+	ToolManager->BeginUndoTransaction(TransactionLabel);
+	ToolManager->EmitObjectChange(this, EndChange(), TransactionLabel);
+	ToolManager->EndUndoTransaction();
+}
+
+void USkinWeightsPaintTool::UpdateCurrentBone(const FName& BoneName)
+{
+	CurrentBone = BoneName;
+	bVisibleWeightsValid = false;
+}
+
+void USkinWeightsPaintTool::GetVerticesToEdit(TArray<VertexIndex>& OutVertexIndices) const
+{
+	OutVertexIndices.Reset();
+
+	//
+	// 1. Prioritize selected vertices
+	// TODO check selection manager for selected vertices here
+
+	//
+	// 2. Fallback on vertices weighted to selected bones
+	if (!SelectedBoneIndices.IsEmpty())
+	{
+		VertexIndex VertexID = 0;
+		for (const VertexWeights& VertWeights : Weights.PreChangeWeights)
 		{
-		case EBrushBehaviorMode::Add:
-			{
-				NewValueAfterStamp = ValueBeforeStroke + (UseStrength * UseFalloff);
-				break;
-			}
-		case EBrushBehaviorMode::Replace:
+			for (const FVertexBoneWeight& BoneWeight : VertWeights)
 			{
-				NewValueAfterStamp = FMath::Lerp(ValueBeforeStroke, UseStrength, UseFalloff);
-				break;
-			}
-		case EBrushBehaviorMode::Multiply:
-			{
-				const float DeltaFromThisStamp = ((ValueBeforeStroke * UseStrength) - ValueBeforeStroke) * UseFalloff;
-				NewValueAfterStamp = ValueBeforeStroke + DeltaFromThisStamp;
-				break;
+				if (SelectedBoneIndices.Contains(BoneWeight.BoneIndex) && BoneWeight.Weight > MinimumWeightThreshold)
+				{
+					OutVertexIndices.Add(VertexID);	
+				}
 			}
-		default:
-			checkNoEntry();
+			++VertexID;
 		}
+		return;
+	}
 
-		// normalize the values across all bones affecting the vertices in the stamp, and record the bone edits
-		// normalization is done while holding all weights on the current bone constant so that user edits are not overwritten
-		NewValueAfterStamp = FMath::Clamp(NewValueAfterStamp, 0.0f, 1.0f);
-		Weights.EditVertexWeightAndNormalize(
-			CurrentBone,
-			VertexID,
-			NewValueAfterStamp,
-			AllBoneWeightEditsFromStamp);
+	//
+	// 3. Finally, fallback to ALL vertices in the mesh
+	const int32 NumVertices = EditedMesh->Vertices().Num();
+	OutVertexIndices.Reset(NumVertices);
+	for (VertexIndex VertexID = 0; VertexID < NumVertices; VertexID++)
+	{
+		OutVertexIndices.Add(VertexID);
 	}
 }
 
-void USkinWeightsPaintTool::UpdateCurrentBone(const FName& BoneName)
+BoneIndex USkinWeightsPaintTool::GetBoneIndexFromName(const FName BoneName) const
 {
-	CurrentBone = BoneName;
-	bVisibleWeightsValid = false;
+	return  BoneName == NAME_None ? INDEX_NONE : Weights.Deformer.BoneNameToIndexMap[BoneName];
 }
 
 void USkinWeightsPaintTool::OnShutdown(EToolShutdownType ShutdownType)
 {
-	BrushProperties->SaveProperties(this);
+	WeightToolProperties->SaveProperties(this);
 
 	if (ShutdownType == EToolShutdownType::Accept)
 	{
@@ -1152,13 +1290,12 @@ void USkinWeightsPaintTool::BeginChange()
 	ActiveChange = MakeUnique<FMeshSkinWeightsChange>();
 }
 
-
 TUniquePtr<FMeshSkinWeightsChange> USkinWeightsPaintTool::EndChange()
 {
+	Weights.ResetAfterChange();
 	return MoveTemp(ActiveChange);
 }
 
-
 void USkinWeightsPaintTool::ExternalUpdateWeights(const int32 BoneIndex, const TMap<int32, float>& NewValues)
 {
 	for (const TTuple<int32, float>& Pair : NewValues)
@@ -1166,7 +1303,7 @@ void USkinWeightsPaintTool::ExternalUpdateWeights(const int32 BoneIndex, const T
 		const int32 VertexID = Pair.Key;
 		const float Weight = Pair.Value;
 		Weights.SetWeightOfBoneOnVertex(BoneIndex, VertexID, Weight, Weights.CurrentWeights);
-		Weights.SetWeightOfBoneOnVertex(BoneIndex, VertexID, Weight, Weights.PreStrokeWeights);
+		Weights.SetWeightOfBoneOnVertex(BoneIndex, VertexID, Weight, Weights.PreChangeWeights);
 	}
 
 	const FName BoneName = Weights.Deformer.BoneNames[BoneIndex];
@@ -1176,10 +1313,307 @@ void USkinWeightsPaintTool::ExternalUpdateWeights(const int32 BoneIndex, const T
 	}
 }
 
-void USkinWeightsPaintTool::HandleSkeletalMeshModified(const TArray<FName>& InBoneNames, const ESkeletalMeshNotifyType InNotifyType)
+void USkinWeightsPaintTool::MirrorWeights(EAxis::Type Axis, EMirrorDirection Direction)
 {
-	const FName BoneName = InBoneNames.IsEmpty() ? NAME_None : InBoneNames[0];
+	check(Axis != EAxis::None);
+	
+	// TODO, filter vertices by selection
+	const TArray<FVector>& Vertices = Weights.Deformer.RefPoseVertexPositions;
+
+	// TODO, provide some way to edit the mirror bone mapping, either by providing a UMirrorDataTable input or editing directly in the hierarchy view.
+	// build bone map for mirroring
+	TMap<int32, int32> MirrorBoneMap;
+	const FReferenceSkeleton& RefSkeleton = Weights.Deformer.Component->GetSkeletalMeshAsset()->GetRefSkeleton();
+	for (FName BoneName : Weights.Deformer.BoneNames)
+	{
+		FName MirroredBoneName = UMirrorDataTable::FindBestMirroredBone(BoneName, RefSkeleton, Axis);
+
+		int32 BoneIndex = Weights.Deformer.BoneNameToIndexMap[BoneName];
+		int32 MirroredBoneIndex = Weights.Deformer.BoneNameToIndexMap[MirroredBoneName];
+		MirrorBoneMap.Add(BoneIndex, MirroredBoneIndex);
+		
+		// debug view bone mapping
+		//UE_LOG(LogTemp, Log, TEXT("Bone    : %s"), *BoneName.ToString());
+		//UE_LOG(LogTemp, Log, TEXT("Mirrored: %s"), *MirroredBoneName.ToString());
+		//UE_LOG(LogTemp, Log, TEXT("-------"));
+	}
+
+	// hash grid constants
+	constexpr float HashGridCellSize = 2.0f;
+	constexpr float ThresholdRadius = 0.1f;
+	const float MaxSearchRadius = Weights.Deformer.Component->Bounds.SphereRadius;
+
+	// build a point set of the rest pose vertices
+	UE::Geometry::FDynamicPointSet3d PointSet;
+	for (int32 PointID = 0; PointID < Vertices.Num(); ++PointID)
+	{
+		PointSet.InsertVertex(PointID, Vertices[PointID]);
+	}
+
+	// build a spatial hash map from the point set
+	UE::Geometry::FPointSetAdapterd PointSetAdapter = UE::Geometry::MakePointsAdapter(&PointSet);
+	UE::Geometry::FPointSetHashtable PointHash(&PointSetAdapter);
+	PointHash.Build(HashGridCellSize, FVector3d::Zero());
+
+	// generate a map of point IDs on the target side, to their equivalent vertex ID on the source side 
+	TArray<TPair<int32, int32>> VertexMirrorMap; // TMap<TO, FROM>
+	TArray<int> PointsInThreshold;
+	TArray<int> PointsInSphere;
+	bool bSomeVerticesNotMirrored = false;
+	for (int32 TargetVertexID = 0; TargetVertexID < Vertices.Num(); ++TargetVertexID)
+	{
+		const FVector& TargetPosition = Vertices[TargetVertexID];
+
+		if (Direction == EMirrorDirection::PositiveToNegative && TargetPosition[Axis-1] >= 0.f)
+		{
+			continue; // copying to negative side, but vertex is on positive side
+		}
+		if (Direction == EMirrorDirection::NegativeToPositive && TargetPosition[Axis-1] <= 0.f)
+		{
+			continue; // copying to positive side, but vertex is on negative side
+		}
+
+		// flip position across the mirror axis
+		FVector MirroredPosition = TargetPosition;
+		MirroredPosition[Axis-1] *= -1.f;
+
+		// query spatial hash near mirrored position, gradually increasing search radius until at least 1 point is found
+		PointsInSphere.Reset();
+		float SearchRadius = ThresholdRadius;
+		while(PointsInSphere.IsEmpty())
+		{
+			PointHash.FindPointsInBall(MirroredPosition, SearchRadius, PointsInSphere);
+			SearchRadius += ThresholdRadius;
+
+			// forcibly break out if search radius gets bigger than the mesh bounds.
+			// this could potentially happen if mesh is highly non-symmetrical along mirror axis
+			if (SearchRadius >= MaxSearchRadius)
+			{
+				break;
+			}
+		}
+
+		// no mirrored points?
+		if (PointsInSphere.IsEmpty())
+		{
+			bSomeVerticesNotMirrored = true;
+			continue;
+		}
+
+		// find the closest single point
+		float ClosestDistSq = TNumericLimits<float>::Max();
+		int32 ClosestVertexID = INDEX_NONE;
+		for (int32 PointInSphereID : PointsInSphere)
+		{
+			const float DistSq = FVector::DistSquared(Vertices[PointInSphereID], MirroredPosition);
+			if (DistSq < ClosestDistSq)
+			{
+				ClosestDistSq = DistSq;
+				ClosestVertexID = PointInSphereID;
+			}
+		}
+		
+		// record the mirrored vertex ID for this vertex 
+		VertexMirrorMap.Add(TPair<int32, int32>(TargetVertexID, ClosestVertexID)); // (TO, FROM)
+	}
+
+	if (bSomeVerticesNotMirrored)
+	{
+		UE_LOG(LogMeshModelingTools, Log, TEXT("Mirror Skin Weights: some vertex weights were not mirrored because a vertex was not found close enough to the mirrored location."));
+	}
+
+	// iterate mirror map and copy weights from source to target vertices
+	FMultiBoneWeightEdits WeightEditsFromMirroring;
+	for (const TTuple<int32, int32>& VerticesToCopy : VertexMirrorMap)
+	{
+		int32 SourceVertexID = VerticesToCopy.Value;
+		int32 TargetVertexID = VerticesToCopy.Key;
+
+		// remove all weight on vertex
+		for (const FVertexBoneWeight& TargetBoneWeight : Weights.PreChangeWeights[TargetVertexID])
+		{
+			const float OldWeight = TargetBoneWeight.Weight;
+			constexpr float NewWeight = 0.f;
+			WeightEditsFromMirroring.MergeSingleEdit(TargetBoneWeight.BoneIndex, TargetVertexID, OldWeight, NewWeight);
+		}
+
+		// copy source weights, but with mirrored bones
+		for (const FVertexBoneWeight& SourceBoneWeight : Weights.PreChangeWeights[SourceVertexID])
+		{
+			int32 MirroredBoneIndex = MirrorBoneMap[SourceBoneWeight.BoneIndex];
+			const float OldWeight = Weights.GetWeightOfBoneOnVertex(MirroredBoneIndex, TargetVertexID, Weights.PreChangeWeights);
+			const float NewWeight = SourceBoneWeight.Weight;
+			WeightEditsFromMirroring.MergeSingleEdit(MirroredBoneIndex, TargetVertexID, OldWeight, NewWeight);
+		}
+	}
+
+	// apply the changes
+	const FText TransactionLabel = LOCTEXT("MirrorWeightChange", "Mirror skin weights.");
+	ApplyWeightEditsToMesh(TransactionLabel, WeightEditsFromMirroring);
+}
+
+void USkinWeightsPaintTool::FloodWeights(const float Weight, const EWeightEditOperation FloodMode)
+{
+	if (CurrentBone == NAME_None)
+	{
+		return;
+	}
 
+	// flood the weights on selected vertices
+	FMultiBoneWeightEdits WeightEditsFromFlood;
+	TArray<VertexIndex> VerticesToEdit;
+	GetVerticesToEdit(VerticesToEdit);
+	const TArray<float> VertexFalloffs = {};
+	if (FloodMode == EWeightEditOperation::Relax)
+	{
+		RelaxWeightOnVertices(VerticesToEdit, VertexFalloffs, Weight, WeightEditsFromFlood);
+	}
+	else
+	{
+		EditWeightOnVertices(FloodMode, VerticesToEdit, VertexFalloffs, Weight, WeightEditsFromFlood);	
+	}
+
+	// apply the changes
+	const FText TransactionLabel = LOCTEXT("FloodWeightChange", "Flood skin weights.");
+	ApplyWeightEditsToMesh(TransactionLabel, WeightEditsFromFlood);
+}
+
+void USkinWeightsPaintTool::PruneWeights(float Threshold)
+{
+	// set weights below the given threshold to zero
+	FMultiBoneWeightEdits WeightEditsFromPrune;
+	TArray<VertexIndex> VerticesToEdit;
+	GetVerticesToEdit(VerticesToEdit);
+	const bool bPruningAllBones = SelectedBoneIndices.IsEmpty();
+	for (const VertexIndex VertexID : VerticesToEdit)
+	{
+		const VertexWeights& VertexWeights = Weights.CurrentWeights[VertexID];
+		for (const FVertexBoneWeight& BoneWeight : VertexWeights)
+		{
+			// are we pruning weights only on selected bones?
+			if (!bPruningAllBones && !SelectedBoneIndices.Contains(BoneWeight.BoneIndex))
+			{
+				// not the bone we're pruning...
+				continue;
+			}
+			
+			if (BoneWeight.Weight > Threshold)
+			{
+				continue;
+			}
+
+			// note: removing ALL weight from a vertex disqualifies it as a candidate to receive weight during normalization
+			// so there is no need to remove the influence from the per-vertex influence array to "prune" it.
+			Weights.EditVertexWeightAndNormalize(BoneWeight.BoneIndex, VertexID,0.f,WeightEditsFromPrune);
+		}
+	}
+
+	// apply the changes
+	const FText TransactionLabel = LOCTEXT("PruneWeightValuesChange", "Prune skin weights.");
+	ApplyWeightEditsToMesh(TransactionLabel, WeightEditsFromPrune);
+}
+
+void USkinWeightsPaintTool::AverageWeights()
+{
+	// get vertices to edit weights on
+	TArray<VertexIndex> VerticesToEdit;
+	GetVerticesToEdit(VerticesToEdit);
+
+	// accumulate ALL the weights on the vertices
+	TMap<BoneIndex, float> AccumulatedWeightMap;
+	for (const VertexIndex VertexID : VerticesToEdit)
+	{
+		for (const FVertexBoneWeight& BoneWeight : Weights.PreChangeWeights[VertexID])
+		{
+			float& AccumulatedWeight = AccumulatedWeightMap.FindOrAdd(BoneWeight.BoneIndex);
+			AccumulatedWeight += BoneWeight.Weight;
+		}
+	}
+
+	// sort influences by total weight
+	AccumulatedWeightMap.ValueSort([](const float& A, const float& B)
+	{
+		return A > B;
+	});
+
+	// truncate to MaxInfluences
+	int32 Index = 0;
+	for (TMap<BoneIndex, float>::TIterator It(AccumulatedWeightMap); It; ++It)
+	{
+		if (Index >= MAX_TOTAL_INFLUENCES)
+		{
+			It.RemoveCurrent();
+		}
+		else
+		{
+			++Index;
+		}
+	}
+
+	// normalize remaining influences
+	float TotalWeight = 0.f;
+	for (const TTuple<BoneIndex, float>& AccumulatedWeight : AccumulatedWeightMap)
+	{
+		TotalWeight += AccumulatedWeight.Value;
+	}
+	for (TTuple<BoneIndex, float>& AccumulatedWeight : AccumulatedWeightMap)
+	{
+		AccumulatedWeight.Value /= TotalWeight;
+	}
+
+	// apply averaged weights to vertices
+	FMultiBoneWeightEdits WeightEditsFromAveraging;
+	for (const VertexIndex VertexID : VerticesToEdit)
+	{
+		// remove influences not a part of the average results
+		for (const FVertexBoneWeight& BoneWeight : Weights.PreChangeWeights[VertexID])
+		{
+			if (!AccumulatedWeightMap.Contains(BoneWeight.BoneIndex))
+			{
+				const float OldWeight = BoneWeight.Weight;
+				constexpr float NewWeight = 0.f;
+				WeightEditsFromAveraging.MergeSingleEdit(BoneWeight.BoneIndex, VertexID, OldWeight, NewWeight);
+			}
+		}
+
+		// add influences from the averaging results
+		for (const TTuple<BoneIndex, float>& AccumulatedWeight : AccumulatedWeightMap)
+		{
+			const BoneIndex IndexOfBone = AccumulatedWeight.Key;
+			const float OldWeight = Weights.GetWeightOfBoneOnVertex(IndexOfBone, VertexID, Weights.PreChangeWeights);
+			const float NewWeight = AccumulatedWeight.Value;
+			WeightEditsFromAveraging.MergeSingleEdit(IndexOfBone, VertexID, OldWeight, NewWeight);
+		}
+	}
+
+	// apply the changes
+	const FText TransactionLabel = LOCTEXT("AverageWeightValuesChange", "Average skin weights.");
+	ApplyWeightEditsToMesh(TransactionLabel, WeightEditsFromAveraging);
+}
+
+void USkinWeightsPaintTool::NormalizeWeights()
+{
+	// re-set a weight on each vertex to force normalization
+	FMultiBoneWeightEdits WeightEditsFromNormalization;
+	TArray<VertexIndex> SelectedVertices;
+	GetVerticesToEdit(SelectedVertices);
+	for (const VertexIndex VertexID : SelectedVertices)
+	{
+		const VertexWeights& VertexWeights = Weights.CurrentWeights[VertexID];
+		for (const FVertexBoneWeight& BoneWeight : VertexWeights)
+		{
+			// set weight to current value, just to force re-normalization
+			Weights.EditVertexWeightAndNormalize(BoneWeight.BoneIndex, VertexID, BoneWeight.Weight,WeightEditsFromNormalization);
+		}
+	}
+
+	// apply the changes
+	const FText TransactionLabel = LOCTEXT("NormalizeWeightValuesChange", "Normalize skin weights.");
+	ApplyWeightEditsToMesh(TransactionLabel, WeightEditsFromNormalization);
+}
+
+void USkinWeightsPaintTool::HandleSkeletalMeshModified(const TArray<FName>& InBoneNames, const ESkeletalMeshNotifyType InNotifyType)
+{
 	switch (InNotifyType)
 	{
 	case ESkeletalMeshNotifyType::BonesAdded:
@@ -1189,16 +1623,28 @@ void USkinWeightsPaintTool::HandleSkeletalMeshModified(const TArray<FName>& InBo
 	case ESkeletalMeshNotifyType::BonesMoved:
 		break;
 	case ESkeletalMeshNotifyType::BonesSelected:
-		if (BoneName != NAME_None)
 		{
-			PendingCurrentBone = BoneName;
+			// store selected bones
+			SelectedBoneNames = InBoneNames;
+			PendingCurrentBone = InBoneNames.IsEmpty() ? NAME_None : InBoneNames[0];
+
+			// update selected bone indices from names
+			SelectedBoneIndices.Reset();
+			for (const FName SelectedBoneName : SelectedBoneNames)
+			{
+				SelectedBoneIndices.Add(GetBoneIndexFromName(SelectedBoneName));
+			}
 		}
 		break;
+	default:
+		checkNoEntry();
 	}
 }
 
-void USkinWeightsPaintTool::OnToolPropertiesModified(UObject* ModifiedObject, FProperty* ModifiedProperty)
-{	
+void USkinWeightsPaintTool::OnPropertyModified(UObject* ModifiedObject, FProperty* ModifiedProperty)
+{
+	Super::OnPropertyModified(ModifiedObject, ModifiedProperty);
+	
 	// invalidate vertex color cache when weight color properties are modified
 	const bool bColorModeModified = ModifiedProperty->GetNameCPP() == GET_MEMBER_NAME_STRING_CHECKED(USkinWeightsPaintToolProperties, ColorMode);
 	const bool bColorRampModified = ModifiedProperty->GetNameCPP() == GET_MEMBER_NAME_STRING_CHECKED(USkinWeightsPaintToolProperties, ColorRamp);
diff --git a/Engine/Plugins/Experimental/MeshModelingToolsetExp/Source/MeshModelingToolsExp/Public/MeshModelingToolsExp.h b/Engine/Plugins/Experimental/MeshModelingToolsetExp/Source/MeshModelingToolsExp/Public/MeshModelingToolsExp.h
--- a/Engine/Plugins/Experimental/MeshModelingToolsetExp/Source/MeshModelingToolsExp/Public/MeshModelingToolsExp.h
+++ b/Engine/Plugins/Experimental/MeshModelingToolsetExp/Source/MeshModelingToolsExp/Public/MeshModelingToolsExp.h
@@ -5,6 +5,8 @@
 #include "CoreMinimal.h"
 #include "Modules/ModuleManager.h"
 
+DECLARE_LOG_CATEGORY_EXTERN(LogMeshModelingTools, Log, All);
+
 class FMeshModelingToolsExpModule : public IModuleInterface
 {
 public:
diff --git a/Engine/Plugins/Experimental/MeshModelingToolsetExp/Source/MeshModelingToolsExp/Public/SkinWeightsPaintTool.h b/Engine/Plugins/Experimental/MeshModelingToolsetExp/Source/MeshModelingToolsExp/Public/SkinWeightsPaintTool.h
--- a/Engine/Plugins/Experimental/MeshModelingToolsetExp/Source/MeshModelingToolsExp/Public/SkinWeightsPaintTool.h
+++ b/Engine/Plugins/Experimental/MeshModelingToolsetExp/Source/MeshModelingToolsExp/Public/SkinWeightsPaintTool.h
@@ -91,9 +91,9 @@ namespace SkinPaintTool
 			return FString(TEXT("Paint Skin Weights"));
 		}
 
-		void Apply(UObject* Object) override;
+		virtual void Apply(UObject* Object) override;
 
-		void Revert(UObject* Object) override;
+		virtual void Revert(UObject* Object) override;
 
 		void AddBoneWeightEdit(const FSingleBoneWeightEdits& BoneWeightEdit);
 
@@ -112,13 +112,13 @@ namespace SkinPaintTool
 		// applies an edit to a single vertex weight on a single bone, then normalizes the remaining weights while
 		// keeping the edited weight intact (ie, adapts OTHER influences to achieve normalization)
 		void EditVertexWeightAndNormalize(
-			const FName BoneToHoldConstant,
+			const int32 BoneIndex,
 			const int32 VertexId,
-			const float NewWeightValue,
+			float NewWeightValue,
 			FMultiBoneWeightEdits& WeightEdits);
 
 		void ApplyCurrentWeightsToMeshDescription(FMeshDescription* EditedMesh);
-
+		
 		static float GetWeightOfBoneOnVertex(
 			const int32 BoneIndex,
 			const int32 VertexID,
@@ -130,7 +130,7 @@ namespace SkinPaintTool
 			const float Weight,
 			TArray<VertexWeights>& InOutVertexData);
 
-		void ResetAfterStroke();
+		void ResetAfterChange();
 
 		float SetCurrentFalloffAndGetMaxFalloffThisStroke(int32 VertexID, float CurrentStrength);
 
@@ -139,8 +139,8 @@ namespace SkinPaintTool
 		// double-buffer of the entire weight matrix (stored sparsely for fast deformation)
 		// "Pre" is state of weights at stroke start
 		// "Current" is state of weights during stroke
-		// When stroke is over, Pre weights are synchronized with current.
-		TArray<VertexWeights> PreStrokeWeights;
+		// When stroke is over, PreChangeWeights are synchronized with CurrentWeights
+		TArray<VertexWeights> PreChangeWeights;
 		TArray<VertexWeights> CurrentWeights;
 
 		// record the current maximum amount of falloff applied to each vertex during the current stroke
@@ -151,19 +151,14 @@ namespace SkinPaintTool
 		// update deformation when vertex weights are modified
 		FSkinToolDeformer Deformer;
 	};
-
 }
 
-/**
- *
- */
-UCLASS()
-class MESHMODELINGTOOLSEXP_API USkinWeightsPaintToolBuilder : public UMeshSurfacePointMeshEditingToolBuilder
+// weight edit mode
+UENUM()
+enum class EWeightEditMode : uint8
 {
-	GENERATED_BODY()
-
-public:
-	UMeshSurfacePointTool* CreateNewTool(const FToolBuilderState& SceneState) const override;
+	Brush,
+	Selection,
 };
 
 // weight color mode
@@ -182,9 +177,9 @@ enum class EWeightBrushFalloffMode : uint8
 	Volume,
 };
 
-// brush behavior mode
+// operation type when editing weights
 UENUM()
-enum class EBrushBehaviorMode : uint8
+enum class EWeightEditOperation : uint8
 {
 	Add,
 	Replace,
@@ -192,41 +187,70 @@ enum class EBrushBehaviorMode : uint8
 	Relax
 };
 
+// mirror direction mode
+UENUM()
+enum class EMirrorDirection : uint8
+{
+	PositiveToNegative,
+	NegativeToPositive,
+};
 
 UCLASS()
-class MESHMODELINGTOOLSEXP_API USkinWeightsPaintToolProperties : public UInteractiveToolPropertySet
+class MESHMODELINGTOOLSEXP_API USkinWeightsPaintToolBuilder : public UMeshSurfacePointMeshEditingToolBuilder
 {
 	GENERATED_BODY()
 
-	USkinWeightsPaintToolProperties(const FObjectInitializer& ObjectInitializer);
-	
 public:
+	UMeshSurfacePointTool* CreateNewTool(const FToolBuilderState& SceneState) const override;
+};
 
-	UPROPERTY(EditAnywhere, Category = Brush)
-	EBrushBehaviorMode BrushMode;
+// Container for properties displayed in Details panel while using USkinWeightsPaintTool
+UCLASS()
+class MESHMODELINGTOOLSEXP_API USkinWeightsPaintToolProperties : public UBrushBaseProperties
+{
+	GENERATED_BODY()
+
+	USkinWeightsPaintToolProperties();
 	
-	UPROPERTY(EditAnywhere, Category = Brush)
+public:
+
+	// brush vs selection modes
+	UPROPERTY()
+	EWeightEditMode EditingMode;
+
+	// custom brush modes and falloff types
+	UPROPERTY()
+	EWeightEditOperation BrushMode;
+	UPROPERTY()
 	EWeightBrushFalloffMode FalloffMode;
 
-	UPROPERTY(EditAnywhere, Category = Brush)
+	// weight color properties
+	UPROPERTY(EditAnywhere, Category = WeightColors)
 	EWeightColorMode ColorMode;
-	
 	UPROPERTY(EditAnywhere, Category = WeightColors)
 	TArray<FLinearColor> ColorRamp;
-
 	UPROPERTY(EditAnywhere, Category = WeightColors)
 	FLinearColor MinColor;
-
 	UPROPERTY(EditAnywhere, Category = WeightColors)
 	FLinearColor MaxColor;
-	
-	TObjectPtr<USkeletalMesh> SkeletalMesh;
 	bool bColorModeChanged = false;
+
+	// weight editing arguments
+	UPROPERTY()
+	TEnumAsByte<EAxis::Type> MirrorAxis = EAxis::X;
+	UPROPERTY()
+	EMirrorDirection MirrorDirection = EMirrorDirection::PositiveToNegative;
+	UPROPERTY()
+	EWeightEditOperation FloodMode;
+	UPROPERTY()
+	float FloodValue = 1.f;
+	UPROPERTY()
+	float PruneValue = 0.01;
+
+	TObjectPtr<USkinWeightsPaintTool> WeightTool;
 };
 
-/**
- *
- */
+// An interactive tool for painting and editing skin weights.
 UCLASS()
 class MESHMODELINGTOOLSEXP_API USkinWeightsPaintTool : public UDynamicMeshBrushTool, public ISkeletalMeshEditionInterface
 {
@@ -250,6 +274,13 @@ public:
 
 	void ExternalUpdateWeights(const int32 BoneIndex, const TMap<int32, float>& IndexValues);
 
+	// weight editing operations (selection based)
+	void MirrorWeights(EAxis::Type Axis, EMirrorDirection Direction);
+	void FloodWeights(const float Weight, const EWeightEditOperation FloodMode);
+	void PruneWeights(const float Threshold);
+	void AverageWeights();
+	void NormalizeWeights();
+
 protected:
 
 	virtual void ApplyStamp(const FBrushStampData& Stamp);
@@ -257,20 +288,12 @@ protected:
 	void OnTick(float DeltaTime) override;
 
 	// stamp
-	double CalculateBrushFalloff(double Distance) const;
+	float CalculateBrushFalloff(float Distance) const;
 	void CalculateVertexROI(
 		const FBrushStampData& Stamp,
-		TArray<int>& VertexROI,
-		TArray<float>& VertexSqDistances) const;
-	void EditWeightOfVerticesInStamp(
-		EBrushBehaviorMode EditMode,
-		const TArray<int32>& VerticesInStamp,
-		const TArray<float>& VertexSqDistances,
-		SkinPaintTool::FMultiBoneWeightEdits& AllBoneWeightEditsFromStamp);
-	void RelaxWeightOnVertices(
-		TArray<int32> VerticesInStamp,
-		TArray<float> VertexSqDistances,
-		SkinPaintTool::FMultiBoneWeightEdits& AllBoneWeightEditsFromStamp);
+		TArray<VertexIndex>& VertexIDs,
+		TArray<float>& VertexFalloffs);
+	float CalculateBrushStrengthToUse(EWeightEditOperation EditMode) const;
 	bool bInvertStroke = false;
 	bool bSmoothStroke = false;
 	FBrushStampData StartStamp;
@@ -279,6 +302,21 @@ protected:
 	int32 TriangleUnderStamp;
 	FVector StampLocalPos;
 
+	// modify vertex weights according to the specified operation,
+	// generating bone weight edits to be stored in a transaction
+	void EditWeightOnVertices(
+		EWeightEditOperation EditOperation,
+		const TArray<int32>& VerticesToEdit,
+		const TArray<float>& VertexFalloffs,
+		const float UseStrength,
+		SkinPaintTool::FMultiBoneWeightEdits& InOutWeightEdits);
+	// same as EditWeightOfVertices() but specific to relaxation (topology aware operation)
+	void RelaxWeightOnVertices(
+		TArray<int32> VerticesToEdit,
+		TArray<float> VertexFalloffs,
+		const float UseStrength,
+		SkinPaintTool::FMultiBoneWeightEdits& InOutWeightEdits);
+
 	// used to accelerate mesh queries
 	UE::Geometry::TDynamicVerticesOctree3<FDynamicMesh3> VerticesOctree;
 	UE::Geometry::FDynamicMeshOctree3 TrianglesOctree;
@@ -287,8 +325,8 @@ protected:
 
 	// tool properties
 	UPROPERTY()
-	TObjectPtr<USkinWeightsPaintToolProperties> ToolProps;
-	void OnToolPropertiesModified(UObject* ModifiedObject, FProperty* ModifiedProperty);
+	TObjectPtr<USkinWeightsPaintToolProperties> WeightToolProperties;
+	virtual void OnPropertyModified(UObject* ModifiedObject, FProperty* ModifiedProperty) override;
 	
 	// the currently edited mesh description
 	TUniquePtr<FMeshDescription> EditedMesh;
@@ -305,12 +343,22 @@ protected:
 	void BeginChange();
 	TUniquePtr<SkinPaintTool::FMeshSkinWeightsChange> EndChange();
 	TUniquePtr<SkinPaintTool::FMeshSkinWeightsChange> ActiveChange;
+	// used to apply manual weight modifications (not from the brush)
+	void ApplyWeightEditsToMesh(const FText& TransactionLabel, const SkinPaintTool::FMultiBoneWeightEdits& WeightEdits);
 
 	// which bone are we currently painting?
 	void UpdateCurrentBone(const FName &BoneName);
 	FName CurrentBone = NAME_None;
 	TOptional<FName> PendingCurrentBone;
+	TArray<FName> SelectedBoneNames;
+	TArray<BoneIndex> SelectedBoneIndices;
+
+	// determines the set of vertices to operate on, using selection as the priority
+	void GetVerticesToEdit(TArray<VertexIndex>& OutVertexIndices) const;
+	
+	BoneIndex GetBoneIndexFromName(const FName BoneName) const;
 
+	// ISkeletalMeshEditionInterface
 	virtual void HandleSkeletalMeshModified(const TArray<FName>& InBoneNames, const ESkeletalMeshNotifyType InNotifyType) override;
 
 	friend SkinPaintTool::FSkinToolDeformer;
diff --git a/Engine/Plugins/Experimental/MeshModelingToolsetExp/Source/ModelingEditorUI/Private/ModelingEditorUIModule.cpp b/Engine/Plugins/Experimental/MeshModelingToolsetExp/Source/ModelingEditorUI/Private/ModelingEditorUIModule.cpp
--- a/Engine/Plugins/Experimental/MeshModelingToolsetExp/Source/ModelingEditorUI/Private/ModelingEditorUIModule.cpp
+++ b/Engine/Plugins/Experimental/MeshModelingToolsetExp/Source/ModelingEditorUI/Private/ModelingEditorUIModule.cpp
@@ -2,17 +2,25 @@
 
 #include "ModelingEditorUIModule.h"
 
+#include "SkinWeightDetailCustomization.h"
+#include "SkinWeightsPaintTool.h"
+#include "PropertyEditorModule.h"
+
 #define LOCTEXT_NAMESPACE "FModelingEditorUIModule"
 
 void FModelingEditorUIModule::StartupModule()
 {
-	// This code will execute after your module is loaded into memory; the exact timing is specified in the .uplugin file per-module
+	// register detail customizations
+	FPropertyEditorModule& PropertyEditorModule = FModuleManager::LoadModuleChecked<FPropertyEditorModule>("PropertyEditor");
+	PropertyEditorModule.RegisterCustomClassLayout(USkinWeightsPaintToolProperties::StaticClass()->GetFName(), FOnGetDetailCustomizationInstance::CreateStatic(&FSkinWeightDetailCustomization::MakeInstance));
+	
 }
 
 void FModelingEditorUIModule::ShutdownModule()
 {
-	// This function may be called during shutdown to clean up your module.  For modules that support dynamic reloading,
-	// we call this function before unloading the module.
+	// unregister detail customizations
+	FPropertyEditorModule& PropertyEditorModule = FModuleManager::LoadModuleChecked<FPropertyEditorModule>("PropertyEditor");
+	PropertyEditorModule.UnregisterCustomClassLayout(USkinWeightsPaintToolProperties::StaticClass()->GetFName());
 }
 
 #undef LOCTEXT_NAMESPACE
diff --git a/Engine/Plugins/Experimental/MeshModelingToolsetExp/Source/ModelingEditorUI/Private/SkinWeightDetailCustomization.cpp b/Engine/Plugins/Experimental/MeshModelingToolsetExp/Source/ModelingEditorUI/Private/SkinWeightDetailCustomization.cpp
--- a/Engine/Plugins/Experimental/MeshModelingToolsetExp/Source/ModelingEditorUI/Private/SkinWeightDetailCustomization.cpp
+++ b/Engine/Plugins/Experimental/MeshModelingToolsetExp/Source/ModelingEditorUI/Private/SkinWeightDetailCustomization.cpp
@@ -8,6 +8,9 @@
 #include "Widgets/Input/SSegmentedControl.h"
 #include "SkinWeightsPaintTool.h"
 #include "SSkinWeightProfileImportOptions.h"
+#include "UObject/UnrealTypePrivate.h"
+#include "Widgets/Input/SSpinBox.h"
+#include "Widgets/Input/SButton.h"
 
 #define LOCTEXT_NAMESPACE "SkinWeightToolSettingsEditor"
 
@@ -19,128 +22,597 @@ void FSkinWeightDetailCustomization::CustomizeDetails(IDetailLayoutBuilder& Deta
 	// should be impossible to get multiple settings objects for a single tool
 	ensure(DetailObjects.Num()==1);
 	SkinToolSettings = Cast<USkinWeightsPaintToolProperties>(DetailObjects[0]);
+
+	// layout constants
+	constexpr float WeightSliderWidths = 150.0f;
+	constexpr float WeightEditingLabelsPercent = 0.40f;
+	constexpr float WeightEditVerticalPadding = 4.0f;
+
+	// custom display of falloff mode as segmented toggle buttons
+	IDetailCategoryBuilder& EditModeCategory = DetailBuilder.EditCategory("Weight Editing Mode", FText::GetEmpty(), ECategoryPriority::Important);
+
+	// add segmented control toggle for brush falloff modes ("Surface" or "Volume")
+	EditModeCategory.AddCustomRow(LOCTEXT("EditModeCategory", "Weight Editing Mode"), false)
+	.WholeRowContent()
+	[
+		SNew(SBox)
+		.Padding(2.0f)
+		[
+			SNew(SSegmentedControl<EWeightEditMode>)
+			.ToolTipText(LOCTEXT("EditingModeTooltip",
+					"Brush: the interactive viewport brush interface for editing weights.\n"
+					"Selection: the 1-off editing operations for weights on selected elements (bones and vertices).\n"))
+			.Value_Lambda([this]()
+			{
+				return SkinToolSettings->EditingMode;
+			})
+			.OnValueChanged_Lambda([this](EWeightEditMode Mode)
+			{
+				SkinToolSettings->EditingMode = Mode;
+			})
+			+SSegmentedControl<EWeightEditMode>::Slot(EWeightEditMode::Brush)
+			.Text(LOCTEXT("EditMode", "Brush"))
+			+ SSegmentedControl<EWeightEditMode>::Slot(EWeightEditMode::Selection)
+			.Text(LOCTEXT("EditMode", "Selection"))
+		]
+	];
 	
 	// custom display of falloff mode as segmented toggle buttons
 	IDetailCategoryBuilder& BrushCategory = DetailBuilder.EditCategory("Brush", FText::GetEmpty(), ECategoryPriority::Important);
 
 	// add segmented control toggle for brush behavior modes ("Add", "Replace", etc..)
 	BrushCategory.AddCustomRow(LOCTEXT("BrushModeCategory", "Brush Mode"), false)
-		.NameContent()
+	.WholeRowContent()
+	[
+		SNew(SBox)
+		.Padding(2.0f)
+		[
+			SNew(SSegmentedControl<EWeightEditOperation>)
+			.ToolTipText(LOCTEXT("FloodOperationTooltip",
+					"Add: applies the current weight PLUS the strength value to the new weight.\n"
+					"Remove: applies the current weight MINUS the strength value to the new weight.\n"
+					"Multiply: applies the current weight MULTIPLIED by the strength value to the new weight.\n"
+					"Relax: applies the average of the connected (by edge) vertex weights to the new vertex weight, blended by the strength.\n"))
+			.IsEnabled_Lambda([this]()
+			{
+				return SkinToolSettings->EditingMode == EWeightEditMode::Brush;
+			})
+			.Value_Lambda([this]()
+			{
+				return SkinToolSettings->BrushMode;
+			})
+			.OnValueChanged_Lambda([this](EWeightEditOperation Mode)
+			{
+				SkinToolSettings->BrushMode = Mode;
+			})
+			+SSegmentedControl<EWeightEditOperation>::Slot(EWeightEditOperation::Add)
+			.Text(LOCTEXT("AddMode", "Add"))
+			+ SSegmentedControl<EWeightEditOperation>::Slot(EWeightEditOperation::Replace)
+			.Text(LOCTEXT("ReplaceMode", "Replace"))
+			+ SSegmentedControl<EWeightEditOperation>::Slot(EWeightEditOperation::Multiply)
+			.Text(LOCTEXT("MultiplyMode", "Multiply"))
+			+ SSegmentedControl<EWeightEditOperation>::Slot(EWeightEditOperation::Relax)
+			.Text(LOCTEXT("RelaxMode", "Relax"))
+		]
+	];
+
+	// add segmented control toggle for brush falloff modes ("Surface" or "Volume")
+	BrushCategory.AddCustomRow(LOCTEXT("BrushFalloffModeCategory", "Brush Falloff Mode"), false)
+	.WholeRowContent()
+	[
+		SNew(SBox)
+		.Padding(2.0f)
+		[
+			SNew(SSegmentedControl<EWeightBrushFalloffMode>)
+			.ToolTipText(LOCTEXT("BrushFalloffModeTooltip",
+					"Surface: falloff is based on the distance along the surface from the brush center to nearby connected vertices.\n"
+					"Volume: falloff is based on the straight-line distance from the brush center to surrounding vertices.\n"))
+			.IsEnabled_Lambda([this]()
+			{
+				return SkinToolSettings->EditingMode == EWeightEditMode::Brush;
+			})
+			.Value_Lambda([this]()
+			{
+				return SkinToolSettings->FalloffMode;
+			})
+			.OnValueChanged_Lambda([this](EWeightBrushFalloffMode Mode)
+			{
+				SkinToolSettings->FalloffMode = Mode;
+				SkinToolSettings->bColorModeChanged = true;
+			})
+			+SSegmentedControl<EWeightBrushFalloffMode>::Slot(EWeightBrushFalloffMode::Surface)
+			.Text(LOCTEXT("SurfaceMode", "Surface"))
+			+ SSegmentedControl<EWeightBrushFalloffMode>::Slot(EWeightBrushFalloffMode::Volume)
+			.Text(LOCTEXT("VolumeMode", "Volume"))
+		]
+	];
+
+	BrushCategory.AddCustomRow(LOCTEXT("BrushSizeCategory", "Brush Radius"), false)
+	.NameContent()
+	[
+		SNew(STextBlock)
+		.Text(LOCTEXT("BrushRadiusLabel", "Radius"))
+		.Font(FAppStyle::GetFontStyle(TEXT("PropertyWindow.NormalFont")))
+		.ToolTipText(LOCTEXT("BrushRadiusTooltip", "The radius of the brush in scene units."))
+	]
+	.ValueContent()
+	[
+		SNew(SSpinBox<float>)
+		.MinValue(0.01f)
+		.MaxSliderValue(20.f)
+		.Value(10.0f)
+		.SupportDynamicSliderMaxValue(true)
+		.IsEnabled_Lambda([this]()
+		{
+			return SkinToolSettings->EditingMode == EWeightEditMode::Brush;
+		})
+		.Value_Lambda([this]()
+		{
+			return SkinToolSettings->BrushRadius;
+		})
+		.OnValueChanged_Lambda([this](float NewValue)
+		{
+			SkinToolSettings->BrushRadius = NewValue;
+			FPropertyChangedEvent PropertyChangedEvent(UBrushBaseProperties::StaticClass()->FindPropertyByName(GET_MEMBER_NAME_CHECKED(UBrushBaseProperties, BrushRadius)));
+			SkinToolSettings->PostEditChangeProperty(PropertyChangedEvent);
+		})
+	];
+
+	BrushCategory.AddCustomRow(LOCTEXT("BrushStrengthCategory", "Brush Strength"), false)
+	.NameContent()
+	[
+		SNew(STextBlock)
+		.Text(LOCTEXT("BrushStrengthLabel", "Strength"))
+		.Font(FAppStyle::GetFontStyle(TEXT("PropertyWindow.NormalFont")))
+		.ToolTipText(LOCTEXT("BrushStrengthTooltip", "The strength of the effect on the weights. Exact effect depends on the Brush mode."))
+	]
+	.ValueContent()
+	[
+		SNew(SSpinBox<float>)
+		.MinValue(0.f)
+		.MaxValue(1.f)
+		.IsEnabled_Lambda([this]()
+		{
+			return SkinToolSettings->EditingMode == EWeightEditMode::Brush;
+		})
+		.Value_Lambda([this]()
+		{
+			return SkinToolSettings->BrushStrength;
+		})
+		.OnValueChanged_Lambda([this](float NewValue)
+		{
+			SkinToolSettings->BrushStrength = NewValue;
+			FPropertyChangedEvent PropertyChangedEvent(UBrushBaseProperties::StaticClass()->FindPropertyByName(GET_MEMBER_NAME_CHECKED(UBrushBaseProperties, BrushStrength)));
+			SkinToolSettings->PostEditChangeProperty(PropertyChangedEvent);
+		})
+	];
+
+	BrushCategory.AddCustomRow(LOCTEXT("BrushFalloffCategory", "Brush Falloff"), false)
+	.NameContent()
+	[
+		SNew(STextBlock)
+		.Text(LOCTEXT("BrushFalloffLabel", "Falloff"))
+		.Font(FAppStyle::GetFontStyle(TEXT("PropertyWindow.NormalFont")))
+		.ToolTipText(LOCTEXT("BrushFalloffTooltip", "At 0, the brush has no falloff. At 1 it has exponential falloff."))
+	]
+	.ValueContent()
+	[
+		SNew(SSpinBox<float>)
+		.MinValue(0.f)
+		.MaxValue(1.f)
+		.IsEnabled_Lambda([this]()
+		{
+			return SkinToolSettings->EditingMode == EWeightEditMode::Brush;
+		})
+		.Value_Lambda([this]()
+		{
+			return SkinToolSettings->BrushFalloffAmount;
+		})
+		.OnValueChanged_Lambda([this](float NewValue)
+		{
+			SkinToolSettings->BrushFalloffAmount = NewValue;
+			FPropertyChangedEvent PropertyChangedEvent(UBrushBaseProperties::StaticClass()->FindPropertyByName(GET_MEMBER_NAME_CHECKED(UBrushBaseProperties, BrushFalloffAmount)));
+			SkinToolSettings->PostEditChangeProperty(PropertyChangedEvent);
+		})
+	];
+
+	// custom display of weight editing tools
+	IDetailCategoryBuilder& EditWeightsCategory = DetailBuilder.EditCategory("EditWeights", FText::GetEmpty(), ECategoryPriority::Important);
+	EditWeightsCategory.InitiallyCollapsed(true);
+
+	// AVERAGE/RELAX/NORMALIZE WEIGHTS category
+	EditWeightsCategory.AddCustomRow(LOCTEXT("NormalizeWeightsRow", "Normalize"), false)
+	.WholeRowContent()
+	[
+		SNew(SHorizontalBox)
+
+		+SHorizontalBox::Slot()
+		.Padding(2.f, WeightEditVerticalPadding)
 		[
-			SNew(STextBlock)
-			.Text(LOCTEXT("BrushModeLabel", "Brush Mode"))
-			.Font(FAppStyle::GetFontStyle(TEXT("PropertyWindow.NormalFont")))
-			.ToolTipText(LOCTEXT("BrushModeTooltip", "Determines how the weights are affected by the brush."))
+			SNew(SButton)
+			.HAlign(HAlign_Center)
+			.VAlign(VAlign_Center)
+			.Text(LOCTEXT("AverageWeightsButtonLabel", "Average"))
+			.ToolTipText(LOCTEXT("AverageWeightsTooltip",
+					"Takes the average of vertex weights and applies it to the selected vertices.\n"
+					"This command operates on the selected bone(s) and selected vertices.\n "
+					"If no bones are selected, ALL bone weights are considered.\n "
+					"If no vertices are selected, ALL vertices are considered."))
+			.IsEnabled_Lambda([this]()
+			{
+				return SkinToolSettings->EditingMode == EWeightEditMode::Selection;
+			})
+			.OnClicked_Lambda([this]()
+			{
+				SkinToolSettings->WeightTool->AverageWeights();
+				return FReply::Handled();
+			})
 		]
-		.ValueContent()
+
+		+SHorizontalBox::Slot()
+		.Padding(2.f, WeightEditVerticalPadding)
+		[
+			SNew(SButton)
+			.HAlign(HAlign_Center)
+			.VAlign(VAlign_Center)
+			.Text(LOCTEXT("NormalizeWeightsButtonLabel", "Normalize"))
+			.ToolTipText(LOCTEXT("NormalizeWeightsTooltip",
+					"Forces the weights on the selected vertices to sum to 1.\n"
+					"This command operates on the selected vertices.\n "
+					"If no vertices are selected, ALL vertices are considered."))
+			.IsEnabled_Lambda([this]()
+			{
+				return SkinToolSettings->EditingMode == EWeightEditMode::Selection;
+			})
+			.OnClicked_Lambda([this]()
+			{
+				SkinToolSettings->WeightTool->NormalizeWeights();
+				return FReply::Handled();
+			})
+		]
+	];
+
+	// MIRROR WEIGHTS category
+	EditWeightsCategory.AddCustomRow(LOCTEXT("MirrorWeightsRow", "Mirror"), false)
+	.WholeRowContent()
+	[
+		SNew(SVerticalBox)
+		
+		+ SVerticalBox::Slot()
+		.Padding(0.f, WeightEditVerticalPadding)
 		[
 			SNew(SHorizontalBox)
-			+ SHorizontalBox::Slot()
+			
+			+SHorizontalBox::Slot()
+			.VAlign(VAlign_Center)
+			.FillWidth(WeightEditingLabelsPercent)
 			[
-				SNew(SBox)
-				.Padding(2.0f)
-				.HAlign(HAlign_Center)
+				SNew(STextBlock)
+				.Text(LOCTEXT("MirrorPlaneLabel", "Mirror Plane"))
+				.Font(FAppStyle::GetFontStyle(TEXT("PropertyWindow.NormalFont")))
+				.ToolTipText(LOCTEXT("MirrorPlaneTooltip", "The plane to copy weights across."))
+			]
+
+			+SHorizontalBox::Slot()
+			.FillWidth(1.f)
+			[
+				SNew(SHorizontalBox)
+				+SHorizontalBox::Slot()
 				[
-					SNew(SSegmentedControl<EBrushBehaviorMode>)
+					SNew(SSegmentedControl<EAxis::Type>)
+					.ToolTipText(LOCTEXT("MirrorAxisTooltip",
+						"X: copies weights across the YZ plane.\n"
+						"Y: copies weights across the XZ plane.\n"
+						"Z: copies weights across the XY plane."))
+					.IsEnabled_Lambda([this]()
+					{
+						return SkinToolSettings->EditingMode == EWeightEditMode::Selection;
+					})
 					.Value_Lambda([this]()
 					{
-						return SkinToolSettings->BrushMode;
+						return SkinToolSettings->MirrorAxis;
 					})
-					.OnValueChanged_Lambda([this](EBrushBehaviorMode Mode)
+					.OnValueChanged_Lambda([this](EAxis::Type Mode)
 					{
-						SkinToolSettings->BrushMode = Mode;
+						SkinToolSettings->MirrorAxis = Mode;
 					})
-					+SSegmentedControl<EBrushBehaviorMode>::Slot(EBrushBehaviorMode::Add)
-					.Text(LOCTEXT("AddMode", "Add"))
-					+ SSegmentedControl<EBrushBehaviorMode>::Slot(EBrushBehaviorMode::Replace)
-					.Text(LOCTEXT("ReplaceMode", "Replace"))
-					+ SSegmentedControl<EBrushBehaviorMode>::Slot(EBrushBehaviorMode::Multiply)
-					.Text(LOCTEXT("MultiplyMode", "Multiply"))
-					+ SSegmentedControl<EBrushBehaviorMode>::Slot(EBrushBehaviorMode::Relax)
-					.Text(LOCTEXT("RelaxMode", "Relax"))
+					+ SSegmentedControl<EAxis::Type>::Slot(EAxis::X)
+					.Text(LOCTEXT("MirrorXLabel", "X"))
+					+ SSegmentedControl<EAxis::Type>::Slot(EAxis::Y)
+					.Text(LOCTEXT("MirrorYLabel", "Y"))
+					+ SSegmentedControl<EAxis::Type>::Slot(EAxis::Z)
+					.Text(LOCTEXT("MirrorZLabel", "Z"))
 				]
-			]
-		];
-	const TSharedRef<IPropertyHandle> BrushPropHandle = DetailBuilder.GetProperty(GET_MEMBER_NAME_CHECKED(USkinWeightsPaintToolProperties, BrushMode));
-	DetailBuilder.HideProperty(BrushPropHandle);
-
-	// add segmented control toggle for brush falloff modes ("Surface" or "Volume")
-	BrushCategory.AddCustomRow(LOCTEXT("BrushFalloffCategory", "Brush Falloff"), false)
-		.NameContent()
-		[
-			SNew(STextBlock)
-			.Text(LOCTEXT("BrushFalloffLabel", "Falloff Mode"))
-			.Font(FAppStyle::GetFontStyle(TEXT("PropertyWindow.NormalFont")))
-			.ToolTipText(LOCTEXT("BrushFalloffTooltip", "Falloff brush influence along the surface of the mesh, or using straight-line distance (volume)."))
-		]
-		.ValueContent()
-		[
-			SNew(SHorizontalBox)
-			+ SHorizontalBox::Slot()
-			[
-				SNew(SBox)
-				.Padding(2.0f)
-				.HAlign(HAlign_Center)
+					
+				+SHorizontalBox::Slot()
 				[
-					SNew(SSegmentedControl<EWeightBrushFalloffMode>)
+					SNew(SSegmentedControl<EMirrorDirection>)
+					.ToolTipText(LOCTEXT("MirrorDirectionTooltip", "The direction that determines what side of the plane to copy weights from."))
+					.IsEnabled_Lambda([this]()
+					{
+						return SkinToolSettings->EditingMode == EWeightEditMode::Selection;
+					})
 					.Value_Lambda([this]()
 					{
-						return SkinToolSettings->FalloffMode;
+						return SkinToolSettings->MirrorDirection;
 					})
-					.OnValueChanged_Lambda([this](EWeightBrushFalloffMode Mode)
+					.OnValueChanged_Lambda([this](EMirrorDirection Mode)
 					{
-						SkinToolSettings->FalloffMode = Mode;
-						SkinToolSettings->bColorModeChanged = true;
+						SkinToolSettings->MirrorDirection = Mode;
 					})
-					+SSegmentedControl<EWeightBrushFalloffMode>::Slot(EWeightBrushFalloffMode::Surface)
-					.Text(LOCTEXT("SurfaceMode", "Surface"))
-					+ SSegmentedControl<EWeightBrushFalloffMode>::Slot(EWeightBrushFalloffMode::Volume)
-					.Text(LOCTEXT("VolumeMode", "Volume"))
+					+ SSegmentedControl<EMirrorDirection>::Slot(EMirrorDirection::PositiveToNegative)
+					.Text(LOCTEXT("MirrorPosToNegLabel", "+ to -"))
+					+ SSegmentedControl<EMirrorDirection>::Slot(EMirrorDirection::NegativeToPositive)
+					.Text(LOCTEXT("MirrorNegToPosLabel", "- to +"))
 				]
 			]
-		];
-	const TSharedRef<IPropertyHandle> FalloffPropHandle = DetailBuilder.GetProperty(GET_MEMBER_NAME_CHECKED(USkinWeightsPaintToolProperties, FalloffMode));
-	DetailBuilder.HideProperty(FalloffPropHandle);
+		]
 
-	// custom display of falloff mode as segmented toggle buttons
-	IDetailCategoryBuilder& WeightColorsCategory = DetailBuilder.EditCategory("WeightColors", FText::GetEmpty(), ECategoryPriority::Important);
-	WeightColorsCategory.InitiallyCollapsed(true);
+		+SVerticalBox::Slot()
+		.Padding(0.f, WeightEditVerticalPadding)
+		[
+			SNew(SBox)
+			[
+				SNew(SButton)
+				.HAlign(HAlign_Center)
+				.Text(LOCTEXT("MirrorWeightsButtonLabel", "Mirror"))
+				.ToolTipText(LOCTEXT("MirrorButtonTooltip",
+					"Weights are copied across the given plane in the given direction.\n"
+					"This command operates on the selected bone(s) and selected vertices.\n "
+					"If no bones are selected, ALL bone weights are considered.\n "
+					"If no vertices are selected, ALL vertices are considered."))
+				.IsEnabled_Lambda([this]()
+				{
+					return SkinToolSettings->EditingMode == EWeightEditMode::Selection;
+				})
+				.OnClicked_Lambda([this]()
+				{
+					SkinToolSettings->WeightTool->MirrorWeights(SkinToolSettings->MirrorAxis, SkinToolSettings->MirrorDirection);
+					return FReply::Handled();
+				})
+			]
+		]
+	];
 
-	// add segmented control toggle for brush behavior modes ("Add", "Replace", etc..)
-	WeightColorsCategory.AddCustomRow(LOCTEXT("ColorModeCategory", "Color Mode"), false)
-		.NameContent()
+	// FLOOD WEIGHTS category
+	EditWeightsCategory.AddCustomRow(LOCTEXT("FloodWeightsRow", "Flood"), false)
+	.WholeRowContent()
+	[
+		SNew(SVerticalBox)
+
+		+ SVerticalBox::Slot()
+		.Padding(0.f, WeightEditVerticalPadding)
 		[
-			SNew(STextBlock)
-			.Text(LOCTEXT("ColorModeLabel", "Color Mode"))
-			.Font(FAppStyle::GetFontStyle(TEXT("PropertyWindow.NormalFont")))
-			.ToolTipText(LOCTEXT("ColorModeTooltip", "Determines how the weight colors are displayed."))
+			SNew(SHorizontalBox)
+				
+			+SHorizontalBox::Slot()
+			.VAlign(VAlign_Center)
+			.FillWidth(WeightEditingLabelsPercent)
+			[
+				SNew(STextBlock)
+				.Text(LOCTEXT("FloodAmountLabel", "Flood Amount"))
+				.Font(FAppStyle::GetFontStyle(TEXT("PropertyWindow.NormalFont")))
+				.ToolTipText(LOCTEXT("FloodAmountTooltip", "The amount of weight to apply in Flood operation."))
+			]
+
+			+SHorizontalBox::Slot()
+			.MaxWidth(WeightSliderWidths)
+			[
+				SNew(SSpinBox<float>)
+				.MinValue(0.f)
+				.MaxValue(1.f)
+				.IsEnabled_Lambda([this]()
+				{
+					return SkinToolSettings->EditingMode == EWeightEditMode::Selection;
+				})
+				.Value_Lambda([this]()
+				{
+					return SkinToolSettings->FloodValue;
+				})
+				.OnValueChanged_Lambda([this](float NewValue)
+				{
+					SkinToolSettings->FloodValue = NewValue;
+				})
+			]
 		]
-		.ValueContent()
+
+		+ SVerticalBox::Slot()
+		.Padding(0.f, WeightEditVerticalPadding)
 		[
 			SNew(SHorizontalBox)
-			+ SHorizontalBox::Slot()
+				
+			+SHorizontalBox::Slot()
+			.VAlign(VAlign_Center)
+			.FillWidth(WeightEditingLabelsPercent)
+			[
+				SNew(STextBlock)
+				.Text(LOCTEXT("FloodOperationLabel", "Flood Operation"))
+				.Font(FAppStyle::GetFontStyle(TEXT("PropertyWindow.NormalFont")))
+				.ToolTipText(LOCTEXT("FloodOperationTooltip", "The operation to perform when Flood button is pressed."))
+			]
+
+			+SHorizontalBox::Slot()
+			.FillWidth(1.f)
 			[
-				SNew(SBox)
-				.Padding(2.0f)
+				SNew(SSegmentedControl<EWeightEditOperation>)
+				.ToolTipText(LOCTEXT("FloodOperationTooltip",
+					"Add: applies the current weight PLUS the flood value to the new weight.\n"
+					"Remove: applies the current weight MINUS the flood value to the new weight.\n"
+					"Multiply: applies the current weight MULTIPLIED by the flood value to the new weight.\n"
+					"Relax: applies the average of the connected (by edge) vertex weights to the new vertex weight.\n"))
+				.IsEnabled_Lambda([this]()
+				{
+					return SkinToolSettings->EditingMode == EWeightEditMode::Selection;
+				})
+				.Value_Lambda([this]()
+				{
+					return SkinToolSettings->FloodMode;
+				})
+				.OnValueChanged_Lambda([this](EWeightEditOperation Mode)
+				{
+					SkinToolSettings->FloodMode = Mode;
+				})
+				+SSegmentedControl<EWeightEditOperation>::Slot(EWeightEditOperation::Add)
+				.Text(LOCTEXT("AddMode", "Add"))
+				+ SSegmentedControl<EWeightEditOperation>::Slot(EWeightEditOperation::Replace)
+				.Text(LOCTEXT("ReplaceMode", "Replace"))
+				+ SSegmentedControl<EWeightEditOperation>::Slot(EWeightEditOperation::Multiply)
+				.Text(LOCTEXT("MultiplyMode", "Multiply"))
+				+ SSegmentedControl<EWeightEditOperation>::Slot(EWeightEditOperation::Relax)
+				.Text(LOCTEXT("RelaxMode", "Relax"))
+			]
+		]
+		
+		+SVerticalBox::Slot()
+		.Padding(0.f, WeightEditVerticalPadding)
+		[
+			SNew(SBox)
+			[
+				SNew(SButton)
 				.HAlign(HAlign_Center)
-				[
-					SNew(SSegmentedControl<EWeightColorMode>)
-					.Value_Lambda([this]()
-					{
-						return SkinToolSettings->ColorMode;
-					})
-					.OnValueChanged_Lambda([this](EWeightColorMode Mode)
-					{
-						SkinToolSettings->ColorMode = Mode;
-						SkinToolSettings->bColorModeChanged = true;
-					})
-					+ SSegmentedControl<EWeightColorMode>::Slot(EWeightColorMode::Greyscale)
-					.Text(LOCTEXT("GreyscaleMode", "Greyscale"))
-					+SSegmentedControl<EWeightColorMode>::Slot(EWeightColorMode::ColorRamp)
-					.Text(LOCTEXT("RampMode", "Color Ramp"))
-				]
+				.Text(LOCTEXT("FloodWeightsButtonLabel", "Flood"))
+				.ToolTipText(LOCTEXT("FloodButtonTooltip",
+					"Modifies the vertex weights according to the chosen operation. See the operation tooltip for details.\n"
+					"This command operates on the selected bone(s) and selected vertices.\n"
+					"If no bones are selected, ALL bones are considered.\n"
+					"If no vertices are selected, ALL vertices are considered."))
+				.IsEnabled_Lambda([this]()
+				{
+					return SkinToolSettings->EditingMode == EWeightEditMode::Selection;
+				})
+				.OnClicked_Lambda([this]()
+				{
+					SkinToolSettings->WeightTool->FloodWeights(SkinToolSettings->FloodValue, SkinToolSettings->FloodMode);
+					return FReply::Handled();
+				})
+			]
+		]
+	];
+
+	// PRUNE WEIGHTS category
+	EditWeightsCategory.AddCustomRow(LOCTEXT("PruneWeightsRow", "Prune"), false)
+	.WholeRowContent()
+	[
+		SNew(SVerticalBox)
+
+		+ SVerticalBox::Slot()
+		.Padding(0.f, WeightEditVerticalPadding)
+		[
+			SNew(SHorizontalBox)
+					
+			+SHorizontalBox::Slot()
+			.VAlign(VAlign_Center)
+			.FillWidth(WeightEditingLabelsPercent)
+			[
+				SNew(STextBlock)
+				.Text(LOCTEXT("PruneThresholdLabel", "Prune Threshold"))
+				.Font(FAppStyle::GetFontStyle(TEXT("PropertyWindow.NormalFont")))
+				.ToolTipText(LOCTEXT("PruneThresholdTooltip", "The threshold weight value to use with Prune operation."))
+			]
+
+			+SHorizontalBox::Slot()
+			.MaxWidth(WeightSliderWidths)
+			[
+				SNew(SSpinBox<float>)
+				.MinValue(0.f)
+				.MaxValue(1.f)
+				.IsEnabled_Lambda([this]()
+				{
+					return SkinToolSettings->EditingMode == EWeightEditMode::Selection;
+				})
+				.Value_Lambda([this]()
+				{
+					return SkinToolSettings->PruneValue;
+				})
+				.OnValueChanged_Lambda([this](float NewValue)
+				{
+					SkinToolSettings->PruneValue = NewValue;
+				})
+			]
+		]
+		
+		+SVerticalBox::Slot()
+		.Padding(0.f, WeightEditVerticalPadding)
+		[
+			SNew(SBox)
+			.MinDesiredWidth(WeightSliderWidths)
+			[
+				SNew(SButton)
+				.HAlign(HAlign_Center)
+				.Text(LOCTEXT("PruneWeightsButtonLabel", "Prune"))
+				.ToolTipText(LOCTEXT("PruneButtonTooltip",
+					"Weights below the given threshold value are removed.\n"
+					"This command operates on the selected bone(s) and selected vertices.\n "
+					"If no bones are selected, ALL bone weights are considered.\n "
+					"If no vertices are selected, ALL vertices are considered."))
+				.IsEnabled_Lambda([this]()
+				{
+					return SkinToolSettings->EditingMode == EWeightEditMode::Selection;
+				})
+				.OnClicked_Lambda([this]()
+				{
+					SkinToolSettings->WeightTool->PruneWeights(SkinToolSettings->PruneValue);
+					return FReply::Handled();
+				})
 			]
-		];
+		]
+	];
+	
+	// COLOR MODE category
+	IDetailCategoryBuilder& WeightColorsCategory = DetailBuilder.EditCategory("WeightColors", FText::GetEmpty(), ECategoryPriority::Important);
+	WeightColorsCategory.InitiallyCollapsed(true);
+	WeightColorsCategory.AddCustomRow(LOCTEXT("ColorModeCategory", "Color Mode"), false)
+	.NameContent()
+	[
+		SNew(STextBlock)
+		.Text(LOCTEXT("ColorModeLabel", "Color Mode"))
+		.Font(FAppStyle::GetFontStyle(TEXT("PropertyWindow.NormalFont")))
+		.ToolTipText(LOCTEXT("ColorModeTooltip", "Determines how the weight colors are displayed."))
+	]
+	.ValueContent()
+	[
+		SNew(SHorizontalBox)
+		+ SHorizontalBox::Slot()
+		[
+			SNew(SBox)
+			.Padding(2.0f)
+			.HAlign(HAlign_Center)
+			[
+				SNew(SSegmentedControl<EWeightColorMode>)
+				.Value_Lambda([this]()
+				{
+					return SkinToolSettings->ColorMode;
+				})
+				.OnValueChanged_Lambda([this](EWeightColorMode Mode)
+				{
+					SkinToolSettings->ColorMode = Mode;
+					SkinToolSettings->bColorModeChanged = true;
+				})
+				+ SSegmentedControl<EWeightColorMode>::Slot(EWeightColorMode::Greyscale)
+				.Text(LOCTEXT("GreyscaleMode", "Greyscale"))
+				+SSegmentedControl<EWeightColorMode>::Slot(EWeightColorMode::ColorRamp)
+				.Text(LOCTEXT("RampMode", "Color Ramp"))
+			]
+		]
+	];
+
+	// hide all base brush properties that have been customized
+	const TSharedRef<IPropertyHandle> BrushModeHandle = DetailBuilder.GetProperty(GET_MEMBER_NAME_CHECKED(USkinWeightsPaintToolProperties, BrushMode));
+	DetailBuilder.HideProperty(BrushModeHandle);
+	const TSharedRef<IPropertyHandle> BrushSizeHandle = DetailBuilder.GetProperty(GET_MEMBER_NAME_CHECKED(UBrushBaseProperties, BrushSize), UBrushBaseProperties::StaticClass());
+	DetailBuilder.HideProperty(BrushSizeHandle);
+	const TSharedRef<IPropertyHandle> BrushStrengthHandle = DetailBuilder.GetProperty(GET_MEMBER_NAME_CHECKED(UBrushBaseProperties, BrushStrength), UBrushBaseProperties::StaticClass());
+	DetailBuilder.HideProperty(BrushStrengthHandle);
+	const TSharedRef<IPropertyHandle> BrushFalloffHandle = DetailBuilder.GetProperty(GET_MEMBER_NAME_CHECKED(UBrushBaseProperties, BrushFalloffAmount), UBrushBaseProperties::StaticClass());
+	DetailBuilder.HideProperty(BrushFalloffHandle);
+	const TSharedRef<IPropertyHandle> BrushRadiusHandle = DetailBuilder.GetProperty(GET_MEMBER_NAME_CHECKED(UBrushBaseProperties, BrushRadius), UBrushBaseProperties::StaticClass());
+	DetailBuilder.HideProperty(BrushRadiusHandle);
+	const TSharedRef<IPropertyHandle> SpecifyRadiusHandle = DetailBuilder.GetProperty(GET_MEMBER_NAME_CHECKED(UBrushBaseProperties, bSpecifyRadius), UBrushBaseProperties::StaticClass());
+	DetailBuilder.HideProperty(SpecifyRadiusHandle);
+	const TSharedRef<IPropertyHandle> FalloffPropHandle = DetailBuilder.GetProperty(GET_MEMBER_NAME_CHECKED(USkinWeightsPaintToolProperties, FalloffMode));
+	DetailBuilder.HideProperty(FalloffPropHandle);
+	const TSharedRef<IPropertyHandle> EditModePropHandle = DetailBuilder.GetProperty(GET_MEMBER_NAME_CHECKED(USkinWeightsPaintToolProperties, EditingMode));
+	DetailBuilder.HideProperty(EditModePropHandle);
 	const TSharedRef<IPropertyHandle> ColorModePropHandle = DetailBuilder.GetProperty(GET_MEMBER_NAME_CHECKED(USkinWeightsPaintToolProperties, ColorMode));
 	DetailBuilder.HideProperty(ColorModePropHandle);
 }
diff --git a/Engine/Source/Runtime/Engine/Classes/Animation/MirrorDataTable.h b/Engine/Source/Runtime/Engine/Classes/Animation/MirrorDataTable.h
--- a/Engine/Source/Runtime/Engine/Classes/Animation/MirrorDataTable.h
+++ b/Engine/Source/Runtime/Engine/Classes/Animation/MirrorDataTable.h
@@ -176,6 +176,27 @@ public:
 	 */
 	ENGINE_API FName FindReplace(FName InName) const;
 
+	/**
+	 * Finds the "best matching" mirrored bone across the specified axis. Priority is given to bones with the mirrored name,
+	 * falling back to spatial proximity if no mirrored bone is found using the naming rules.
+	 *
+	 * When falling back to proximity, bones within the SearchThreshold distance are considered coincident and a fuzzy string
+	 * comparison is used to find the most likely bone that matches the input bone.
+	 *
+	 * NOTE: The naming scheme assumes a mirror axis of X (Left/Right). Naming rules for other axes are not supported.
+	 *
+	 * @param	InBoneName		The input bone for which you want to find the mirrored equivalent
+	 * @param	InRefSkeleton	The reference skeleton used to find bone names and their spatial relationships (in ref pose)
+	 * @param	InMirrorAxis	The axis to cross when searching for a mirrored bone
+	 * @param	SearchThreshold	The distance in Unreal units to consider when trying to "tie-break" coincident bones
+	 * @return					The "best match" mirrored bone
+	 */
+	ENGINE_API static FName FindBestMirroredBone(
+		const FName InBoneName,
+		const FReferenceSkeleton& InRefSkeleton,
+		EAxis::Type InMirrorAxis,
+		const float SearchThreshold = 2.0f);
+
 public:
 
 	UPROPERTY(EditAnywhere, Category = CreateTable)
diff --git a/Engine/Source/Runtime/Engine/Private/Animation/MirrorDataTable.cpp b/Engine/Source/Runtime/Engine/Private/Animation/MirrorDataTable.cpp
--- a/Engine/Source/Runtime/Engine/Private/Animation/MirrorDataTable.cpp
+++ b/Engine/Source/Runtime/Engine/Private/Animation/MirrorDataTable.cpp
@@ -1,6 +1,9 @@
 // Copyright Epic Games, Inc. All Rights Reserved.
 
 #include "Animation/MirrorDataTable.h"
+
+#include "AnimationRuntime.h"
+#include "Algo/LevenshteinDistance.h"
 #include "Animation/AnimationSettings.h"
 #include "Animation/Skeleton.h"
 #include "Internationalization/Regex.h"
@@ -147,6 +150,91 @@ FName UMirrorDataTable::FindReplace(FName InName) const
 	return GetMirrorName(InName, MirrorFindReplaceExpressions); 
 }
 
+FName UMirrorDataTable::FindBestMirroredBone(
+	const FName InBoneName,
+	const FReferenceSkeleton& InRefSkeleton,
+	EAxis::Type InMirrorAxis,
+	const float SearchThreshold)
+{
+	const int32 SourceBoneIndex = InRefSkeleton.FindBoneIndex(InBoneName);
+	if (!ensureMsgf(SourceBoneIndex != INDEX_NONE, TEXT("Trying to find a mirror for a bone that isn't in the skeleton.")))
+	{
+		return NAME_None;
+	}
+	
+	// if the bone with the mirrored name exists in the skeleton, then great just use that...
+	const FName MirroredName = GetSettingsMirrorName(InBoneName);
+	if (InRefSkeleton.FindBoneIndex(MirroredName) != INDEX_NONE)
+	{
+		return MirroredName;
+	}
+
+	// fallback to closest mirrored bone, breaking ties (coincident bones) with fuzzy string score
+	TArray<FTransform> RefPoseGlobal;
+	FAnimationRuntime::FillUpComponentSpaceTransforms(InRefSkeleton, InRefSkeleton.GetRefBonePose(), RefPoseGlobal);
+	FVector MirroredLocation = RefPoseGlobal[SourceBoneIndex].GetLocation();
+	switch (InMirrorAxis)
+	{
+	case EAxis::X:
+		MirroredLocation.X *= -1.f;
+		break;
+	case EAxis::Y:
+		MirroredLocation.Y *= -1.f;
+		break;
+	case EAxis::Z:
+		MirroredLocation.Z *= -1.f;
+		break;
+	default:
+		checkNoEntry();
+	}
+
+	// find closest bone and all bones near the mirrored location within our search threshold
+	TArray<int32> BonesWithinThreshold;
+	int32 ClosestBone = 0;
+	int32 CurrentBone = 0;
+	float ClosestDistSq = TNumericLimits<float>::Max();
+	for (const FTransform& BoneTransform : RefPoseGlobal)
+	{
+		const float DistSq = (BoneTransform.GetLocation() - MirroredLocation).SizeSquared();
+		if (DistSq < ClosestDistSq)
+		{
+			ClosestDistSq = DistSq;
+			ClosestBone = CurrentBone;
+		}
+		if (DistSq <= FMath::Pow(SearchThreshold, 2.f))
+		{
+			BonesWithinThreshold.Add(CurrentBone);
+		}
+		++CurrentBone;
+	}
+
+	// no other bones were found near the mirrored location, so return the closest one
+	if (BonesWithinThreshold.Num() <= 1)
+	{
+		return InRefSkeleton.GetBoneName(ClosestBone);
+	}
+	
+	// in the case where we have multiple bones at or near the mirrored location (that are within our search threshold)
+	// it would be arbitrary to pick the "closest" one since bones are not always placed with that degree of precision.
+	// in this case, we break the tie with a fuzzy string comparison with the source bone name...
+	const FString SourceBoneStr = InBoneName.ToString().ToLower();
+	float BestScore = 0.f;
+	int32 BestBoneIndex = INDEX_NONE;
+	for (const int32 BoneToTestIndex : BonesWithinThreshold)
+	{
+		FString BoneToTestStr = InRefSkeleton.GetBoneName(BoneToTestIndex).ToString().ToLower();
+		const float WorstCase = SourceBoneStr.Len() + BoneToTestStr.Len();
+		const float Score = 1.0f - (Algo::LevenshteinDistance(BoneToTestStr, SourceBoneStr) / WorstCase);
+		if (Score > BestScore)
+		{
+			BestBoneIndex = BoneToTestIndex;
+			BestScore = Score;
+		}
+	}
+	
+	return InRefSkeleton.GetBoneName(BestBoneIndex);
+}
+
 #if WITH_EDITOR  
 
 void UMirrorDataTable::FindReplaceMirroredNames()
