Modeling Toolset: Store the name of the class being customized so that UnloadModule doesn't try to reference a UObject that might be deleted already during shutdown.

#rnx

diff --git a/Engine/Plugins/Experimental/MeshModelingToolsetExp/Source/ModelingEditorUI/Private/ModelingEditorUIModule.cpp b/Engine/Plugins/Experimental/MeshModelingToolsetExp/Source/ModelingEditorUI/Private/ModelingEditorUIModule.cpp
--- a/Engine/Plugins/Experimental/MeshModelingToolsetExp/Source/ModelingEditorUI/Private/ModelingEditorUIModule.cpp
+++ b/Engine/Plugins/Experimental/MeshModelingToolsetExp/Source/ModelingEditorUI/Private/ModelingEditorUIModule.cpp
@@ -10,17 +10,33 @@
 
 void FModelingEditorUIModule::StartupModule()
 {
-	// register detail customizations
 	FPropertyEditorModule& PropertyEditorModule = FModuleManager::LoadModuleChecked<FPropertyEditorModule>("PropertyEditor");
-	PropertyEditorModule.RegisterCustomClassLayout(USkinWeightsPaintToolProperties::StaticClass()->GetFName(), FOnGetDetailCustomizationInstance::CreateStatic(&FSkinWeightDetailCustomization::MakeInstance));
 	
+	auto RegisterDetailCustomization = [&](FName InStructName, auto InCustomizationFactory)
+    {
+    	PropertyEditorModule.RegisterCustomClassLayout(
+    		InStructName,
+    		FOnGetDetailCustomizationInstance::CreateStatic(InCustomizationFactory)
+    	);
+    	CustomizedClasses.Add(InStructName);
+    };
+
+	// register detail customizations
+	RegisterDetailCustomization(USkinWeightsPaintToolProperties::StaticClass()->GetFName(), &FSkinWeightDetailCustomization::MakeInstance);
+	
+	PropertyEditorModule.NotifyCustomizationModuleChanged();
 }
 
 void FModelingEditorUIModule::ShutdownModule()
 {
-	// unregister detail customizations
-	FPropertyEditorModule& PropertyEditorModule = FModuleManager::LoadModuleChecked<FPropertyEditorModule>("PropertyEditor");
-	PropertyEditorModule.UnregisterCustomClassLayout(USkinWeightsPaintToolProperties::StaticClass()->GetFName());
+	if (FPropertyEditorModule* PropertyModule = FModuleManager::GetModulePtr<FPropertyEditorModule>("PropertyEditor"))
+	{
+		for (const FName& ClassName : CustomizedClasses)
+		{
+			PropertyModule->UnregisterCustomClassLayout(ClassName);
+		}
+		PropertyModule->NotifyCustomizationModuleChanged();
+	}
 }
 
 #undef LOCTEXT_NAMESPACE
diff --git a/Engine/Plugins/Experimental/MeshModelingToolsetExp/Source/ModelingEditorUI/Public/ModelingEditorUIModule.h b/Engine/Plugins/Experimental/MeshModelingToolsetExp/Source/ModelingEditorUI/Public/ModelingEditorUIModule.h
--- a/Engine/Plugins/Experimental/MeshModelingToolsetExp/Source/ModelingEditorUI/Public/ModelingEditorUIModule.h
+++ b/Engine/Plugins/Experimental/MeshModelingToolsetExp/Source/ModelingEditorUI/Public/ModelingEditorUIModule.h
@@ -12,4 +12,7 @@ public:
 	/** IModuleInterface implementation */
 	virtual void StartupModule() override;
 	virtual void ShutdownModule() override;
+
+private:
+	TArray<FName> CustomizedClasses;
 };
