Fixing localization strings

diff --git a/Engine/Plugins/Experimental/MeshModelingToolsetExp/Source/ModelingEditorUI/Private/SkinWeightDetailCustomization.cpp b/Engine/Plugins/Experimental/MeshModelingToolsetExp/Source/ModelingEditorUI/Private/SkinWeightDetailCustomization.cpp
--- a/Engine/Plugins/Experimental/MeshModelingToolsetExp/Source/ModelingEditorUI/Private/SkinWeightDetailCustomization.cpp
+++ b/Engine/Plugins/Experimental/MeshModelingToolsetExp/Source/ModelingEditorUI/Private/SkinWeightDetailCustomization.cpp
@@ -51,9 +51,9 @@ void FSkinWeightDetailCustomization::CustomizeDetails(IDetailLayoutBuilder& Deta
 				SkinToolSettings->EditingMode = Mode;
 			})
 			+SSegmentedControl<EWeightEditMode>::Slot(EWeightEditMode::Brush)
-			.Text(LOCTEXT("EditMode", "Brush"))
+			.Text(LOCTEXT("BrushEditMode", "Brush"))
 			+ SSegmentedControl<EWeightEditMode>::Slot(EWeightEditMode::Selection)
-			.Text(LOCTEXT("EditMode", "Selection"))
+			.Text(LOCTEXT("SelectionEditMode", "Selection"))
 		]
 	];
 	
@@ -68,7 +68,7 @@ void FSkinWeightDetailCustomization::CustomizeDetails(IDetailLayoutBuilder& Deta
 		.Padding(2.0f)
 		[
 			SNew(SSegmentedControl<EWeightEditOperation>)
-			.ToolTipText(LOCTEXT("FloodOperationTooltip",
+			.ToolTipText(LOCTEXT("WeightEditModeTooltip",
 					"Add: applies the current weight PLUS the strength value to the new weight.\n"
 					"Remove: applies the current weight MINUS the strength value to the new weight.\n"
 					"Multiply: applies the current weight MULTIPLIED by the strength value to the new weight.\n"
@@ -86,13 +86,13 @@ void FSkinWeightDetailCustomization::CustomizeDetails(IDetailLayoutBuilder& Deta
 				SkinToolSettings->BrushMode = Mode;
 			})
 			+SSegmentedControl<EWeightEditOperation>::Slot(EWeightEditOperation::Add)
-			.Text(LOCTEXT("AddMode", "Add"))
+			.Text(LOCTEXT("BrushAddMode", "Add"))
 			+ SSegmentedControl<EWeightEditOperation>::Slot(EWeightEditOperation::Replace)
-			.Text(LOCTEXT("ReplaceMode", "Replace"))
+			.Text(LOCTEXT("BrushReplaceMode", "Replace"))
 			+ SSegmentedControl<EWeightEditOperation>::Slot(EWeightEditOperation::Multiply)
-			.Text(LOCTEXT("MultiplyMode", "Multiply"))
+			.Text(LOCTEXT("BrushMultiplyMode", "Multiply"))
 			+ SSegmentedControl<EWeightEditOperation>::Slot(EWeightEditOperation::Relax)
-			.Text(LOCTEXT("RelaxMode", "Relax"))
+			.Text(LOCTEXT("BrushRelaxMode", "Relax"))
 		]
 	];
 
@@ -435,7 +435,7 @@ void FSkinWeightDetailCustomization::CustomizeDetails(IDetailLayoutBuilder& Deta
 			.FillWidth(1.f)
 			[
 				SNew(SSegmentedControl<EWeightEditOperation>)
-				.ToolTipText(LOCTEXT("FloodOperationTooltip",
+				.ToolTipText(LOCTEXT("FloodOperationButtonTooltip",
 					"Add: applies the current weight PLUS the flood value to the new weight.\n"
 					"Remove: applies the current weight MINUS the flood value to the new weight.\n"
 					"Multiply: applies the current weight MULTIPLIED by the flood value to the new weight.\n"
@@ -453,13 +453,13 @@ void FSkinWeightDetailCustomization::CustomizeDetails(IDetailLayoutBuilder& Deta
 					SkinToolSettings->FloodMode = Mode;
 				})
 				+SSegmentedControl<EWeightEditOperation>::Slot(EWeightEditOperation::Add)
-				.Text(LOCTEXT("AddMode", "Add"))
+				.Text(LOCTEXT("FloodAddMode", "Add"))
 				+ SSegmentedControl<EWeightEditOperation>::Slot(EWeightEditOperation::Replace)
-				.Text(LOCTEXT("ReplaceMode", "Replace"))
+				.Text(LOCTEXT("FloodReplaceMode", "Replace"))
 				+ SSegmentedControl<EWeightEditOperation>::Slot(EWeightEditOperation::Multiply)
-				.Text(LOCTEXT("MultiplyMode", "Multiply"))
+				.Text(LOCTEXT("FloodMultiplyMode", "Multiply"))
 				+ SSegmentedControl<EWeightEditOperation>::Slot(EWeightEditOperation::Relax)
-				.Text(LOCTEXT("RelaxMode", "Relax"))
+				.Text(LOCTEXT("FloodRelaxMode", "Relax"))
 			]
 		]
 		
