Adding vertex selection to skin weight tool.

diff --git a/Engine/Plugins/Experimental/MeshModelingToolsetExp/Source/MeshModelingToolsExp/Private/SkinWeightsPaintTool.cpp b/Engine/Plugins/Experimental/MeshModelingToolsetExp/Source/MeshModelingToolsExp/Private/SkinWeightsPaintTool.cpp
--- a/Engine/Plugins/Experimental/MeshModelingToolsetExp/Source/MeshModelingToolsExp/Private/SkinWeightsPaintTool.cpp
+++ b/Engine/Plugins/Experimental/MeshModelingToolsetExp/Source/MeshModelingToolsExp/Private/SkinWeightsPaintTool.cpp
@@ -19,7 +19,9 @@
 #include "Spatial/FastWinding.h"
 #include "Async/Async.h"
 #include "DynamicMesh/MeshAdapterUtil.h"
+#include "Selection/PolygonSelectionMechanic.h"
 #include "Spatial/PointSetHashTable.h"
+#include "Util/ColorConstants.h"
 
 #include UE_INLINE_GENERATED_CPP_BY_NAME(SkinWeightsPaintTool)
 
@@ -614,6 +616,38 @@ void USkinWeightsPaintTool::Setup()
 	bVisibleWeightsValid = false;
 
 	RecalculateBrushRadius();
+
+	// set up vertex selection mechanic
+	PolygonSelectionMechanic = NewObject<UPolygonSelectionMechanic>(this);
+	PolygonSelectionMechanic->bAddSelectionFilterPropertiesToParentTool = false;
+	PolygonSelectionMechanic->Setup(this);
+	PolygonSelectionMechanic->SetIsEnabled(false);
+	PolygonSelectionMechanic->OnSelectionChanged.AddUObject(this, &USkinWeightsPaintTool::OnSelectionModified);
+	// only select vertices
+	PolygonSelectionMechanic->Properties->bSelectEdges = false;
+	PolygonSelectionMechanic->Properties->bSelectFaces = false;
+	PolygonSelectionMechanic->Properties->bSelectVertices = true;
+	// adjust selection rendering for this context
+	PolygonSelectionMechanic->PolyEdgesRenderer.PointColor = FLinearColor(0.78f, 0.f, 0.78f);
+	PolygonSelectionMechanic->PolyEdgesRenderer.PointSize = 5.0f;
+	PolygonSelectionMechanic->HilightRenderer.PointColor = FLinearColor::Red;
+	PolygonSelectionMechanic->HilightRenderer.PointSize = 10.0f;
+	PolygonSelectionMechanic->SelectionRenderer.LineThickness = 0.0f;
+	PolygonSelectionMechanic->SelectionRenderer.PointColor = FLinearColor::Yellow;
+	PolygonSelectionMechanic->SelectionRenderer.PointSize = 5.0f;
+	PolygonSelectionMechanic->SetShowEdges(false);
+	// initialize the polygon selection mechanic
+	constexpr bool bAutoBuild = true;
+	const FDynamicMesh3* DynamicMesh = PreviewMesh->GetPreviewDynamicMesh();
+	SelectionTopology = MakeUnique<UE::Geometry::FTriangleGroupTopology>(DynamicMesh, bAutoBuild);
+	MeshSpatial = MakeUnique<UE::Geometry::FDynamicMeshAABBTree3>(DynamicMesh, bAutoBuild);
+	PolygonSelectionMechanic->Initialize(
+		DynamicMesh,
+		FTransform::Identity,
+		Component->GetWorld(),
+		SelectionTopology.Get(),
+		[this]() { return MeshSpatial.Get(); }
+	);
 	
 	// inform user of tool keys
 	// TODO talk with UX team about viewport overlay to show hotkeys
@@ -627,6 +661,29 @@ void USkinWeightsPaintTool::RegisterActions(FInteractiveToolActionSet& ActionSet
 	UDynamicMeshBrushTool::RegisterActions(ActionSet);
 }
 
+void USkinWeightsPaintTool::DrawHUD(FCanvas* Canvas, IToolsContextRenderAPI* RenderAPI)
+{
+	Super::DrawHUD(Canvas, RenderAPI);
+
+	if (PolygonSelectionMechanic)
+	{
+		PolygonSelectionMechanic->DrawHUD(Canvas, RenderAPI);
+	}
+}
+
+void USkinWeightsPaintTool::Render(IToolsContextRenderAPI* RenderAPI)
+{
+	if (WeightToolProperties->EditingMode == EWeightEditMode::Brush)
+	{
+		Super::Render(RenderAPI);	
+	}
+	
+	if (PolygonSelectionMechanic && WeightToolProperties->EditingMode == EWeightEditMode::Selection)
+	{
+		PolygonSelectionMechanic->Render(RenderAPI);
+	}
+}
+
 void USkinWeightsPaintTool::OnTick(float DeltaTime)
 {
 	if (bStampPending)
@@ -1234,6 +1291,12 @@ void USkinWeightsPaintTool::GetVerticesToEdit(TArray<VertexIndex>& OutVertexIndi
 	//
 	// 1. Prioritize selected vertices
 	// TODO check selection manager for selected vertices here
+	const FGroupTopologySelection& Selection = PolygonSelectionMechanic->GetActiveSelection();
+	if (!Selection.SelectedCornerIDs.IsEmpty())
+	{
+		OutVertexIndices = Selection.SelectedCornerIDs.Array();
+		return;
+	}
 
 	//
 	// 2. Fallback on vertices weighted to selected bones
@@ -1271,8 +1334,17 @@ BoneIndex USkinWeightsPaintTool::GetBoneIndexFromName(const FName BoneName) cons
 
 void USkinWeightsPaintTool::OnShutdown(EToolShutdownType ShutdownType)
 {
+	// save tool properties
 	WeightToolProperties->SaveProperties(this);
 
+	// shutdown polygon selection mechanic
+	if (PolygonSelectionMechanic)
+	{
+		PolygonSelectionMechanic->Shutdown();
+		PolygonSelectionMechanic = nullptr;
+	}
+
+	// apply changes to asset
 	if (ShutdownType == EToolShutdownType::Accept)
 	{
 		// apply the weights to the mesh description
@@ -1313,24 +1385,35 @@ void USkinWeightsPaintTool::ExternalUpdateWeights(const int32 BoneIndex, const T
 	}
 }
 
-void USkinWeightsPaintTool::MirrorWeights(EAxis::Type Axis, EMirrorDirection Direction)
+void FSkinMirrorData::RegenerateMirrorData(
+    const TArray<FName>& BoneNames,
+    const TMap<FName, BoneIndex>& BoneNameToIndexMap,
+	const FReferenceSkeleton& RefSkeleton,
+	const TArray<FVector>& RefPoseVertices,
+	EAxis::Type InMirrorAxis,
+	EMirrorDirection InMirrorDirection,
+	const float MaxSearchRadius)
 {
-	check(Axis != EAxis::None);
-	
-	// TODO, filter vertices by selection
-	const TArray<FVector>& Vertices = Weights.Deformer.RefPoseVertexPositions;
+	if (bIsInitialized && InMirrorAxis == Axis && InMirrorDirection==Direction)
+	{
+		// already initialized, just re-use cached data
+		return;
+	}
 
-	// TODO, provide some way to edit the mirror bone mapping, either by providing a UMirrorDataTable input or editing directly in the hierarchy view.
+	// need to re-initialize
+	bIsInitialized = false;
+	Axis = InMirrorAxis;
+	Direction = InMirrorDirection;
+	
 	// build bone map for mirroring
-	TMap<int32, int32> MirrorBoneMap;
-	const FReferenceSkeleton& RefSkeleton = Weights.Deformer.Component->GetSkeletalMeshAsset()->GetRefSkeleton();
-	for (FName BoneName : Weights.Deformer.BoneNames)
+	// TODO, provide some way to edit the mirror bone mapping, either by providing a UMirrorDataTable input or editing directly in the hierarchy view.
+	for (FName BoneName : BoneNames)
 	{
 		FName MirroredBoneName = UMirrorDataTable::FindBestMirroredBone(BoneName, RefSkeleton, Axis);
 
-		int32 BoneIndex = Weights.Deformer.BoneNameToIndexMap[BoneName];
-		int32 MirroredBoneIndex = Weights.Deformer.BoneNameToIndexMap[MirroredBoneName];
-		MirrorBoneMap.Add(BoneIndex, MirroredBoneIndex);
+		int32 BoneIndex = BoneNameToIndexMap[BoneName];
+		int32 MirroredBoneIndex = BoneNameToIndexMap[MirroredBoneName];
+		BoneMap.Add(BoneIndex, MirroredBoneIndex);
 		
 		// debug view bone mapping
 		//UE_LOG(LogTemp, Log, TEXT("Bone    : %s"), *BoneName.ToString());
@@ -1341,28 +1424,26 @@ void USkinWeightsPaintTool::MirrorWeights(EAxis::Type Axis, EMirrorDirection Dir
 	// hash grid constants
 	constexpr float HashGridCellSize = 2.0f;
 	constexpr float ThresholdRadius = 0.1f;
-	const float MaxSearchRadius = Weights.Deformer.Component->Bounds.SphereRadius;
 
 	// build a point set of the rest pose vertices
 	UE::Geometry::FDynamicPointSet3d PointSet;
-	for (int32 PointID = 0; PointID < Vertices.Num(); ++PointID)
+	for (int32 PointID = 0; PointID < RefPoseVertices.Num(); ++PointID)
 	{
-		PointSet.InsertVertex(PointID, Vertices[PointID]);
+		PointSet.InsertVertex(PointID, RefPoseVertices[PointID]);
 	}
 
 	// build a spatial hash map from the point set
 	UE::Geometry::FPointSetAdapterd PointSetAdapter = UE::Geometry::MakePointsAdapter(&PointSet);
 	UE::Geometry::FPointSetHashtable PointHash(&PointSetAdapter);
 	PointHash.Build(HashGridCellSize, FVector3d::Zero());
-
+	
 	// generate a map of point IDs on the target side, to their equivalent vertex ID on the source side 
-	TArray<TPair<int32, int32>> VertexMirrorMap; // TMap<TO, FROM>
 	TArray<int> PointsInThreshold;
 	TArray<int> PointsInSphere;
-	bool bSomeVerticesNotMirrored = false;
-	for (int32 TargetVertexID = 0; TargetVertexID < Vertices.Num(); ++TargetVertexID)
+	bAllVerticesMirrored = true;
+	for (int32 TargetVertexID = 0; TargetVertexID < RefPoseVertices.Num(); ++TargetVertexID)
 	{
-		const FVector& TargetPosition = Vertices[TargetVertexID];
+		const FVector& TargetPosition = RefPoseVertices[TargetVertexID];
 
 		if (Direction == EMirrorDirection::PositiveToNegative && TargetPosition[Axis-1] >= 0.f)
 		{
@@ -1396,16 +1477,16 @@ void USkinWeightsPaintTool::MirrorWeights(EAxis::Type Axis, EMirrorDirection Dir
 		// no mirrored points?
 		if (PointsInSphere.IsEmpty())
 		{
-			bSomeVerticesNotMirrored = true;
+			bAllVerticesMirrored = false;
 			continue;
 		}
 
 		// find the closest single point
 		float ClosestDistSq = TNumericLimits<float>::Max();
 		int32 ClosestVertexID = INDEX_NONE;
-		for (int32 PointInSphereID : PointsInSphere)
+		for (const int32 PointInSphereID : PointsInSphere)
 		{
-			const float DistSq = FVector::DistSquared(Vertices[PointInSphereID], MirroredPosition);
+			const float DistSq = FVector::DistSquared(RefPoseVertices[PointInSphereID], MirroredPosition);
 			if (DistSq < ClosestDistSq)
 			{
 				ClosestDistSq = DistSq;
@@ -1414,20 +1495,70 @@ void USkinWeightsPaintTool::MirrorWeights(EAxis::Type Axis, EMirrorDirection Dir
 		}
 		
 		// record the mirrored vertex ID for this vertex 
-		VertexMirrorMap.Add(TPair<int32, int32>(TargetVertexID, ClosestVertexID)); // (TO, FROM)
+		VertexMap.FindOrAdd(TargetVertexID, ClosestVertexID); // (TO, FROM)
 	}
+	
+	bIsInitialized = true;
+}
 
-	if (bSomeVerticesNotMirrored)
+
+void USkinWeightsPaintTool::MirrorWeights(EAxis::Type Axis, EMirrorDirection Direction)
+{
+	check(Axis != EAxis::None);
+	
+	// get all ref pose vertices
+	const TArray<FVector>& RefPoseVertices = Weights.Deformer.RefPoseVertexPositions;
+	const FReferenceSkeleton& RefSkeleton = Weights.Deformer.Component->GetSkeletalMeshAsset()->GetRefSkeleton();
+
+	// refresh mirror tables (cached / lazy generated)
+	MirrorData.RegenerateMirrorData(
+		Weights.Deformer.BoneNames,
+		Weights.Deformer.BoneNameToIndexMap,
+		RefSkeleton,
+		RefPoseVertices,
+		Axis,
+		Direction,
+		Weights.Deformer.Component->Bounds.SphereRadius);
+
+	// get a reference to the mirror tables
+	const TMap<int32, int32>& BoneMap = MirrorData.GetBoneMap();
+	const TMap<int32, int32>& VertexMirrorMap = MirrorData.GetVertexMap(); // <Target, Source>
+
+	// get set of vertices to mirror
+	TArray<VertexIndex> SelectedVertices;
+	GetVerticesToEdit(SelectedVertices);
+
+	// convert all vertex indices to the target side of the mirror plane
+	TSet<VertexIndex> VerticesToMirror;
+	
+	for (const VertexIndex SelectedVertex : SelectedVertices)
 	{
-		UE_LOG(LogMeshModelingTools, Log, TEXT("Mirror Skin Weights: some vertex weights were not mirrored because a vertex was not found close enough to the mirrored location."));
+		if (VertexMirrorMap.Contains(SelectedVertex))
+		{
+			// vertex is located across the mirror plane (target side, to copy TO)
+			VerticesToMirror.Add(SelectedVertex);
+		}
+		else
+		{
+			// vertex is located on the source side (to copy FROM), so we need to search for it's mirror target vertex
+			for (const TPair<int32, int32>& ToFromPair : VertexMirrorMap)
+			{
+				if (ToFromPair.Value != SelectedVertex)
+				{
+					continue;
+				}
+				VerticesToMirror.Add(ToFromPair.Key);
+				break;
+			}
+		}
 	}
-
-	// iterate mirror map and copy weights from source to target vertices
+	
+	// spin through all target vertices to mirror and copy weights from source
 	FMultiBoneWeightEdits WeightEditsFromMirroring;
-	for (const TTuple<int32, int32>& VerticesToCopy : VertexMirrorMap)
+	for (const VertexIndex VertexToMirror : VerticesToMirror)
 	{
-		int32 SourceVertexID = VerticesToCopy.Value;
-		int32 TargetVertexID = VerticesToCopy.Key;
+		const int32 SourceVertexID = VertexMirrorMap[VertexToMirror];
+		const int32 TargetVertexID = VertexToMirror;
 
 		// remove all weight on vertex
 		for (const FVertexBoneWeight& TargetBoneWeight : Weights.PreChangeWeights[TargetVertexID])
@@ -1440,7 +1571,7 @@ void USkinWeightsPaintTool::MirrorWeights(EAxis::Type Axis, EMirrorDirection Dir
 		// copy source weights, but with mirrored bones
 		for (const FVertexBoneWeight& SourceBoneWeight : Weights.PreChangeWeights[SourceVertexID])
 		{
-			int32 MirroredBoneIndex = MirrorBoneMap[SourceBoneWeight.BoneIndex];
+			const int32 MirroredBoneIndex = BoneMap[SourceBoneWeight.BoneIndex];
 			const float OldWeight = Weights.GetWeightOfBoneOnVertex(MirroredBoneIndex, TargetVertexID, Weights.PreChangeWeights);
 			const float NewWeight = SourceBoneWeight.Weight;
 			WeightEditsFromMirroring.MergeSingleEdit(MirroredBoneIndex, TargetVertexID, OldWeight, NewWeight);
@@ -1450,6 +1581,12 @@ void USkinWeightsPaintTool::MirrorWeights(EAxis::Type Axis, EMirrorDirection Dir
 	// apply the changes
 	const FText TransactionLabel = LOCTEXT("MirrorWeightChange", "Mirror skin weights.");
 	ApplyWeightEditsToMesh(TransactionLabel, WeightEditsFromMirroring);
+
+	// warn if some vertices were not mirrored
+	if (!MirrorData.GetAllVerticesMirrored())
+	{
+		UE_LOG(LogMeshModelingTools, Log, TEXT("Mirror Skin Weights: some vertex weights were not mirrored because a vertex was not found close enough to the mirrored location."));
+	}
 }
 
 void USkinWeightsPaintTool::FloodWeights(const float Weight, const EWeightEditOperation FloodMode)
@@ -1641,6 +1778,15 @@ void USkinWeightsPaintTool::HandleSkeletalMeshModified(const TArray<FName>& InBo
 	}
 }
 
+void USkinWeightsPaintTool::OnSelectionModified()
+{
+}
+
+void USkinWeightsPaintTool::ToggleEditingMode()
+{
+	PolygonSelectionMechanic->SetIsEnabled(WeightToolProperties->EditingMode == EWeightEditMode::Selection);
+}
+
 void USkinWeightsPaintTool::OnPropertyModified(UObject* ModifiedObject, FProperty* ModifiedProperty)
 {
 	Super::OnPropertyModified(ModifiedObject, ModifiedProperty);
diff --git a/Engine/Plugins/Experimental/MeshModelingToolsetExp/Source/MeshModelingToolsExp/Public/SkinWeightsPaintTool.h b/Engine/Plugins/Experimental/MeshModelingToolsetExp/Source/MeshModelingToolsExp/Public/SkinWeightsPaintTool.h
--- a/Engine/Plugins/Experimental/MeshModelingToolsetExp/Source/MeshModelingToolsExp/Public/SkinWeightsPaintTool.h
+++ b/Engine/Plugins/Experimental/MeshModelingToolsetExp/Source/MeshModelingToolsExp/Public/SkinWeightsPaintTool.h
@@ -7,6 +7,7 @@
 #include "MeshDescription.h"
 #include "DynamicMesh/DynamicVerticesOctree3.h"
 #include "BoneWeights.h"
+#include "GroupTopology.h"
 #include "Misc/Optional.h"
 #include "Containers/Map.h"
 #include "DynamicMesh/DynamicMeshOctree3.h"
@@ -18,10 +19,53 @@
 
 struct FMeshDescription;
 class USkinWeightsPaintTool;
+class UPolygonSelectionMechanic;
 
 using BoneIndex = int32;
 using VertexIndex = int32;
 
+// weight edit mode
+UENUM()
+enum class EWeightEditMode : uint8
+{
+	Brush,
+	Selection,
+};
+
+// weight color mode
+UENUM()
+enum class EWeightColorMode : uint8
+{
+	Greyscale,
+	ColorRamp,
+};
+
+// brush falloff mode
+UENUM()
+enum class EWeightBrushFalloffMode : uint8
+{
+	Surface,
+	Volume,
+};
+
+// operation type when editing weights
+UENUM()
+enum class EWeightEditOperation : uint8
+{
+	Add,
+	Replace,
+	Multiply,
+	Relax
+};
+
+// mirror direction mode
+UENUM()
+enum class EMirrorDirection : uint8
+{
+	PositiveToNegative,
+	NegativeToPositive,
+};
+
 namespace SkinPaintTool
 {
 	struct FSkinToolWeights;
@@ -151,49 +195,32 @@ namespace SkinPaintTool
 		// update deformation when vertex weights are modified
 		FSkinToolDeformer Deformer;
 	};
-}
 
-// weight edit mode
-UENUM()
-enum class EWeightEditMode : uint8
-{
-	Brush,
-	Selection,
-};
-
-// weight color mode
-UENUM()
-enum class EWeightColorMode : uint8
-{
-	Greyscale,
-	ColorRamp,
-};
-
-// brush falloff mode
-UENUM()
-enum class EWeightBrushFalloffMode : uint8
-{
-	Surface,
-	Volume,
-};
-
-// operation type when editing weights
-UENUM()
-enum class EWeightEditOperation : uint8
-{
-	Add,
-	Replace,
-	Multiply,
-	Relax
-};
-
-// mirror direction mode
-UENUM()
-enum class EMirrorDirection : uint8
-{
-	PositiveToNegative,
-	NegativeToPositive,
-};
+	struct FSkinMirrorData
+	{
+		void RegenerateMirrorData(
+			const TArray<FName>& BoneNames,
+			const TMap<FName, BoneIndex>& BoneNameToIndexMap,
+			const FReferenceSkeleton& RefSkeleton,
+			const TArray<FVector>& RefPoseVertices,
+			EAxis::Type InMirrorAxis,
+			EMirrorDirection InMirrorDirection,
+			const float MaxSearchRadius);
+
+		const TMap<int32, int32>& GetBoneMap() const { return BoneMap; };
+		const TMap<int32, int32>& GetVertexMap() const { return VertexMap; };
+		bool GetAllVerticesMirrored() const {return bAllVerticesMirrored; };
+		
+	private:
+		
+		bool bIsInitialized = false;
+		bool bAllVerticesMirrored = false;
+		TEnumAsByte<EAxis::Type> Axis;
+		EMirrorDirection Direction; 
+		TMap<int32, int32> BoneMap;
+		TMap<int32, int32> VertexMap; // <Target, Source>
+	};
+}
 
 UCLASS()
 class MESHMODELINGTOOLSEXP_API USkinWeightsPaintToolBuilder : public UMeshSurfacePointMeshEditingToolBuilder
@@ -257,20 +284,22 @@ class MESHMODELINGTOOLSEXP_API USkinWeightsPaintTool : public UDynamicMeshBrushT
 	GENERATED_BODY()
 
 public:
-	void RegisterActions(FInteractiveToolActionSet& ActionSet) override;
-
-	void Setup() override;
-
-	bool HasCancel() const override { return true; }
-	bool HasAccept() const override { return true; }
-	bool CanAccept() const override { return true; }
 
 	// UBaseBrushTool overrides
-	bool HitTest(const FRay& Ray, FHitResult& OutHit) override;
-	void OnBeginDrag(const FRay& Ray) override;
-	void OnUpdateDrag(const FRay& Ray) override;
-	void OnEndDrag(const FRay& Ray) override;
-	bool OnUpdateHover(const FInputDeviceRay& DevicePos) override;
+	virtual bool HasCancel() const override { return true; }
+	virtual bool HasAccept() const override { return true; }
+	virtual bool CanAccept() const override { return true; }
+	virtual bool HitTest(const FRay& Ray, FHitResult& OutHit) override;
+	virtual void OnBeginDrag(const FRay& Ray) override;
+	virtual void OnUpdateDrag(const FRay& Ray) override;
+	virtual void OnEndDrag(const FRay& Ray) override;
+	virtual bool OnUpdateHover(const FInputDeviceRay& DevicePos) override;
+	
+	// UInteractiveTool
+	virtual void Setup() override;
+	virtual void RegisterActions(FInteractiveToolActionSet& ActionSet) override;
+	virtual void DrawHUD(FCanvas* Canvas, IToolsContextRenderAPI* RenderAPI) override;
+	virtual void Render(IToolsContextRenderAPI* RenderAPI) override;
 
 	void ExternalUpdateWeights(const int32 BoneIndex, const TMap<int32, float>& IndexValues);
 
@@ -281,6 +310,9 @@ public:
 	void AverageWeights();
 	void NormalizeWeights();
 
+	// toggle brush / selection mode
+	void ToggleEditingMode();
+
 protected:
 
 	virtual void ApplyStamp(const FBrushStampData& Stamp);
@@ -334,6 +366,9 @@ protected:
 	// storage of vertex weights per bone 
 	SkinPaintTool::FSkinToolWeights Weights;
 
+	// cached mirror data
+	SkinPaintTool::FSkinMirrorData MirrorData;
+
 	// vertex colors updated when switching current bone or editing weights
 	void UpdateCurrentBoneVertexColors();
 	FVector4f WeightToColor(float Value) const;
@@ -361,5 +396,12 @@ protected:
 	// ISkeletalMeshEditionInterface
 	virtual void HandleSkeletalMeshModified(const TArray<FName>& InBoneNames, const ESkeletalMeshNotifyType InNotifyType) override;
 
+	// polygon selection mechanic
+	UPROPERTY()
+	TObjectPtr<UPolygonSelectionMechanic> PolygonSelectionMechanic;
+	TUniquePtr<UE::Geometry::FDynamicMeshAABBTree3> MeshSpatial = nullptr;
+	TUniquePtr<UE::Geometry::FTriangleGroupTopology> SelectionTopology = nullptr;
+	void OnSelectionModified();
+
 	friend SkinPaintTool::FSkinToolDeformer;
 };
diff --git a/Engine/Plugins/Experimental/MeshModelingToolsetExp/Source/ModelingEditorUI/Private/SkinWeightDetailCustomization.cpp b/Engine/Plugins/Experimental/MeshModelingToolsetExp/Source/ModelingEditorUI/Private/SkinWeightDetailCustomization.cpp
--- a/Engine/Plugins/Experimental/MeshModelingToolsetExp/Source/ModelingEditorUI/Private/SkinWeightDetailCustomization.cpp
+++ b/Engine/Plugins/Experimental/MeshModelingToolsetExp/Source/ModelingEditorUI/Private/SkinWeightDetailCustomization.cpp
@@ -49,6 +49,7 @@ void FSkinWeightDetailCustomization::CustomizeDetails(IDetailLayoutBuilder& Deta
 			.OnValueChanged_Lambda([this](EWeightEditMode Mode)
 			{
 				SkinToolSettings->EditingMode = Mode;
+				SkinToolSettings->WeightTool->ToggleEditingMode();
 			})
 			+SSegmentedControl<EWeightEditMode>::Slot(EWeightEditMode::Brush)
 			.Text(LOCTEXT("BrushEditMode", "Brush"))
