fix crash in modeling mode's bake transform tool when some selected meshes have the same source asset

diff --git a/Engine/Plugins/Experimental/MeshModelingToolsetExp/Source/MeshModelingToolsExp/Private/BakeTransformTool.cpp b/Engine/Plugins/Experimental/MeshModelingToolsetExp/Source/MeshModelingToolsExp/Private/BakeTransformTool.cpp
--- a/Engine/Plugins/Experimental/MeshModelingToolsetExp/Source/MeshModelingToolsExp/Private/BakeTransformTool.cpp
+++ b/Engine/Plugins/Experimental/MeshModelingToolsetExp/Source/MeshModelingToolsExp/Private/BakeTransformTool.cpp
@@ -136,7 +136,8 @@ void UBakeTransformTool::UpdateAssets()
 		const FMeshDescription* MeshDescription = UE::ToolTarget::GetMeshDescription(Targets[ComponentIdx], FGetMeshParameters(bTargetSupportsLODs, LODs[0]));
 		if (MapToFirstOccurrences[ComponentIdx] < ComponentIdx)
 		{
-			BoundsOfScaledRotatedMesh.Add(BoundsOfScaledRotatedMesh[MapToFirstOccurrences[ComponentIdx]]);
+			FBox Bounds = BoundsOfScaledRotatedMesh[MapToFirstOccurrences[ComponentIdx]];
+			BoundsOfScaledRotatedMesh.Add(Bounds);
 		}
 		else
 		{
