Adding ability to assign weight value > 1.0  in skin paint tool. Useful for "multiply" operation. Still normalizes afterwards and fixed minor issue this exposed with weight normalization.

diff --git a/Engine/Plugins/Experimental/MeshModelingToolsetExp/Source/MeshModelingToolsExp/Private/SkinWeightsPaintTool.cpp b/Engine/Plugins/Experimental/MeshModelingToolsetExp/Source/MeshModelingToolsExp/Private/SkinWeightsPaintTool.cpp
--- a/Engine/Plugins/Experimental/MeshModelingToolsetExp/Source/MeshModelingToolsExp/Private/SkinWeightsPaintTool.cpp
+++ b/Engine/Plugins/Experimental/MeshModelingToolsetExp/Source/MeshModelingToolsExp/Private/SkinWeightsPaintTool.cpp
@@ -283,10 +283,10 @@ void FSkinToolWeights::EditVertexWeightAndNormalize(
 	// are there no OTHER influences on this vertex?
 	const bool bVertexHasNoOtherInfluences = Total <= MinimumWeightThreshold;
 	
-	// if user applied FULL weight to this vertex AND there's no other weights of any significance,
+	// if user applied any weight to this vertex AND there's no other weights of any significance,
 	// then simply apply full weight to this vertex, set all other influences to zero and return
-	const bool bApplyFullWeightToThisVertex = NewWeightValue >= (1.0f - MinimumWeightThreshold);
-	if (bApplyFullWeightToThisVertex && bVertexHasNoOtherInfluences)
+	const bool bApplyWeightToThisVertex = NewWeightValue >= MinimumWeightThreshold;
+	if (bApplyWeightToThisVertex && bVertexHasNoOtherInfluences)
 	{
 		// set all other influences to 0.0f
 		for (int32 i=0; i<ValuesToNormalize.Num(); ++i)
diff --git a/Engine/Plugins/Experimental/MeshModelingToolsetExp/Source/ModelingEditorUI/Private/SkinWeightDetailCustomization.cpp b/Engine/Plugins/Experimental/MeshModelingToolsetExp/Source/ModelingEditorUI/Private/SkinWeightDetailCustomization.cpp
--- a/Engine/Plugins/Experimental/MeshModelingToolsetExp/Source/ModelingEditorUI/Private/SkinWeightDetailCustomization.cpp
+++ b/Engine/Plugins/Experimental/MeshModelingToolsetExp/Source/ModelingEditorUI/Private/SkinWeightDetailCustomization.cpp
@@ -171,7 +171,10 @@ void FSkinWeightDetailCustomization::CustomizeDetails(IDetailLayoutBuilder& Deta
 	[
 		SNew(SSpinBox<float>)
 		.MinValue(0.f)
-		.MaxValue(1.f)
+		.MaxValue(2.0f)
+		.MaxSliderValue(1.f)
+		.Value(1.0f)
+		.SupportDynamicSliderMaxValue(true)
 		.IsEnabled_Lambda([this]()
 		{
 			return SkinToolSettings->EditingMode == EWeightEditMode::Brush;
@@ -401,7 +404,10 @@ void FSkinWeightDetailCustomization::CustomizeDetails(IDetailLayoutBuilder& Deta
 			[
 				SNew(SSpinBox<float>)
 				.MinValue(0.f)
-				.MaxValue(1.f)
+				.MaxValue(2.0f)
+				.MaxSliderValue(1.f)
+				.Value(1.0f)
+				.SupportDynamicSliderMaxValue(true)
 				.IsEnabled_Lambda([this]()
 				{
 					return SkinToolSettings->EditingMode == EWeightEditMode::Selection;
