CombineMeshInstances improvements. Add support for externally-provided fixed LOD budgets for different LODs. Add 'Part Promotion' strategy in ComputeMeshApproximations() that incrementally "promotes" per-part LODs upwards so that combined mesh achieves triangle budget (see code comments for details).
Add explicit support for Axis-Aligned Box Approximation in addition to existing Oriented Box. Add explicit preference for AABB if it's volume is within 20% of OBB or best-shape approximation.
Other minor fixes.

diff --git a/Engine/Plugins/Experimental/MeshModelingToolsetExp/Source/GeometryProcessingAdapters/Private/CombineMeshInstancesImpl.cpp b/Engine/Plugins/Experimental/MeshModelingToolsetExp/Source/GeometryProcessingAdapters/Private/CombineMeshInstancesImpl.cpp
--- a/Engine/Plugins/Experimental/MeshModelingToolsetExp/Source/GeometryProcessingAdapters/Private/CombineMeshInstancesImpl.cpp
+++ b/Engine/Plugins/Experimental/MeshModelingToolsetExp/Source/GeometryProcessingAdapters/Private/CombineMeshInstancesImpl.cpp
@@ -716,23 +716,32 @@ static void PlanarRetriangulatePartMesh(
 
 static void ComputeBoxApproximation(
 	const FDynamicMesh3& SourceMesh,
-	FDynamicMesh3& OutputMesh )
+	FDynamicMesh3& OutputMesh,
+	bool bForceAxisAligned)
 {
-	FMeshSimpleShapeApproximation ShapeApprox;
-	ShapeApprox.InitializeSourceMeshes( {&SourceMesh} );
-	ShapeApprox.bDetectBoxes = ShapeApprox.bDetectCapsules = ShapeApprox.bDetectConvexes = ShapeApprox.bDetectSpheres = false;
+	UE::Geometry::FOrientedBox3d OrientedBox;
+	if (bForceAxisAligned)
+	{
+		OrientedBox = UE::Geometry::FOrientedBox3d(SourceMesh.GetBounds(false));
+	}
+	else
+	{
+		FMeshSimpleShapeApproximation ShapeApprox;
+		ShapeApprox.InitializeSourceMeshes({ &SourceMesh });
+		ShapeApprox.bDetectBoxes = ShapeApprox.bDetectCapsules = ShapeApprox.bDetectConvexes = ShapeApprox.bDetectSpheres = false;
 
-	FSimpleShapeSet3d ResultBoxes;
-	ShapeApprox.Generate_OrientedBoxes(ResultBoxes);
-	UE::Geometry::FOrientedBox3d OrientedBox = ResultBoxes.Boxes[0].Box;
+		FSimpleShapeSet3d ResultBoxes;
+		ShapeApprox.Generate_OrientedBoxes(ResultBoxes);
+		OrientedBox = ResultBoxes.Boxes[0].Box;
 
-	// oriented box fitting is under-determined, in cases where the AABB and the OBB have the nearly the
-	// same volume, generally we prefer an AABB
-	// (note: this rarely works due to tessellation of (eg) circles/spheres, and should be replaced w/ a better heuristic)
-	FAxisAlignedBox3d AlignedBox = SourceMesh.GetBounds(false);
-	if (AlignedBox.Volume() < 1.05*OrientedBox.Volume())
-	{
-		OrientedBox = UE::Geometry::FOrientedBox3d(AlignedBox);
+		// oriented box fitting is under-determined, in cases where the AABB and the OBB have the nearly the
+		// same volume, generally we prefer an AABB
+		// (note: this rarely works due to tessellation of (eg) circles/spheres, and should be replaced w/ a better heuristic)
+		FAxisAlignedBox3d AlignedBox = SourceMesh.GetBounds(false);
+		if (AlignedBox.Volume() < 1.2 * OrientedBox.Volume())
+		{
+			OrientedBox = UE::Geometry::FOrientedBox3d(AlignedBox);
+		}
 	}
 
 	FGridBoxMeshGenerator BoxGen;
@@ -745,8 +754,9 @@ static void ComputeBoxApproximation(
 
 enum class EApproximatePartMethod : uint8
 {
-	OrientedBox = 0,
-	MinVolumeSweptHull = 1,
+	AxisAlignedBox = 0,
+	OrientedBox = 1,
+	MinVolumeSweptHull = 2,
 	ConvexHull = 3,
 	MinTriCountHull = 4,
 	FlattendExtrusion = 5,
@@ -765,9 +775,16 @@ static void ComputeSimplePartApproximation(
 	EApproximatePartMethod ApproxMethod)
 {
 
+	if (ApproxMethod == EApproximatePartMethod::AxisAlignedBox)
+	{
+		ComputeBoxApproximation(SourcePartMesh, DestMesh, true);
+		return;
+	}
+
 	if (ApproxMethod == EApproximatePartMethod::OrientedBox)
 	{
-		ComputeBoxApproximation(SourcePartMesh, DestMesh);
+		ComputeBoxApproximation(SourcePartMesh, DestMesh, false);
+		return;
 	}
 
 	FMeshSimpleShapeApproximation ShapeApprox;
@@ -987,6 +1004,10 @@ static void SelectBestFittingMeshApproximation(
 	ApproxSelector.Initialize(&OriginalMesh, &OriginalMeshSpatial);
 	ApproxSelector.TriangleCost = TriangleCost;
 
+	ApproxSelector.AddGeneratedMesh( [&](FDynamicMesh3& PartMeshInOut) {
+		ComputeSimplePartApproximation(PartMeshInOut, PartMeshInOut, EApproximatePartMethod::AxisAlignedBox);
+	}, (int32)EApproximatePartMethod::AxisAlignedBox);
+
 	ApproxSelector.AddGeneratedMesh( [&](FDynamicMesh3& PartMeshInOut) {
 		ComputeSimplePartApproximation(PartMeshInOut, PartMeshInOut, EApproximatePartMethod::OrientedBox);
 	}, (int32)EApproximatePartMethod::OrientedBox );
@@ -1033,6 +1054,17 @@ static void SelectBestFittingMeshApproximation(
 	}
 
 	ApproxSelector.SelectBestOption(ResultMesh);
+
+	// if Axis-Aligned box volume is less than (100+k%) larger than best option, just use that instead
+	// (todo should be configurable)
+	const double BoxPreferenceVolumeRatioPercent = 20.0;
+	FVector2d ApproxMeshVolArea = TMeshQueries<FDynamicMesh3>::GetVolumeArea(ResultMesh);
+	FAxisAlignedBox3d AlignedBox = OriginalMesh.GetBounds(false);
+	if ((AlignedBox.Volume() / ApproxMeshVolArea.X) < (1.0 + BoxPreferenceVolumeRatioPercent/100.0) )
+	{
+		ComputeSimplePartApproximation(OriginalMesh, ResultMesh, EApproximatePartMethod::AxisAlignedBox);
+	}
+
 }
 
 
@@ -1079,10 +1111,14 @@ void ComputeMeshApproximations(
 			InitialTolerance *= CombineOptions.SimplifyLODLevelToleranceScale;
 		}
 
-		// compute shape approximation LODs
-		ApproxGeo.ApproximateMeshLODs.SetNum(NumApproxLODs);
+		// Compute shape approximation LODs. 
+		// Note that UseNumApproxLODs is a hack here - we are computing more than necessary
+		// so that the cost approximation strategy below has additional simplified approximations available. 
+		// This could be smarter, but this dumb method works OK for now...
+		int32 UseNumApproxLODs = NumApproxLODs + 10;
+		ApproxGeo.ApproximateMeshLODs.SetNum(UseNumApproxLODs);
 		double InitialTriCost = CombineOptions.OptimizeBaseTriCost;
-		for (int32 k = 0; k < NumApproxLODs; ++k)
+		for (int32 k = 0; k < UseNumApproxLODs; ++k)
 		{
 			SelectBestFittingMeshApproximation(OptimizationSourceMesh, OptimizationSourceMeshSpatial, 
 				ApproxGeo.ApproximateMeshLODs[k], CombineOptions.SimplifyBaseTolerance, InitialTriCost);
@@ -1096,7 +1132,7 @@ void ComputeMeshApproximations(
 			FMeshNormals::QuickRecomputeOverlayNormals(ApproxGeo.ApproximateMeshLODs[k]);
 		}
 
-		// test remeshing last source LOD
+		// try remeshing the last Source LOD to reduce it's triangle count, by removing spurious geometry
 		if (CombineOptions.bRetriangulateSourceLODs)
 		{
 			for (int32 SourceLODIndex = CombineOptions.bRetriangulateSourceLODs; SourceLODIndex < NumSourceLODs; ++SourceLODIndex)
@@ -1112,6 +1148,163 @@ void ComputeMeshApproximations(
 
 	// try to filter out simplifications that did bad things
 	ReplaceBadSimplifiedLODs(Assembly);
+
+
+	// Now that we have our per-part LOD stacks, we can estimate total triangle count that will be
+	// used by the combined mesh. This will not be accurate due to hidden removal and face merging.
+	// But we can, given a target triangle budget, try to "promote" simpler part approximations up
+	// their individual LOD chains in an attempt to achieve a triangle budget. Generally the external
+	// code will want to provide larger triangle counts, eg 40-50% larger, than the desired final
+	// triangle count, to account for hidden removal and face merging. 
+
+	int32 TotalNumLODs = CombineOptions.NumCopiedLODs + NumSimplifiedLODs + NumApproxLODs;
+	int32 LODsToProcess = FMath::Min(CombineOptions.HardLODBudgets.Num(), TotalNumLODs);
+	if (CombineOptions.bEnableBudgetStrategy_PartLODPromotion && LODsToProcess > 0)
+	{
+		struct FPartCostInfo
+		{
+			// these fields are precomputed
+			int32 PartIndex = 0;			// index into Assembly.InstanceSets
+			int32 NumInstances = 0;			// number of instances of this part in the final mesh, currently *excluding* decorative parts
+			double ReplacedWeight = 1.0;	// reduce this as we shift LODs up, to try to let other comparable parts take the hit...
+
+			TArray<FDynamicMesh3*> LODChainMeshes;	// flattened list of mesh pointers for [Source LODs][Simplified LODs][Approximate LODs]
+
+			// these fields are temporary storage updated during the algorilthm below
+			int32 PartTriCount = 0;			// triangle count of the part for the active LOD
+			int32 TotalTriCount = 0;		// total estimated triangle count for this part in the combined mesh 
+
+			double PartCostWeight() const 
+			{
+				if (NumInstances == 0 || PartTriCount <= 12) return 0.0;		// assume a box is min-cost and cannot be improved
+				return (double)TotalTriCount * ReplacedWeight;
+			}
+		};
+		TArray<FPartCostInfo> CostInfo;
+		CostInfo.SetNum(NumSets);
+
+		// initialize precomputed parts of the CostInfo array that we will incrementally update
+		for (int32 SetIndex = 0; SetIndex < NumSets; ++SetIndex)
+		{
+			const TUniquePtr<FMeshInstanceSet>& InstanceSet = Assembly.InstanceSets[SetIndex];
+			CostInfo[SetIndex].PartIndex = SetIndex;
+
+			CostInfo[SetIndex].NumInstances = 0;
+			for (const FMeshInstance& Instance : InstanceSet->Instances)
+			{
+				bool bSkipInstance = (Instance.DetailLevel == EMeshDetailLevel::Decorative);
+				if (bSkipInstance == false)
+				{
+					CostInfo[SetIndex].NumInstances++;
+				}
+			}
+
+			for (FDynamicMesh3& SourceLODMesh : Assembly.SourceMeshGeometry[SetIndex].SourceMeshLODs)
+			{
+				CostInfo[SetIndex].LODChainMeshes.Add(&SourceLODMesh);
+			}
+			for (FDynamicMesh3& SimplifiedLODMesh : Assembly.OptimizedMeshGeometry[SetIndex].SimplifiedMeshLODs)
+			{
+				CostInfo[SetIndex].LODChainMeshes.Add(&SimplifiedLODMesh);
+			}
+			for (FDynamicMesh3& SimplifiedLODMesh : Assembly.OptimizedMeshGeometry[SetIndex].ApproximateMeshLODs)
+			{
+				CostInfo[SetIndex].LODChainMeshes.Add(&SimplifiedLODMesh);
+			}
+		}
+
+		// for each LOD where we have a budget, while we are overbudget, do iterations of selecting
+		// a "most expensive" part and promoting LODN+1 up to LODN for that part. Repeat until
+		// budget is reached, or Max Iterations
+		for (int32 LODIndex = 0; LODIndex < LODsToProcess; ++LODIndex)
+		{
+			int32 LODTriangleBudget = CombineOptions.HardLODBudgets[LODIndex] * CombineOptions.PartLODPromotionBudgetMultiplier;
+			if (LODTriangleBudget <= 0) continue;
+
+			int32 LastTotalCurLODTriCount = 999999;
+			int32 MaxIters = 1000;
+			int32 NoProgressIters = 0;
+			for ( int32 NumIter = 0; NumIter < MaxIters; ++NumIter)
+			{
+				// compute current estimate of total part count for this LOD
+				int32 TotalCurLODTriCount = 0;
+				for (int32 SetIndex = 0; SetIndex < NumSets; ++SetIndex)
+				{
+					const FDynamicMesh3* CurSourcePartMeshLOD = CostInfo[SetIndex].LODChainMeshes[LODIndex];
+					CostInfo[SetIndex].PartTriCount = CurSourcePartMeshLOD->TriangleCount();
+					CostInfo[SetIndex].TotalTriCount = CostInfo[SetIndex].NumInstances * CostInfo[SetIndex].PartTriCount;
+					TotalCurLODTriCount += CostInfo[SetIndex].TotalTriCount;
+				}
+				if (TotalCurLODTriCount == LastTotalCurLODTriCount)
+				{
+					NoProgressIters++;
+				}
+				else
+				{
+					NoProgressIters = 0;
+				}
+				if (bVerbose && (NumIter % 25 == 0))
+				{
+					UE_LOG(LogGeometry, Log, TEXT("    PartPromotion LOD %1d: Iter %4d  CurTris %6d LastTris %6d Budget %6d  NoProgress %d"), LODIndex, NumIter, TotalCurLODTriCount, LastTotalCurLODTriCount, LODTriangleBudget, NoProgressIters);
+				}
+				if (TotalCurLODTriCount < LODTriangleBudget || NoProgressIters > 25)
+				{
+					break;		// either we are within budget, or we made no progress for too long
+				}
+				LastTotalCurLODTriCount = TotalCurLODTriCount;
+
+				// "No Progress" has to be allowed to occur for more than one step because it is often the case that the search gets 
+				// "stuck" for a while, promoting a LODN+1 to LODN that have the same triangle count. This is quite common eg whenever
+				// a simple box is reached in the LOD chain, as the entire LOD chain from that point will have the same triangle count. 
+				// It also occurs because of the effect of the per-part ReplacedWeight, which makes an expensive part "cheaper" 
+				// immediately after it is replaced. 
+				// This could be improved by making the search below smarter, ie if bumping the max-cost part doesn't help, try
+				// a different one. *However* note that because of how ReplacedWeight incrementally grows for each part over time,
+				// it is generally the case that even if there are a few no-progress iterations, eventually the ReplacedWeight on
+				// still-more-expensive parts will grow larger and progress is made again
+
+				// find part with largest current cost
+				int32 MaxSetIndex = 0;
+				for (int32 k = 1; k < NumSets; ++k)
+				{
+					double MaxCost = CostInfo[MaxSetIndex].PartCostWeight();
+					double CurCost = CostInfo[k].PartCostWeight();
+					if (CurCost > MaxCost)
+					{
+						MaxSetIndex = k;
+					}
+				}
+				int32 SetIndex = MaxSetIndex;
+				// if our worst part is a box (CostWeight == 0), there is no point in replacing it
+				if ( (CostInfo[SetIndex].PartCostWeight() > 0) 
+					&& (LODIndex < CostInfo[SetIndex].LODChainMeshes.Num()-2) )	
+				{
+					FPartCostInfo& ReplaceInfo = CostInfo[SetIndex];
+
+					if (bVerbose)
+					{
+						UE_LOG(LogGeometry, Log, TEXT("    PartPromotion LOD %1d: Iter %4d  Promoting Part %4d, from %5d tris to %5d tris"), 
+							LODIndex, NumIter, SetIndex, ReplaceInfo.LODChainMeshes[LODIndex]->TriangleCount(), ReplaceInfo.LODChainMeshes[LODIndex+1]->TriangleCount());
+					}
+
+					int32 NumAllLODs = ReplaceInfo.LODChainMeshes.Num();
+					for (int32 k = LODIndex; k < NumAllLODs - 1; ++k)
+					{
+						*ReplaceInfo.LODChainMeshes[k] = *ReplaceInfo.LODChainMeshes[k + 1];
+					}
+
+					ReplaceInfo.ReplacedWeight *= 0.5;
+				}
+				// slowly increase weights of parts   (should this be modulated by tri count?)
+				for (int32 k = 0; k < NumSets; ++k)
+				{
+					CostInfo[k].ReplacedWeight += 0.1;
+				}
+			}
+		}
+
+	}
+
 }
 
 
@@ -1389,6 +1582,7 @@ static void PostProcessHiddenFaceRemovedMesh(
 		FMergeCoincidentMeshEdges Welder(&TargetMesh);
 		Welder.MergeVertexTolerance = Tolerance * 0.01;
 		Welder.OnlyUniquePairs = false;
+		Welder.bWeldAttrsOnMergedEdges = true;
 		Welder.Apply();
 	}
 
@@ -1444,11 +1638,29 @@ static void PostProcessHiddenFaceRemovedMesh(
 		FMergeCoincidentMeshEdges Welder(&SubRegionMesh);
 		Welder.MergeVertexTolerance = Tolerance * 0.01;
 		Welder.OnlyUniquePairs = false;
+		Welder.bWeldAttrsOnMergedEdges = true;
 		Welder.Apply();
 
+		// todo: could discard UVs here to improve merging. Possibly also do planar remesh, it 
+		// seems like the current method still will not remove boundary vertices...
+
+		SubRegionMesh.Attributes()->SplitAllBowties();
+		SubRegionMesh.CompactInPlace();
+
 		// simplify to planar
 		FQEMSimplification Simplifier(&SubRegionMesh);
-		Simplifier.CollapseMode = FQEMSimplification::ESimplificationCollapseModes::AverageVertexPosition;
+
+		// set up constraints, necessary to avoid crashing in presence of attributes
+		FMeshConstraints Constraints;
+		FMeshConstraintsUtil::ConstrainAllBoundariesAndSeams(Constraints, SubRegionMesh,
+			EEdgeRefineFlags::NoFlip, EEdgeRefineFlags::NoConstraint, EEdgeRefineFlags::NoConstraint, true, false, true);
+		Simplifier.SetExternalConstraints(Constraints);
+		// need to transfer constraint setting to the simplifier, these are used to update the constraints as edges collapse.
+		Simplifier.MeshBoundaryConstraint = EEdgeRefineFlags::NoFlip;
+		Simplifier.GroupBoundaryConstraint = EEdgeRefineFlags::NoConstraint;
+		Simplifier.MaterialBoundaryConstraint = EEdgeRefineFlags::NoConstraint;
+
+		Simplifier.CollapseMode = FQEMSimplification::ESimplificationCollapseModes::MinimalExistingVertexError;
 		Simplifier.SimplifyToMinimalPlanar( 0.01 );
 	}
 
@@ -1472,6 +1684,7 @@ static void PostProcessHiddenFaceRemovedMesh(
 		FMergeCoincidentMeshEdges Welder(&TargetMesh);
 		Welder.MergeVertexTolerance = Tolerance * 0.01;
 		Welder.OnlyUniquePairs = false;
+		Welder.bWeldAttrsOnMergedEdges = true;
 		Welder.Apply();
 	}
 
@@ -1756,13 +1969,16 @@ static void ProjectAttributes(
 
 		if (SourceColors != nullptr)
 		{
-			FIndex3i SourceTriElems = SourceColors->GetTriangle(NearestTID);
-			// TODO be smarter here...
-			FVector4f Color = SourceColors->GetElement(SourceTriElems.A);
-			int A = TargetColors->AppendElement(Color);
-			int B = TargetColors->AppendElement(Color);
-			int C = TargetColors->AppendElement(Color);
-			TargetColors->SetTriangle(tid, FIndex3i(A, B, C));
+			if (SourceColors->IsSetTriangle(NearestTID))
+			{
+				FIndex3i SourceTriElems = SourceColors->GetTriangle(NearestTID);
+				// TODO be smarter here...
+				FVector4f Color = SourceColors->GetElement(SourceTriElems.A);
+				int A = TargetColors->AppendElement(Color);
+				int B = TargetColors->AppendElement(Color);
+				int C = TargetColors->AppendElement(Color);
+				TargetColors->SetTriangle(tid, FIndex3i(A, B, C));
+			}
 		}
 	}
 }
diff --git a/Engine/Source/Developer/GeometryProcessingInterfaces/Public/GeometryProcessingInterfaces/CombineMeshInstances.h b/Engine/Source/Developer/GeometryProcessingInterfaces/Public/GeometryProcessingInterfaces/CombineMeshInstances.h
--- a/Engine/Source/Developer/GeometryProcessingInterfaces/Public/GeometryProcessingInterfaces/CombineMeshInstances.h
+++ b/Engine/Source/Developer/GeometryProcessingInterfaces/Public/GeometryProcessingInterfaces/CombineMeshInstances.h
@@ -136,6 +136,34 @@ public:
 		// which Source LOD to start retriangulating at
 		int32 StartRetriangulateSourceLOD = 1;
 
+		//
+		// Optional Support for hitting explicit triangle counts for different LODs. 
+		// The HardLODBudgets list should be provided in LOD order, ie first for Copied LODs,
+		// then for Simplified LODs, then for Approximate LODs (Voxel LOD triangle counts are configured via VoxWrapMaxTriCountBase)
+		// 
+		// Currently the only explicit tri-count strategy in use is Part Promotion, where a coarser approximations
+		// (eg Lower Copied LOD, Simplified LOD, Approximate LOD) are "promoted" upwards into higher combined-mesh LODs
+		// as necessary to achieve the triangle count.
+		// 
+		// Note that final LOD triangle counts cannot be guaranteed, due to the combinatorial nature of the approximation.
+		// For example the coarsest part LOD is a box with 12 triangles, so NumParts*12 is a lower bound on the initial combined mesh.
+		// The Part Promotion strategy is applied *before* hidden removal and further mesh processing (eg coplanar merging),
+		// so the final triangle count may be substantially lower than the budget (this is why the Multiplier is used below)
+		//
+
+		// list of fixed-triangle-count LOD budgets, in LOD order (ie LOD0, LOD1, LOD2, ...). If a triangle budgets are not specified for
+		// a LOD, either by placing -1 in the array or truncating the array, that LOD will be left as-is. 
+		TArray<int32> HardLODBudgets;
+		// enable/disable the Part Promotion LOD strategy (described above)
+		bool bEnableBudgetStrategy_PartLODPromotion = false;
+		// Multiplier on LOD Budgets for PartLODPromotion strategy. This can be used to compensate for hidden-geometry removal and other optimizations done after the strategy is applied.
+		double PartLODPromotionBudgetMultiplier = 2.0;
+
+
+		//
+		// Debug/utility options
+		// 
+
 		// Color mapping modes for vertex colors, primarily used for debugging
 		EVertexColorMappingMode VertexColorMappingMode = EVertexColorMappingMode::None;
 	};
