CombineMeshInstaces:
Renamed FInstanceSet to FSourceInstanceList in CombineMeshInstances API
Added FMeshInstanceGroupData.bPreserveUVs flag, defaults to false, in which case UVs are discarded on simplified and approximate meshes to encourage simplification
Added FOptions::MaxAllowableApproximationDeviation, prevents very large deviation in approximation of highly non-convex parts (may break some approximations, might need per-part overrides in future)
Added FOptions::bDoubleSidedHiddenRemoval, treating as single-sided improves hidden removal
FOptions::ApproximationSourceLOD can now be set to a generated simplified LOD, in addition to Source LODs. This helps avoid approximations being higher-quality than simplified parts

In CombineMeshInstancesImpl, renamed FMeshInstanceSet to FMeshPart, FMeshInstance to FMeshPartInstance, FMeshInstanceAssembly to FMeshPartsAssembly, to help clarify internal vs external data structures (many renaming changes as a result)
FPartApproxSelector now has MaxAllowableDeviation, set from new FOptions member, prevents major shape deviation (maybe too strict? might be good to count how many things deviate...)
tweaked AABB-override of approximations based on volume comparison
reduced InitialTriCost scaling for extra approximate LODs in ComputeMeshApproximations - this value is used as an exponent and so if it gets larger than ~5, for large deviations it will basically be infinity and stops being useful. This may need more work as it tends to go to box approximation very quickly now.
Various improvements to RemoveHiddenFaces_ExteriorVisibility. Now can backface-cull, does so by FOptions default. Improved filtering around "glancing angle" rays/hits. Changed from using fixed exterior point set as visibility locations, to direction set. Add slight downscaling of UV triangles, and of hit-triangle barycentric coords, to filter out hits very near to triangle edges, which can often be spurious.

updated internal usage for renaming, set plausible MaxAllowableApproximationDeviation, and using simplified LODs as ApproximationSourceLOD

diff --git a/Engine/Plugins/Experimental/MeshModelingToolsetExp/Source/GeometryProcessingAdapters/Private/CombineMeshInstancesImpl.cpp b/Engine/Plugins/Experimental/MeshModelingToolsetExp/Source/GeometryProcessingAdapters/Private/CombineMeshInstancesImpl.cpp
--- a/Engine/Plugins/Experimental/MeshModelingToolsetExp/Source/GeometryProcessingAdapters/Private/CombineMeshInstancesImpl.cpp
+++ b/Engine/Plugins/Experimental/MeshModelingToolsetExp/Source/GeometryProcessingAdapters/Private/CombineMeshInstancesImpl.cpp
@@ -25,6 +25,7 @@
 #include "ShapeApproximation/ShapeDetection3.h"
 #include "ShapeApproximation/MeshSimpleShapeApproximation.h"
 #include "Generators/GridBoxMeshGenerator.h"
+#include "Generators/MinimalBoxMeshGenerator.h"
 
 #include "Polygroups/PolygroupsGenerator.h"
 #include "GroupTopology.h"
@@ -83,7 +84,8 @@ enum class EMeshDetailLevel
 };
 
 
-struct FMeshInstance
+// FMeshPartInstance represents a single instance of a FMeshPart
+struct FMeshPartInstance
 {
 	UE::Geometry::FTransformSequence3d WorldTransform;
 	TArray<UMaterialInterface*> Materials;
@@ -93,15 +95,19 @@ struct FMeshInstance
 
 	EMeshDetailLevel DetailLevel = EMeshDetailLevel::Standard;
 
-	// allow FMeshInstance to maintain link to external representation of instance
+	// allow FMeshPartInstance to maintain link to external representation of instance
 	FIndex3i ExternalInstanceIndex = FIndex3i::Invalid();
 };
 
 
-struct FMeshInstanceSet
+// MeshPart represents a set of FMeshPartInstances of a particular FSourceGeometry
+struct FMeshPart
 {
 	UStaticMesh* SourceAsset;
-	TArray<FMeshInstance> Instances;
+	TArray<FMeshPartInstance> Instances;
+
+	bool bPreserveUVs = false;
+	bool bAllowMerging = true;
 };
 
 
@@ -121,16 +127,16 @@ struct FOptimizedGeometry
 };
 
 
-class FMeshInstanceAssembly
+class FMeshPartsAssembly
 {
 public:
 	// this is necessary due to TArray<TUniquePtr> below
-	FMeshInstanceAssembly() = default;
-	FMeshInstanceAssembly(FMeshInstanceAssembly&) = delete;
-	FMeshInstanceAssembly& operator=(const FMeshInstanceAssembly&) = delete;
+	FMeshPartsAssembly() = default;
+	FMeshPartsAssembly(FMeshPartsAssembly&) = delete;
+	FMeshPartsAssembly& operator=(const FMeshPartsAssembly&) = delete;
 
 
-	TArray<TUniquePtr<FMeshInstanceSet>> InstanceSets;
+	TArray<TUniquePtr<FMeshPart>> Parts;
 
 	TArray<UMaterialInterface*> UniqueMaterials;
 	TMap<UMaterialInterface*, int32> MaterialMap;
@@ -141,7 +147,7 @@ public:
 	TArray<FDynamicMeshAABBTree3> SourceMeshSpatials;
 
 	// allow external code to preprocess dynamic mesh for a specific instance
-	TFunction<void(FDynamicMesh3&, const FMeshInstance&)> PreProcessInstanceMeshFunc;
+	TFunction<void(FDynamicMesh3&, const FMeshPartInstance&)> PreProcessInstanceMeshFunc;
 };
 
 
@@ -149,40 +155,43 @@ public:
 
 
 
-void InitializeMeshInstanceAssembly(
-	const IGeometryProcessing_CombineMeshInstances::FInstanceSet& SourceInstanceSet,
-	FMeshInstanceAssembly& AssemblyOut)
+void InitializeMeshPartAssembly(
+	const IGeometryProcessing_CombineMeshInstances::FSourceInstanceList& SourceInstanceList,
+	FMeshPartsAssembly& AssemblyOut)
 {
-	TMap<UStaticMesh*, FMeshInstanceSet*> MeshToInstanceMap;
+	TMap<UStaticMesh*, FMeshPart*> MeshToPartMap;
 
-	int32 NumInstances = SourceInstanceSet.StaticMeshInstances.Num();
+	int32 NumInstances = SourceInstanceList.StaticMeshInstances.Num();
 	for ( int32 Index = 0; Index < NumInstances; ++Index)
 	{
-		const IGeometryProcessing_CombineMeshInstances::FStaticMeshInstance& SourceMeshInstance = SourceInstanceSet.StaticMeshInstances[Index];
+		const IGeometryProcessing_CombineMeshInstances::FStaticMeshInstance& SourceMeshInstance = SourceInstanceList.StaticMeshInstances[Index];
 
 		UStaticMesh* StaticMesh = SourceMeshInstance.SourceMesh;
-		FMeshInstanceSet** FoundInstanceSet = MeshToInstanceMap.Find(StaticMesh);
-		if (FoundInstanceSet == nullptr)
+		FMeshPart** FoundPart = MeshToPartMap.Find(StaticMesh);
+		if (FoundPart == nullptr)
 		{
-			TUniquePtr<FMeshInstanceSet> NewInstanceSet = MakeUnique<FMeshInstanceSet>();
-			NewInstanceSet->SourceAsset = StaticMesh;
-			FMeshInstanceSet* Ptr = NewInstanceSet.Get();
+			TUniquePtr<FMeshPart> NewPart = MakeUnique<FMeshPart>();
+			NewPart->SourceAsset = StaticMesh;
+			FMeshPart* Ptr = NewPart.Get();
 					
-			AssemblyOut.InstanceSets.Add(MoveTemp(NewInstanceSet));
+			AssemblyOut.Parts.Add(MoveTemp(NewPart));
 			// store source model?
 
-			MeshToInstanceMap.Add(StaticMesh, Ptr);
-			FoundInstanceSet = &Ptr;
+			MeshToPartMap.Add(StaticMesh, Ptr);
+			FoundPart = &Ptr;
 		}
 
-		FMeshInstance NewInstance;
+		FMeshPartInstance NewInstance;
 		NewInstance.ExternalInstanceIndex = FIndex3i(Index, -1,-1);
 
-		if ( SourceMeshInstance.GroupDataIndex >= 0 && SourceMeshInstance.GroupDataIndex < SourceInstanceSet.InstanceGroupDatas.Num() )
+		if ( SourceMeshInstance.GroupDataIndex >= 0 && SourceMeshInstance.GroupDataIndex < SourceInstanceList.InstanceGroupDatas.Num() )
 		{
 			const IGeometryProcessing_CombineMeshInstances::FMeshInstanceGroupData& GroupData = 
-				SourceInstanceSet.InstanceGroupDatas[SourceMeshInstance.GroupDataIndex];
+				SourceInstanceList.InstanceGroupDatas[SourceMeshInstance.GroupDataIndex];
 			NewInstance.Materials = GroupData.MaterialSet;
+
+			(*FoundPart)->bPreserveUVs = GroupData.bPreserveUVs;
+			(*FoundPart)->bAllowMerging = GroupData.bAllowMerging;
 		}
 
 		NewInstance.SourceComponent = SourceMeshInstance.SourceComponent;
@@ -192,17 +201,22 @@ void InitializeMeshInstanceAssembly(
 		{
 			NewInstance.WorldTransform.Append( Transform );
 		}
-		(*FoundInstanceSet)->Instances.Add(NewInstance);
+		(*FoundPart)->Instances.Add(NewInstance);
 	}
 
+	// sort parts by largest triangle count first
+	AssemblyOut.Parts.Sort([&](const TUniquePtr<FMeshPart>& A, const TUniquePtr<FMeshPart>& B)
+	{
+		return A->SourceAsset->GetNumTriangles(0) > B->SourceAsset->GetNumTriangles(0);
+	});
 
 	// collect unique materials
-	for (TPair<UStaticMesh*, FMeshInstanceSet*>& Pair : MeshToInstanceMap)
+	for (TPair<UStaticMesh*, FMeshPart*>& Pair : MeshToPartMap)
 	{
 		UStaticMesh* StaticMesh = Pair.Key;
-		FMeshInstanceSet& InstanceSet = *(Pair.Value);
+		FMeshPart& Part = *(Pair.Value);
 
-		for (FMeshInstance& Instance : InstanceSet.Instances)
+		for (FMeshPartInstance& Instance : Part.Instances)
 		{
 			for (UMaterialInterface* Material : Instance.Materials)
 			{
@@ -220,7 +234,7 @@ void InitializeMeshInstanceAssembly(
 
 
 void InitializeAssemblySourceMeshesFromLOD(
-	FMeshInstanceAssembly& Assembly,
+	FMeshPartsAssembly& Assembly,
 	int32 SourceAssetBaseLOD,
 	int32 NumSourceLODs)
 {
@@ -228,17 +242,17 @@ void InitializeAssemblySourceMeshesFromLOD(
 
 	check(NumSourceLODs > 0);
 
-	int32 NumSets = Assembly.InstanceSets.Num();
-	Assembly.SourceMeshGeometry.SetNum(NumSets);
+	int32 NumParts = Assembly.Parts.Num();
+	Assembly.SourceMeshGeometry.SetNum(NumParts);
 
 	// collect mesh for each assembly item
-	ParallelFor(NumSets, [&](int32 Index)
+	ParallelFor(NumParts, [&](int32 Index)
 	{
-		TUniquePtr<FMeshInstanceSet>& InstanceSet = Assembly.InstanceSets[Index];
+		TUniquePtr<FMeshPart>& Part = Assembly.Parts[Index];
 		FSourceGeometry& Target = Assembly.SourceMeshGeometry[Index];
 		Target.SourceMeshLODs.SetNum(NumSourceLODs);
 
-		UStaticMesh* StaticMesh = InstanceSet->SourceAsset;
+		UStaticMesh* StaticMesh = Part->SourceAsset;
 
 		for (int32 k = 0; k < NumSourceLODs; ++k)
 		{
@@ -288,12 +302,12 @@ void InitializeAssemblySourceMeshesFromLOD(
 
 
 	// not clear that it is safe to do this in parallel...
-	for (int32 Index = 0; Index < NumSets; ++Index)
+	for (int32 Index = 0; Index < NumParts; ++Index)
 	{
-		TUniquePtr<FMeshInstanceSet>& InstanceSet = Assembly.InstanceSets[Index];
+		TUniquePtr<FMeshPart>& Part = Assembly.Parts[Index];
 		FSourceGeometry& Target = Assembly.SourceMeshGeometry[Index];
 
-		UStaticMesh* StaticMesh = InstanceSet->SourceAsset;
+		UStaticMesh* StaticMesh = Part->SourceAsset;
 		UBodySetup* BodySetup = StaticMesh->GetBodySetup();
 		if (BodySetup)
 		{
@@ -317,7 +331,9 @@ void InitializeAssemblySourceMeshesFromLOD(
  * @return ( Sqrt(Sum-of-squared-distances) / NumPoints , Max(distance)  )
  * 
  */
-static FVector2d DeviationMetric(const FDynamicMesh3& MeasureMesh, const FDynamicMeshAABBTree3& SourceBVH)
+static FVector2d DeviationMetric(
+	const FDynamicMesh3& MeasureMesh, 
+	const FDynamicMeshAABBTree3& SourceBVH)
 {
 	// todo: could consider normal deviation?
 	int PointCount = 0;
@@ -361,6 +377,7 @@ class FPartApproxSelector
 {
 public:
 	double TriangleCost = 0.7;
+	double MaxAllowableDeviation = 0;		// 0 = disabled
 
 	struct FResultOption
 	{
@@ -408,14 +425,22 @@ public:
 		Option.DeviationMetric = DeviationMetric(*Option.Mesh, *Spatial);
 		int32 TriCount = Option.Mesh->TriangleCount();
 		int32 BaseTriCount = 12;		// 2 tris for each face of box
-		Option.CostMetric = 
-			Option.DeviationMetric[0] * FMathd::Pow( (double)TriCount / (double)BaseTriCount, TriangleCost );
+		if (MaxAllowableDeviation > 0 && Option.DeviationMetric[1] > MaxAllowableDeviation)
+		{
+			Option.CostMetric = TNumericLimits<float>::Max() + (double)Option.MethodID;
+		}
+		else
+		{
+			Option.CostMetric = Option.DeviationMetric[0] * FMathd::Pow((double)TriCount / (double)BaseTriCount, TriangleCost);
+		}
 	}
 
 	void SelectBestOption(
-		FDynamicMesh3& ResultMesh)
+		FDynamicMesh3& ResultMesh, 
+		int32& MethodID)
 	{
 		Options.StableSort( [&](const FResultOption& A, const FResultOption& B) { return A.CostMetric < B.CostMetric; } );
+		MethodID = Options[0].MethodID;
 		ResultMesh = MoveTemp(*Options[0].Mesh);
 	}
 };
@@ -429,14 +454,14 @@ public:
 
 
 
-void InitializeInstanceAssemblySpatials(FMeshInstanceAssembly& Assembly)
+void InitializePartAssemblySpatials(FMeshPartsAssembly& Assembly)
 {
-	int32 NumSets = Assembly.InstanceSets.Num();
-	Assembly.SourceMeshSpatials.SetNum(NumSets);
+	int32 NumParts = Assembly.Parts.Num();
+	Assembly.SourceMeshSpatials.SetNum(NumParts);
 	
-	ParallelFor(NumSets, [&](int32 Index)
+	ParallelFor(NumParts, [&](int32 Index)
 	{
-		TUniquePtr<FMeshInstanceSet>& InstanceSet = Assembly.InstanceSets[Index];
+		TUniquePtr<FMeshPart>& Part = Assembly.Parts[Index];
 		FSourceGeometry& Target = Assembly.SourceMeshGeometry[Index];
 		FDynamicMeshAABBTree3& Spatial = Assembly.SourceMeshSpatials[Index];
 		Spatial.SetMesh(&Target.SourceMeshLODs[0], true);
@@ -450,13 +475,13 @@ void InitializeInstanceAssemblySpatials(FMeshInstanceAssembly& Assembly)
  * approximation would be better, use our metric to make this decision.
  * (todo: this could maybe be folded into simplified-mesh computations...)
  */
-void ReplaceBadSimplifiedLODs(FMeshInstanceAssembly& Assembly)
+void ReplaceBadSimplifiedLODs(FMeshPartsAssembly& Assembly, const IGeometryProcessing_CombineMeshInstances::FOptions& CombineOptions)
 {
-	int32 NumSets = Assembly.InstanceSets.Num();
+	int32 NumParts = Assembly.Parts.Num();
 	
-	ParallelFor(NumSets, [&](int32 Index)
+	ParallelFor(NumParts, [&](int32 Index)
 	{
-		TUniquePtr<FMeshInstanceSet>& InstanceSet = Assembly.InstanceSets[Index];
+		TUniquePtr<FMeshPart>& Part = Assembly.Parts[Index];
 		FSourceGeometry& Target = Assembly.SourceMeshGeometry[Index];
 		FDynamicMeshAABBTree3& Spatial = Assembly.SourceMeshSpatials[Index];
 		FOptimizedGeometry& OptimizedTargets = Assembly.OptimizedMeshGeometry[Index];
@@ -464,6 +489,7 @@ void ReplaceBadSimplifiedLODs(FMeshInstanceAssembly& Assembly)
 		for ( int32 k = OptimizedTargets.SimplifiedMeshLODs.Num()-1; k >= 0; --k )
 		{
 			FPartApproxSelector Selector;
+			Selector.MaxAllowableDeviation = CombineOptions.MaxAllowableApproximationDeviation;
 			Selector.Initialize(Spatial.GetMesh(), &Spatial);
 			if ( k == OptimizedTargets.SimplifiedMeshLODs.Num()-1 )
 			{
@@ -476,7 +502,8 @@ void ReplaceBadSimplifiedLODs(FMeshInstanceAssembly& Assembly)
 			Selector.AddGeneratedMesh(OptimizedTargets.SimplifiedMeshLODs[k], 0);
 
 			// either keep current mesh or replace w/ simplified version
-			Selector.SelectBestOption(OptimizedTargets.SimplifiedMeshLODs[k]);
+			int32 SelectedMethodID = -1;
+			Selector.SelectBestOption(OptimizedTargets.SimplifiedMeshLODs[k], SelectedMethodID);
 		}
 	});
 }
@@ -763,6 +790,9 @@ enum class EApproximatePartMethod : uint8
 
 	AutoBestFit = 10,
 
+	SimplifiedMesh = 55,
+	OverrideAxisBox = 77,
+
 	Original = 100
 
 };
@@ -995,14 +1025,17 @@ static void SelectBestFittingMeshApproximation(
 	const FDynamicMesh3& OriginalMesh, 
 	const FDynamicMeshAABBTree3& OriginalMeshSpatial,
 	FDynamicMesh3& ResultMesh,
+	EApproximatePartMethod& BestMethodIDOut,
 	double AcceptableDeviationTol,
 	double TriangleCost,
+	double MaxDeviation,
 	int32 UseExtrudeAxis = -1  // axis index, or -1 means auto / try all 3
 	)
 {
 	FPartApproxSelector ApproxSelector;
 	ApproxSelector.Initialize(&OriginalMesh, &OriginalMeshSpatial);
 	ApproxSelector.TriangleCost = TriangleCost;
+	ApproxSelector.MaxAllowableDeviation = MaxDeviation;
 
 	ApproxSelector.AddGeneratedMesh( [&](FDynamicMesh3& PartMeshInOut) {
 		ComputeSimplePartApproximation(PartMeshInOut, PartMeshInOut, EApproximatePartMethod::AxisAlignedBox);
@@ -1053,16 +1086,22 @@ static void SelectBestFittingMeshApproximation(
 		}
 	}
 
-	ApproxSelector.SelectBestOption(ResultMesh);
+	int32 SelectedMethodID;
+	ApproxSelector.SelectBestOption(ResultMesh, SelectedMethodID);
+	BestMethodIDOut = static_cast<EApproximatePartMethod>(SelectedMethodID);
 
-	// if Axis-Aligned box volume is less than (100+k%) larger than best option, just use that instead
+	// If Axis-Aligned box volume is less than (100+k%) larger than best option, just use that instead.
+	// Default is 10%, but if approximation is likely to also be a box, double it.
 	// (todo should be configurable)
-	const double BoxPreferenceVolumeRatioPercent = 20.0;
 	FVector2d ApproxMeshVolArea = TMeshQueries<FDynamicMesh3>::GetVolumeArea(ResultMesh);
 	FAxisAlignedBox3d AlignedBox = OriginalMesh.GetBounds(false);
-	if ((AlignedBox.Volume() / ApproxMeshVolArea.X) < (1.0 + BoxPreferenceVolumeRatioPercent/100.0) )
+	double BoxVolume = AlignedBox.Volume();
+	double VolRatio = AlignedBox.Volume() / ApproxMeshVolArea.X;
+	double BoxPreferenceVolumeRatioPercent = (ResultMesh.TriangleCount() <= 12) ? 20.0 : 10.0;
+	if (VolRatio < (1.0 + BoxPreferenceVolumeRatioPercent/100.0) )
 	{
 		ComputeSimplePartApproximation(OriginalMesh, ResultMesh, EApproximatePartMethod::AxisAlignedBox);
+		BestMethodIDOut = EApproximatePartMethod::OverrideAxisBox;
 	}
 
 }
@@ -1076,56 +1115,75 @@ static void SelectBestFittingMeshApproximation(
 
 void ComputeMeshApproximations(
 	IGeometryProcessing_CombineMeshInstances::FOptions CombineOptions,
-	FMeshInstanceAssembly& Assembly)
+	FMeshPartsAssembly& Assembly)
 {
 	using namespace UE::Geometry;
 	const double AngleThresholdDeg = CombineOptions.HardNormalAngleDeg;
 
-	int32 NumSets = Assembly.InstanceSets.Num();
-	Assembly.OptimizedMeshGeometry.SetNum(NumSets);
+	int32 NumParts = Assembly.Parts.Num();
+	Assembly.OptimizedMeshGeometry.SetNum(NumParts);
 
 	int32 NumSimplifiedLODs = CombineOptions.NumSimplifiedLODs;
 	int32 NumApproxLODs = FMath::Max(1, 
 		CombineOptions.NumLODs - CombineOptions.NumCopiedLODs - CombineOptions.NumSimplifiedLODs);
 
 	bool bVerbose = CVarGeometryCombineMeshInstancesVerbose.GetValueOnGameThread();
-	ParallelFor(NumSets, [&](int32 Index)
+	ParallelFor(NumParts, [&](int32 Index)
 	{
-		TUniquePtr<FMeshInstanceSet>& InstanceSet = Assembly.InstanceSets[Index];
+		TUniquePtr<FMeshPart>& Part = Assembly.Parts[Index];
 		FSourceGeometry& SourceGeo = Assembly.SourceMeshGeometry[Index];
 		int32 NumSourceLODs = SourceGeo.SourceMeshLODs.Num();
-		const FDynamicMesh3& OptimizationSourceMesh = (CombineOptions.ApproximationSourceLOD < NumSourceLODs) ?
-			SourceGeo.SourceMeshLODs[CombineOptions.ApproximationSourceLOD] : SourceGeo.SourceMeshLODs.Last();
+		const FDynamicMesh3* OptimizationSourceMesh = &SourceGeo.SourceMeshLODs.Last();
+		if (CombineOptions.ApproximationSourceLOD < NumSourceLODs)
+		{
+			OptimizationSourceMesh = &SourceGeo.SourceMeshLODs[CombineOptions.ApproximationSourceLOD];
+		}
 		FOptimizedGeometry& ApproxGeo = Assembly.OptimizedMeshGeometry[Index];
 
-		FDynamicMeshAABBTree3 OptimizationSourceMeshSpatial(&OptimizationSourceMesh, true);
+		FDynamicMeshAABBTree3 OptimizationSourceMeshSpatial(OptimizationSourceMesh, true);
 
 		// compute simplified part LODs
 		ApproxGeo.SimplifiedMeshLODs.SetNum(NumSimplifiedLODs);
 		double InitialTolerance = CombineOptions.SimplifyBaseTolerance;
 		for (int32 k = 0; k < NumSimplifiedLODs; ++k)
 		{
-			ApproxGeo.SimplifiedMeshLODs[k] = OptimizationSourceMesh;
+			ApproxGeo.SimplifiedMeshLODs[k] = *OptimizationSourceMesh;
 			SimplifyPartMesh(ApproxGeo.SimplifiedMeshLODs[k], InitialTolerance, AngleThresholdDeg,
 				CombineOptions.bSimplifyPreserveCorners, CombineOptions.SimplifySharpEdgeAngleDeg, CombineOptions.SimplifyMinSalientDimension);
 			InitialTolerance *= CombineOptions.SimplifyLODLevelToleranceScale;
 		}
 
+		// Approximation Source LOD may be set to a simplified LOD
+		if (CombineOptions.ApproximationSourceLOD >= NumSourceLODs && CombineOptions.ApproximationSourceLOD < (NumSourceLODs+NumSimplifiedLODs))
+		{
+			OptimizationSourceMesh = &ApproxGeo.SimplifiedMeshLODs[CombineOptions.ApproximationSourceLOD - NumSourceLODs];
+		}
+
 		// Compute shape approximation LODs. 
-		// Note that UseNumApproxLODs is a hack here - we are computing more than necessary
+		// Note that ExtraLODs is a hack here - we are computing more than necessary
 		// so that the cost approximation strategy below has additional simplified approximations available. 
 		// This could be smarter, but this dumb method works OK for now...
-		int32 UseNumApproxLODs = NumApproxLODs + 10;
+		int32 ExtraLODs = 10;
+		int32 UseNumApproxLODs = NumApproxLODs + ExtraLODs;
 		ApproxGeo.ApproximateMeshLODs.SetNum(UseNumApproxLODs);
 		double InitialTriCost = CombineOptions.OptimizeBaseTriCost;
+		TArray<EApproximatePartMethod> SelectedMethodID; SelectedMethodID.SetNum(UseNumApproxLODs);		// useful for debugging
 		for (int32 k = 0; k < UseNumApproxLODs; ++k)
 		{
-			SelectBestFittingMeshApproximation(OptimizationSourceMesh, OptimizationSourceMeshSpatial, 
-				ApproxGeo.ApproximateMeshLODs[k], CombineOptions.SimplifyBaseTolerance, InitialTriCost);
-			InitialTriCost *= CombineOptions.OptimizeLODLevelTriCostScale;
+			SelectBestFittingMeshApproximation(*OptimizationSourceMesh, OptimizationSourceMeshSpatial, 
+				ApproxGeo.ApproximateMeshLODs[k], SelectedMethodID[k], 
+				CombineOptions.SimplifyBaseTolerance, InitialTriCost, CombineOptions.MaxAllowableApproximationDeviation);
+			if (k < NumApproxLODs)
+			{
+				InitialTriCost *= CombineOptions.OptimizeLODLevelTriCostScale;
+			}
+			else
+			{
+				InitialTriCost += 0.25;		// TriCost is used as a power so if it gets too big things go badly...
+			}
 
 			// update enabled attribs (is this good?)
-			ApproxGeo.ApproximateMeshLODs[k].EnableMatchingAttributes(OptimizationSourceMesh);
+			ApproxGeo.ApproximateMeshLODs[k].EnableMatchingAttributes(*OptimizationSourceMesh);
 
 			// recompute normals
 			FMeshNormals::InitializeOverlayTopologyFromOpeningAngle(&ApproxGeo.ApproximateMeshLODs[k], ApproxGeo.ApproximateMeshLODs[k].Attributes()->PrimaryNormals(), AngleThresholdDeg);
@@ -1147,7 +1205,7 @@ void ComputeMeshApproximations(
 
 
 	// try to filter out simplifications that did bad things
-	ReplaceBadSimplifiedLODs(Assembly);
+	ReplaceBadSimplifiedLODs(Assembly, CombineOptions);
 
 
 	// Now that we have our per-part LOD stacks, we can estimate total triangle count that will be
@@ -1164,7 +1222,7 @@ void ComputeMeshApproximations(
 		struct FPartCostInfo
 		{
 			// these fields are precomputed
-			int32 PartIndex = 0;			// index into Assembly.InstanceSets
+			int32 PartIndex = 0;			// index into Assembly.Parts
 			int32 NumInstances = 0;			// number of instances of this part in the final mesh, currently *excluding* decorative parts
 			double ReplacedWeight = 1.0;	// reduce this as we shift LODs up, to try to let other comparable parts take the hit...
 
@@ -1181,16 +1239,16 @@ void ComputeMeshApproximations(
 			}
 		};
 		TArray<FPartCostInfo> CostInfo;
-		CostInfo.SetNum(NumSets);
+		CostInfo.SetNum(NumParts);
 
 		// initialize precomputed parts of the CostInfo array that we will incrementally update
-		for (int32 SetIndex = 0; SetIndex < NumSets; ++SetIndex)
+		for (int32 SetIndex = 0; SetIndex < NumParts; ++SetIndex)
 		{
-			const TUniquePtr<FMeshInstanceSet>& InstanceSet = Assembly.InstanceSets[SetIndex];
+			const TUniquePtr<FMeshPart>& Part = Assembly.Parts[SetIndex];
 			CostInfo[SetIndex].PartIndex = SetIndex;
 
 			CostInfo[SetIndex].NumInstances = 0;
-			for (const FMeshInstance& Instance : InstanceSet->Instances)
+			for (const FMeshPartInstance& Instance : Part->Instances)
 			{
 				bool bSkipInstance = (Instance.DetailLevel == EMeshDetailLevel::Decorative);
 				if (bSkipInstance == false)
@@ -1228,7 +1286,7 @@ void ComputeMeshApproximations(
 			{
 				// compute current estimate of total part count for this LOD
 				int32 TotalCurLODTriCount = 0;
-				for (int32 SetIndex = 0; SetIndex < NumSets; ++SetIndex)
+				for (int32 SetIndex = 0; SetIndex < NumParts; ++SetIndex)
 				{
 					const FDynamicMesh3* CurSourcePartMeshLOD = CostInfo[SetIndex].LODChainMeshes[LODIndex];
 					CostInfo[SetIndex].PartTriCount = CurSourcePartMeshLOD->TriangleCount();
@@ -1265,7 +1323,7 @@ void ComputeMeshApproximations(
 
 				// find part with largest current cost
 				int32 MaxSetIndex = 0;
-				for (int32 k = 1; k < NumSets; ++k)
+				for (int32 k = 1; k < NumParts; ++k)
 				{
 					double MaxCost = CostInfo[MaxSetIndex].PartCostWeight();
 					double CurCost = CostInfo[k].PartCostWeight();
@@ -1296,7 +1354,7 @@ void ComputeMeshApproximations(
 					ReplaceInfo.ReplacedWeight *= 0.5;
 				}
 				// slowly increase weights of parts   (should this be modulated by tri count?)
-				for (int32 k = 0; k < NumSets; ++k)
+				for (int32 k = 0; k < NumParts; ++k)
 				{
 					CostInfo[k].ReplacedWeight += 0.1;
 				}
@@ -1351,22 +1409,34 @@ static void RemoveHiddenFaces_Occlusion(FDynamicMesh3& EditMesh, double MaxDista
 
 // Remove hidden faces by casting rays from exterior at sample points on triangles
 // (This method works quite well and should eventually be extracted out to a general algorithm...)
-static void RemoveHiddenFaces_ExteriorVisibility(FDynamicMesh3& TargetMesh, double SampleRadius)
+static void RemoveHiddenFaces_ExteriorVisibility(
+	FDynamicMesh3& TargetMesh, 
+	double SampleRadius, 
+	bool bDoubleSided,
+	int32 LODIndex)
 {
 	FDynamicMeshAABBTree3 Spatial(&TargetMesh, true);
 	FAxisAlignedBox3d Bounds = Spatial.GetBoundingBox();
-	double Radius = Bounds.DiagonalLength() * 0.5;
+	double Radius = Bounds.DiagonalLength();
 
+	// geometric magic numbers used below that have been slightly tuned...
+	double GlancingAngleDotTolerance = FMathd::Cos(85.0 * FMathd::DegToRad);
+	const double TriScalingAlpha = 0.95;
+	const double BaryCoordsThreshold = 0.01;
 
-	auto FindHitTriangleTest = [&](FVector3d TargetPosition, FVector3d FarPosition) -> int
+
+	auto FindHitTriangleTest = [&](FVector3d TargetPosition, FVector3d TargetNormal, FVector3d FarPosition) -> int
 	{
 		FVector3d RayDir(TargetPosition - FarPosition);
 		double Distance = Normalize(RayDir);
+		if (bDoubleSided == false && RayDir.Dot(TargetNormal) > -0.001)
+		{
+			return IndexConstants::InvalidID;
+		}
 		FRay3d Ray(FarPosition, RayDir, true);
 		return Spatial.FindNearestHitTriangle(Ray, IMeshSpatial::FQueryOptions(Distance + 1.0));	// 1.0 is random fudge factor here...
 	};
 
-
 	// final triangle visibility, atomics can be updated on any thread
 	TArray<std::atomic<bool>> ThreadSafeTriVisible;
 	ThreadSafeTriVisible.SetNum(TargetMesh.MaxTriangleID());
@@ -1409,9 +1479,9 @@ static void RemoveHiddenFaces_ExteriorVisibility(FDynamicMesh3& TargetMesh, doub
 			// worse, on rectilinear shapes (eg imagine some stacked cubes or adjacent parts)
 			// the ray can get "through" the cracks between adjacent connected triangles
 			// and manage to hit the search triangle
-			if ( FMathd::Abs(Direction.Dot(Normal)) > 0.01 )
+			if ( FMathd::Abs(Direction.Dot(Normal)) > GlancingAngleDotTolerance )
 			{
-				if (FindHitTriangleTest(Centroid, Centroid + Radius*Direction) == tid)
+				if (FindHitTriangleTest(Centroid, Normal, Centroid + Radius*Direction) == tid)
 				{
 					ThreadSafeTriVisible[tid] = true;
 					TriStatusKnown[tid] = true;
@@ -1419,7 +1489,7 @@ static void RemoveHiddenFaces_ExteriorVisibility(FDynamicMesh3& TargetMesh, doub
 				}
 			}
 		}
-		if (FindHitTriangleTest(Centroid, Centroid + Radius*Normal) == tid)
+		if (FindHitTriangleTest(Centroid, Normal, Centroid + Radius*Normal) == tid)
 		{
 			ThreadSafeTriVisible[tid] = true;
 			TriStatusKnown[tid] = true;
@@ -1430,8 +1500,8 @@ static void RemoveHiddenFaces_ExteriorVisibility(FDynamicMesh3& TargetMesh, doub
 	});
 
 	//
-	// Construct set of exterior sample points, for each triangle sample point
-	// below we will check if it is visible from any of these sample points.
+	// Construct set of exterior sample directions, for each triangle sample point
+	// below we will check if it is visible from any of these sample directions.
 	// Order is shuffled in hopes that for visible tris we don't waste a bunch
 	// of time on the 'far' side
 	//
@@ -1442,20 +1512,13 @@ static void RemoveHiddenFaces_ExteriorVisibility(FDynamicMesh3& TargetMesh, doub
 	uint32 SampleIndex = 0;
 	while (ModuloIter.GetNextIndex(SampleIndex))
 	{
-		ExteriorSamplePoints.Add( Bounds.Center() + Radius * SphereSampler[SampleIndex] );
+		ExteriorSamplePoints.Add( Normalized(SphereSampler[SampleIndex]) );
 	}
-	// add axis directions?
-
 
 	//
 	// For each triangle, generate a set of sample points on the triangle surface,
 	// and then check if that point is visible from any of the exterior sample points.
 	// This is the expensive part!
-	// 
-	// Does using a fixed set of exterior sample points make sense? Could also 
-	// treat it as a set of sample directions. Seems more likely to hit tri
-	// based on sample directions...
-	//
 	ParallelFor(TargetMesh.MaxTriangleID(), [&](int32 tid)
 	{
 		// if we already found out this triangle is visible or hidden, we can skip it
@@ -1469,6 +1532,15 @@ static void RemoveHiddenFaces_ExteriorVisibility(FDynamicMesh3& TargetMesh, doub
 
 		FFrame3d TriFrame(Centroid, TriNormal);
 		FTriangle2d UVTriangle(TriFrame.ToPlaneUV(A), TriFrame.ToPlaneUV(B), TriFrame.ToPlaneUV(C));
+
+		// Slightly shrink the triangle, this helps to avoid spurious hits
+		// TODO obviously should scale by an actual dimension and not just a relative %...
+		FVector2d Center = (UVTriangle.V[0] + UVTriangle.V[1] + UVTriangle.V[2]) / 3.0;
+		for (int32 k = 0; k < 3; ++k)
+		{
+			UVTriangle.V[k] = (1- TriScalingAlpha)*Center + (TriScalingAlpha)*UVTriangle.V[k];
+		}
+
 		double DiscArea = (FMathd::Pi * SampleRadius * SampleRadius);
 		int NumSamples = FMath::Max( (int)(TriArea / DiscArea), 2 );  // a bit arbitrary...
 		FVector2d V1 = UVTriangle.V[1] - UVTriangle.V[0];
@@ -1476,9 +1548,9 @@ static void RemoveHiddenFaces_ExteriorVisibility(FDynamicMesh3& TargetMesh, doub
 
 		TArray<int32> HitTris;		// re-use this array in inner loop to avoid hitting atomics so often
 
-		int NumTested = 0;
+		int NumTested = 0; int Iterations = 0;
 		FRandomStream RandomStream(tid);
-		while (NumTested < NumSamples)
+		while (NumTested < NumSamples && Iterations++ < 10000)
 		{
 			double a1 = RandomStream.GetFraction();
 			double a2 = RandomStream.GetFraction();
@@ -1492,13 +1564,30 @@ static void RemoveHiddenFaces_ExteriorVisibility(FDynamicMesh3& TargetMesh, doub
 				HitTris.Reset();
 				for (int32 k = 0; k < NumExteriorSamplePoints; ++k)
 				{
-					int32 HitTriID = FindHitTriangleTest(Position, ExteriorSamplePoints[k]);
+					FVector3d Direction = ExteriorSamplePoints[k];
+					if (FMathd::Abs(Direction.Dot(TriNormal)) < GlancingAngleDotTolerance)
+					{
+						continue;
+					}
+
+					FVector3d RayFrom = Position + 2.0 * Radius * ExteriorSamplePoints[k];
+					int32 HitTriID = FindHitTriangleTest(Position, TriNormal, RayFrom);
 					if ( HitTriID != IndexConstants::InvalidID && TriStatusKnown[HitTriID] == false )
 					{
-						HitTris.AddUnique(HitTriID);		// we hit some triangle, whether or not it is the one we are testing...
-						if (HitTriID == tid)
+						// Want to filter out on-edge triangle hits, as they are generally spurious and
+						// will result in interior triangles remaining visible
+						FRay3d Ray(RayFrom, Normalized(Position-RayFrom), true);
+						FIntrRay3Triangle3d RayHit = TMeshQueries<FDynamicMesh3>::RayTriangleIntersection(TargetMesh, HitTriID, Ray);
+						if (RayHit.IntersectionType == EIntersectionType::Point &&
+							RayHit.TriangleBaryCoords.GetMin() > BaryCoordsThreshold &&
+							RayHit.TriangleBaryCoords.GetMax() < (1.0-BaryCoordsThreshold) )
 						{
-							break;
+
+							HitTris.AddUnique(HitTriID);		// we hit some triangle, whether or not it is the one we are testing...
+							if (HitTriID == tid)
+							{
+								break;
+							}
 						}
 					}
 				}
@@ -1660,7 +1749,7 @@ static void PostProcessHiddenFaceRemovedMesh(
 		Simplifier.GroupBoundaryConstraint = EEdgeRefineFlags::NoConstraint;
 		Simplifier.MaterialBoundaryConstraint = EEdgeRefineFlags::NoConstraint;
 
-		Simplifier.CollapseMode = FQEMSimplification::ESimplificationCollapseModes::MinimalExistingVertexError;
+		Simplifier.CollapseMode = FQEMSimplification::ESimplificationCollapseModes::AverageVertexPosition;
 		Simplifier.SimplifyToMinimalPlanar( 0.01 );
 	}
 
@@ -2107,7 +2196,7 @@ void ComputeHiddenRemovalForLOD(
 			break;
 		case IGeometryProcessing_CombineMeshInstances::ERemoveHiddenFacesMode::ExteriorVisibility:
 		case IGeometryProcessing_CombineMeshInstances::ERemoveHiddenFacesMode::Fastest:
-			RemoveHiddenFaces_ExteriorVisibility(MeshLOD, CombineOptions.RemoveHiddenSamplingDensity);
+			RemoveHiddenFaces_ExteriorVisibility(MeshLOD, CombineOptions.RemoveHiddenSamplingDensity, CombineOptions.bDoubleSidedHiddenRemoval, LODIndex);
 			bModified = true;
 			break;
 	}
@@ -2137,7 +2226,7 @@ void ComputeHiddenRemovalForLOD(
 // should we even bother w/ storing approx meshes? just generate them as needed?
 
 void BuildCombinedMesh(
-	const FMeshInstanceAssembly& Assembly,
+	const FMeshPartsAssembly& Assembly,
 	IGeometryProcessing_CombineMeshInstances::FOptions CombineOptions,
 	TArray<FDynamicMesh3>& CombinedMeshLODs)
 {
@@ -2175,19 +2264,21 @@ void BuildCombinedMesh(
 	//FDynamicMeshPolygroupAttribute* PartIDAttrib = AccumMesh.Attributes()->GetPolygroupLayer(0);
 	//FDynamicMeshPolygroupAttribute* PartInstanceMapAttrib = AccumMesh.Attributes()->GetPolygroupLayer(1);
 
-	int32 NumSets = Assembly.InstanceSets.Num();
+	int32 NumParts = Assembly.Parts.Num();
 
-	//for ( int32 SetIndex = 0; SetIndex < NumSets; ++SetIndex )
+	//for ( int32 SetIndex = 0; SetIndex < NumParts; ++SetIndex )
 	//{
 	//	CombinedLOD0.EnableMatchingAttributes( Assembly.SourceMeshGeometry[Index].OriginalMesh, false, false );
 	//}
 
-	for ( int32 SetIndex = 0; SetIndex < NumSets; ++SetIndex )
+	for ( int32 SetIndex = 0; SetIndex < NumParts; ++SetIndex )
 	{
-		const TUniquePtr<FMeshInstanceSet>& InstanceSet = Assembly.InstanceSets[SetIndex];
+		const TUniquePtr<FMeshPart>& Part = Assembly.Parts[SetIndex];
 		const FSourceGeometry& SourceGeometry = Assembly.SourceMeshGeometry[SetIndex];
 		const FOptimizedGeometry& OptimizedGeometry = Assembly.OptimizedMeshGeometry[SetIndex];
-		UStaticMesh* StaticMesh = InstanceSet->SourceAsset;
+		UStaticMesh* StaticMesh = Part->SourceAsset;
+
+		check(Part->bAllowMerging == true);		// do not support this yet
 
 		FMeshIndexMappings Mappings;
 
@@ -2228,7 +2319,7 @@ void BuildCombinedMesh(
 
 			FCombinedMeshLOD& CombinedMeshLODData = MeshLODs[LODLevel];
 
-			for ( const FMeshInstance& Instance : InstanceSet->Instances )
+			for ( const FMeshPartInstance& Instance : Part->Instances )
 			{
 				bool bIsDecorativePart = (Instance.DetailLevel == EMeshDetailLevel::Decorative);
 
@@ -2253,6 +2344,17 @@ void BuildCombinedMesh(
 					Assembly.PreProcessInstanceMeshFunc(TempAppendMesh, Instance);
 				}
 
+				// if part does not require UVs, get rid of them here to encourage merging downstream
+				if (LODLevel > 0 && Part->bPreserveUVs == false && TempAppendMesh.HasAttributes())
+				{
+					for (int32 UVLayer = 0; UVLayer < TempAppendMesh.Attributes()->NumUVLayers(); ++UVLayer)
+					{
+						TempAppendMesh.Attributes()->GetUVLayer(UVLayer)->ClearElements();
+					}
+					// if we have no UVs then tangents are invalid
+					TempAppendMesh.Attributes()->DisableTangents();
+				}
+
 				Mappings.Reset();
 				CombinedMeshLODData.Editor.AppendMesh(&TempAppendMesh, Mappings,
 					[&](int, const FVector3d& Pos) { return Instance.WorldTransform.TransformPosition(Pos); },
@@ -2452,46 +2554,6 @@ void BuildCombinedMesh(
 
 
 
-/**
- * Construct a new OrientedBox that contains both A and B. The main problem is to
- * determine the new Orientation, this is done by a 0.5 slerp of the orientations of A and B.
- * The new local Origin and Extents are then computed in this new orientation.
- */
-FOrientedBox3d MergeBoxes(const FOrientedBox3d& A, const FOrientedBox3d& B)
-{
-	FOrientedBox3d NewBox;
-	NewBox.Frame.Origin = (A.Center() + B.Center()) * 0.5;
-
-	FQuaterniond RotationA(A.Frame.Rotation), RotationB(B.Frame.Rotation);
-	if (RotationA.Dot(RotationB) < 0)
-	{
-		RotationB = -RotationB;
-	}
-
-	// this is just a slerp?
-	FQuaterniond HalfRotation = RotationA + RotationB;
-	HalfRotation.Normalize();
-	NewBox.Frame.Rotation = HalfRotation;
-
-	// likely faster to compute the frame X/Y/Z instead of calling ToFramePoint each time...
-	FAxisAlignedBox3d LocalBounds(FVector3d::Zero(), FVector3d::Zero());
-	A.EnumerateCorners([&](FVector3d P)
-	{
-		LocalBounds.Contain( NewBox.Frame.ToFramePoint(P) );
-	});
-	B.EnumerateCorners([&](FVector3d P)
-	{
-		LocalBounds.Contain( NewBox.Frame.ToFramePoint(P) );
-	});
-
-	// update origin and extents
-	NewBox.Frame.Origin = NewBox.Frame.FromFramePoint( LocalBounds.Center() );
-	NewBox.Extents = 0.5 * LocalBounds.Diagonal();
-
-	return NewBox;
-}
-
-
 
 static void CombineCollisionShapes(
 	FSimpleShapeSet3d& CollisionShapes,
@@ -2546,7 +2608,7 @@ static void CombineCollisionShapes(
 				FOrientedBox3d Box2 = Boxes[j];
 
 				// should we just be appending box2 to Box1? prevents getting skewed boxes...
-				FOrientedBox3d NewBox = MergeBoxes(Box1, Box2);
+				FOrientedBox3d NewBox = Box1.Merge(Box2);
 
 				// check if newbox is still aligned w/ box2?
 				bool bAllAxesAligned = true;
@@ -2582,17 +2644,17 @@ static void CombineCollisionShapes(
 
 
 void BuildCombinedCollisionShapes(
-	const FMeshInstanceAssembly& Assembly,
+	const FMeshPartsAssembly& Assembly,
 	IGeometryProcessing_CombineMeshInstances::FOptions CombineOptions,
 	FSimpleShapeSet3d& CombinedCollisionShapes)
 {
-	int32 NumSets = Assembly.InstanceSets.Num();
+	int32 NumParts = Assembly.Parts.Num();
 
-	for ( int32 SetIndex = 0; SetIndex < NumSets; ++SetIndex )
+	for ( int32 SetIndex = 0; SetIndex < NumParts; ++SetIndex )
 	{
-		const TUniquePtr<FMeshInstanceSet>& InstanceSet = Assembly.InstanceSets[SetIndex];
+		const TUniquePtr<FMeshPart>& Part = Assembly.Parts[SetIndex];
 		const FSourceGeometry& SourceGeometry = Assembly.SourceMeshGeometry[SetIndex];
-		for ( const FMeshInstance& Instance : InstanceSet->Instances )
+		for ( const FMeshPartInstance& Instance : Part->Instances )
 		{
 			bool bIsDecorativePart = (Instance.DetailLevel == EMeshDetailLevel::Decorative);
 			if ( ! bIsDecorativePart )
@@ -2669,7 +2731,7 @@ static void SetConstantVertexColor(FDynamicMesh3& Mesh, FLinearColor LinearColor
 
 
 void FCombineMeshInstancesImpl::CombineMeshInstances(
-	const FInstanceSet& MeshInstances, const FOptions& Options, FResults& ResultsOut)
+	const FSourceInstanceList& MeshInstances, const FOptions& Options, FResults& ResultsOut)
 {
 	TRACE_CPUPROFILER_EVENT_SCOPE(CombineMeshInstances);
 
@@ -2682,28 +2744,28 @@ void FCombineMeshInstancesImpl::CombineMeshInstances(
 			Options.NumLODs, Options.NumCopiedLODs, Options.NumSimplifiedLODs, NumApproxLODs, Options.NumVoxWrapLODs);
 	}
 
-	FMeshInstanceAssembly InstanceAssembly;
+	FMeshPartsAssembly PartAssembly;
 
 	{
 		TRACE_CPUPROFILER_EVENT_SCOPE(CombineMeshInst_Setup);
-		InitializeMeshInstanceAssembly(MeshInstances, InstanceAssembly);
+		InitializeMeshPartAssembly(MeshInstances, PartAssembly);
 		if (bVerbose)
 		{
-			UE_LOG(LogGeometry, Log, TEXT("  InstanceAssembly contains %d InstanceSets, %d Unique Materials"), 
-				InstanceAssembly.InstanceSets.Num(), InstanceAssembly.UniqueMaterials.Num());
+			UE_LOG(LogGeometry, Log, TEXT("  PartAssembly contains %d Parts, %d Unique Materials"), 
+				PartAssembly.Parts.Num(), PartAssembly.UniqueMaterials.Num());
 		}
 
-		InitializeAssemblySourceMeshesFromLOD(InstanceAssembly, 0, Options.NumCopiedLODs);
-		InitializeInstanceAssemblySpatials(InstanceAssembly);
+		InitializeAssemblySourceMeshesFromLOD(PartAssembly, 0, Options.NumCopiedLODs);
+		InitializePartAssemblySpatials(PartAssembly);
 	}
 	{
 		TRACE_CPUPROFILER_EVENT_SCOPE(CombineMeshInst_PartApprox);
-		ComputeMeshApproximations(Options, InstanceAssembly);
+		ComputeMeshApproximations(Options, PartAssembly);
 	}
 
 
 	
-	InstanceAssembly.PreProcessInstanceMeshFunc = [&InstanceAssembly, &MeshInstances, &Options](FDynamicMesh3& AppendMesh, const FMeshInstance& Instance)
+	PartAssembly.PreProcessInstanceMeshFunc = [&PartAssembly, &MeshInstances, &Options](FDynamicMesh3& AppendMesh, const FMeshPartInstance& Instance)
 	{
 		int32 SourceInstance = Instance.ExternalInstanceIndex[0];
 		int GroupDataIdx = MeshInstances.StaticMeshInstances[SourceInstance].GroupDataIndex;
@@ -2732,13 +2794,13 @@ void FCombineMeshInstancesImpl::CombineMeshInstances(
 	TArray<UE::Geometry::FDynamicMesh3> CombinedMeshLODs;
 	{
 		TRACE_CPUPROFILER_EVENT_SCOPE(CombineMeshInst_BuildMeshes);
-		BuildCombinedMesh(InstanceAssembly, Options, CombinedMeshLODs);
+		BuildCombinedMesh(PartAssembly, Options, CombinedMeshLODs);
 	}
 
 	FSimpleShapeSet3d CombinedCollisionShapes;
 	{
 		TRACE_CPUPROFILER_EVENT_SCOPE(CombineMeshInst_BuildCollision);
-		BuildCombinedCollisionShapes(InstanceAssembly, Options, CombinedCollisionShapes);
+		BuildCombinedCollisionShapes(PartAssembly, Options, CombinedCollisionShapes);
 
 		if (bVerbose)
 		{
@@ -2754,6 +2816,6 @@ void FCombineMeshInstancesImpl::CombineMeshInstances(
 	ResultsOut.CombinedMeshes.SetNum(1);
 	IGeometryProcessing_CombineMeshInstances::FOutputMesh& OutputMesh = ResultsOut.CombinedMeshes[0];
 	OutputMesh.MeshLODs = MoveTemp(CombinedMeshLODs);
-	OutputMesh.MaterialSet = InstanceAssembly.UniqueMaterials;
+	OutputMesh.MaterialSet = PartAssembly.UniqueMaterials;
 	OutputMesh.SimpleCollisionShapes = PhysicsData.AggGeom;
 }
\ No newline at end of file
diff --git a/Engine/Plugins/Experimental/MeshModelingToolsetExp/Source/GeometryProcessingAdapters/Public/GeometryProcessing/CombineMeshInstancesImpl.h b/Engine/Plugins/Experimental/MeshModelingToolsetExp/Source/GeometryProcessingAdapters/Public/GeometryProcessing/CombineMeshInstancesImpl.h
--- a/Engine/Plugins/Experimental/MeshModelingToolsetExp/Source/GeometryProcessingAdapters/Public/GeometryProcessing/CombineMeshInstancesImpl.h
+++ b/Engine/Plugins/Experimental/MeshModelingToolsetExp/Source/GeometryProcessingAdapters/Public/GeometryProcessing/CombineMeshInstancesImpl.h
@@ -22,7 +22,7 @@ class GEOMETRYPROCESSINGADAPTERS_API FCombineMeshInstancesImpl : public IGeometr
 public:
 	virtual FOptions ConstructDefaultOptions() override;
 
-	virtual void CombineMeshInstances(const FInstanceSet& MeshInstances, const FOptions& Options, FResults& ResultsOut) override;
+	virtual void CombineMeshInstances(const FSourceInstanceList& MeshInstances, const FOptions& Options, FResults& ResultsOut) override;
 
 protected:
 
diff --git a/Engine/Source/Developer/GeometryProcessingInterfaces/Public/GeometryProcessingInterfaces/CombineMeshInstances.h b/Engine/Source/Developer/GeometryProcessingInterfaces/Public/GeometryProcessingInterfaces/CombineMeshInstances.h
--- a/Engine/Source/Developer/GeometryProcessingInterfaces/Public/GeometryProcessingInterfaces/CombineMeshInstances.h
+++ b/Engine/Source/Developer/GeometryProcessingInterfaces/Public/GeometryProcessingInterfaces/CombineMeshInstances.h
@@ -33,12 +33,19 @@ public:
 	};
 
 
+	/**
+	 * FMeshInstanceGroupData is data shared among one or more FBaseMeshInstances.
+	 * For example all instances in an ISMC can share a single FMeshInstanceGroupData.
+	 */
 	struct FMeshInstanceGroupData
 	{
 		TArray<UMaterialInterface*> MaterialSet;
 		
 		bool bHasConstantOverrideVertexColor = false;
 		FColor OverrideVertexColor;
+
+		bool bPreserveUVs = false;
+		bool bAllowMerging = true;		// if false, cannot merge the geometry from this mesh with adjacent meshes to reduce triangle count
 	};
 
 
@@ -46,19 +53,24 @@ public:
 	{
 		EMeshDetailLevel DetailLevel = EMeshDetailLevel::Standard;
 		TArray<FTransform3d> TransformSequence;
-		int32 GroupDataIndex = -1;		// index into FInstanceSet::InstanceGroupDatas
+		int32 GroupDataIndex = -1;		// index into FSourceInstanceList::InstanceGroupDatas
 	};
 
-
+	// a single instance of a static mesh asset
 	struct FStaticMeshInstance : FBaseMeshInstance
 	{
 		UStaticMesh* SourceMesh = nullptr;
 		UPrimitiveComponent* SourceComponent = nullptr;
-		int32 SourceInstanceIndex = 0;
+		int32 SourceInstanceIndex = 0;		// custom index, eg if SourceComponent is an InstancedStaticMeshComponent, this is the Instance index
 	};
 
-
-	struct FInstanceSet
+	/**
+	 * FSourceInstanceList provides a flattened list of mesh instances to CombineMeshInstances.
+	 * Each Instance can refer to shared data in the InstanceGroupDatas list
+	 * 
+	 * This data structure may be replaced in future w/ something more structured
+	 */
+	struct FSourceInstanceList
 	{
 		TArray<FStaticMeshInstance> StaticMeshInstances;
 
@@ -88,10 +100,12 @@ public:
 		// number of requested LODs
 		int32 NumLODs = 5;
 
+		// settings for Copied LODs, these are directly copied from Source Geometry
 		int32 NumCopiedLODs = 1;
 
-		int32 ApproximationSourceLOD = 0;
-
+		//
+		// Settings for Simplifed LODs
+		//
 		int32 NumSimplifiedLODs = 3;
 		double SimplifyBaseTolerance = 1.0;
 		double SimplifyLODLevelToleranceScale = 2.0;
@@ -99,8 +113,19 @@ public:
 		double SimplifySharpEdgeAngleDeg = 44.0;
 		double SimplifyMinSalientDimension = 1.0;
 
+		//
+		// settings for Approximate LODs
+		//
+		
+		// which LOD to use as basis for Approximations. Index is interpreted relative to [CopiedLODs...SimplifiedLODs] set, ie can be set to a Simplified LOD
+		int32 ApproximationSourceLOD = 0;
 		double OptimizeBaseTriCost = 0.7;
 		double OptimizeLODLevelTriCostScale = 1.5;
+		double MaxAllowableApproximationDeviation = 5.0;
+
+		//
+		// settings for Voxel LODs
+		//
 
 		int32 NumVoxWrapLODs = 1;
 		double VoxWrapBaseTolerance = 1.0;
@@ -116,7 +141,8 @@ public:
 		int32 RemoveHiddenStartLOD = 0;
 		// (approximately) spacing between samples on triangle faces used for determining exterior visibility
 		double RemoveHiddenSamplingDensity = 1.0;
-
+		// treat faces as double-sided for hidden removal
+		bool bDoubleSidedHiddenRemoval = false;
 
 		// LOD level to filter out detail parts
 		int32 FilterDecorativePartsLODLevel = 2;
@@ -196,7 +222,7 @@ public:
 
 
 	virtual void CombineMeshInstances(
-		const FInstanceSet& MeshInstances, 
+		const FSourceInstanceList& MeshInstances,
 		const FOptions& Options, 
 		FResults& ResultsOut) 
 	{
