Render Capture Baking: Make BakeRC tool output a new material configured with the generated textures

This change also extends UModelingObjectsCreationAPI to support materials by duplicating a given material

diff --git a/Engine/Plugins/Editor/ModelingToolsEditorMode/Source/ModelingToolsEditorMode/Private/ModelingToolsEditorMode.cpp b/Engine/Plugins/Editor/ModelingToolsEditorMode/Source/ModelingToolsEditorMode/Private/ModelingToolsEditorMode.cpp
--- a/Engine/Plugins/Editor/ModelingToolsEditorMode/Source/ModelingToolsEditorMode/Private/ModelingToolsEditorMode.cpp
+++ b/Engine/Plugins/Editor/ModelingToolsEditorMode/Source/ModelingToolsEditorMode/Private/ModelingToolsEditorMode.cpp
@@ -521,6 +521,13 @@ void UModelingToolsEditorMode::Enter()
 				UE::Modeling::OnNewAssetCreated(CreatedInfo.NewAsset);
 			}
 		});
+		MaterialCreatedEventHandle = ModelCreationAPI->OnModelingMaterialCreated.AddLambda([](const FCreateMaterialObjectResult& CreatedInfo)
+		{
+			if (CreatedInfo.NewAsset != nullptr)
+			{
+				UE::Modeling::OnNewAssetCreated(CreatedInfo.NewAsset);
+			}
+		});
 	}
 
 	const FModelingToolsManagerCommands& ToolManagerCommands = FModelingToolsManagerCommands::Get();
@@ -1217,6 +1224,7 @@ void UModelingToolsEditorMode::Exit()
 		ObjectCreationAPI->GetNewAssetPathNameCallback.Unbind();
 		ObjectCreationAPI->OnModelingMeshCreated.Remove(MeshCreatedEventHandle);
 		ObjectCreationAPI->OnModelingTextureCreated.Remove(TextureCreatedEventHandle);
+		ObjectCreationAPI->OnModelingMaterialCreated.Remove(MaterialCreatedEventHandle);
 		//UEditorModelingObjectsCreationAPI::Deregister(ToolsContext.Get());		// cannot do currently because of shared ToolsContext, revisit in future
 	}
 
diff --git a/Engine/Plugins/Editor/ModelingToolsEditorMode/Source/ModelingToolsEditorMode/Public/ModelingToolsEditorMode.h b/Engine/Plugins/Editor/ModelingToolsEditorMode/Source/ModelingToolsEditorMode/Public/ModelingToolsEditorMode.h
--- a/Engine/Plugins/Editor/ModelingToolsEditorMode/Source/ModelingToolsEditorMode/Public/ModelingToolsEditorMode.h
+++ b/Engine/Plugins/Editor/ModelingToolsEditorMode/Source/ModelingToolsEditorMode/Public/ModelingToolsEditorMode.h
@@ -108,6 +108,7 @@ protected:
 
 	FDelegateHandle MeshCreatedEventHandle;
 	FDelegateHandle TextureCreatedEventHandle;
+	FDelegateHandle MaterialCreatedEventHandle;
 	FDelegateHandle SelectionModifiedEventHandle;
 
 	FDelegateHandle EditorClosedEventHandle;
diff --git a/Engine/Plugins/Experimental/MeshModelingToolsetExp/Source/MeshModelingToolsExp/Private/BakeRenderCaptureTool.cpp b/Engine/Plugins/Experimental/MeshModelingToolsetExp/Source/MeshModelingToolsExp/Private/BakeRenderCaptureTool.cpp
--- a/Engine/Plugins/Experimental/MeshModelingToolsetExp/Source/MeshModelingToolsExp/Private/BakeRenderCaptureTool.cpp
+++ b/Engine/Plugins/Experimental/MeshModelingToolsetExp/Source/MeshModelingToolsExp/Private/BakeRenderCaptureTool.cpp
@@ -8,6 +8,7 @@
 #include "TargetInterfaces/PrimitiveComponentBackedTarget.h"
 #include "TargetInterfaces/StaticMeshBackedTarget.h"
 #include "ToolTargetManager.h"
+#include "InteractiveToolManager.h"
 
 #include "DynamicMesh/MeshTransforms.h"
 
@@ -22,7 +23,7 @@
 #include "Baking/BakingTypes.h"
 #include "Sampling/MeshImageBakingCache.h"
 #include "Sampling/MeshMapBaker.h"
-#include "Sampling/RenderCaptureMapEvaluator.h"
+#include "AssetUtils/Texture2DBuilder.h"
 
 #include UE_INLINE_GENERATED_CPP_BY_NAME(BakeRenderCaptureTool)
 
@@ -407,7 +408,7 @@ void UBakeRenderCaptureTool::OnShutdown(EToolShutdownType ShutdownType)
 		IStaticMeshBackedTarget* StaticMeshTarget = Cast<IStaticMeshBackedTarget>(Targets[0]);
 		UObject* SourceAsset = StaticMeshTarget ? StaticMeshTarget->GetStaticMesh() : nullptr;
 		const UPrimitiveComponent* SourceComponent = UE::ToolTarget::GetTargetComponent(Targets[0]);
-		CreateTextureAssets(SourceComponent->GetWorld(), SourceAsset);
+		CreateAssets(SourceComponent->GetWorld(), SourceAsset);
 	}
 
 	// Clear actors on shutdown so that their lifetime is not tied to the lifetime of the tool
@@ -416,91 +417,174 @@ void UBakeRenderCaptureTool::OnShutdown(EToolShutdownType ShutdownType)
 	UE::ToolTarget::ShowSourceObject(Targets[0]);
 }
 
-void UBakeRenderCaptureTool::CreateTextureAssets(UWorld* SourceWorld, UObject* SourceAsset)
+// TODO Move the BakeRC Tool into MeshModellingToolsEditorOnlyExp (this function uses WITH_EDITOR functionality, other
+// functions may also use editor only features)
+void UBakeRenderCaptureTool::CreateAssets(UWorld* SourceWorld, UObject* SourceAsset)
 {
-	bool bCreatedAssetOK = true;
 	const FString BaseName = UE::ToolTarget::GetTargetActor(Targets[0])->GetActorNameOrLabel();
+	const bool bPackedMRS = ResultSettings->PackedMRSMap != nullptr;
+	const bool bSubsurfaceMaterial = ResultSettings->SubsurfaceColorMap || ResultSettings->OpacityMap;
 
-	auto CreateTextureAsset = [this, &bCreatedAssetOK, &SourceWorld, &SourceAsset](const FString& TexName, FTexture2DBuilder::ETextureType Type, TObjectPtr<UTexture2D> Tex)
-	{
-		// See :DeferredPopulateSourceData
-		FTexture2DBuilder::CopyPlatformDataToSourceData(Tex, Type);
-
-		// TODO The original implementation in ApproximateActors also did the following, see WriteTextureLambda in ApproximateActorsImpl.cpp
-		//if (Type == FTexture2DBuilder::ETextureType::Roughness
-		//	|| Type == FTexture2DBuilder::ETextureType::Metallic
-		//	|| Type == FTexture2DBuilder::ETextureType::Specular)
-		//{
-		//	UE::AssetUtils::ConvertToSingleChannel(Texture);
-		//}
-
-		bCreatedAssetOK = bCreatedAssetOK &&
-			UE::Modeling::CreateTextureObject(
-				GetToolManager(),
-				FCreateTextureObjectParams{ 0, SourceWorld, SourceAsset, TexName, Tex }).IsOK();
-	};
-
-	if (RenderCaptureProperties->bBaseColorMap && ResultSettings->BaseColorMap != nullptr)
 	{
-		const FString TexName = FString::Printf(TEXT("%s_%s"), *BaseName, *BaseColorTexParamName);
-		CreateTextureAsset(TexName, FTexture2DBuilder::ETextureType::Color, ResultSettings->BaseColorMap);
-	}
+		// Choose the material
+		TObjectPtr<UMaterialInstanceDynamic> Material = ResultSettings->PackedMRSMap ? PreviewMaterialPackedRC : PreviewMaterialRC;
+		if (bSubsurfaceMaterial)
+		{
+			Material = ResultSettings->PackedMRSMap ? PreviewMaterialPackedRC_Subsurface : PreviewMaterialRC_Subsurface;
+			ensure(Material->GetShadingModels().HasShadingModel(EMaterialShadingModel::MSM_Subsurface));
+			ensure(Material->GetBlendMode() == EBlendMode::BLEND_Masked);
+		}
 
-	if (RenderCaptureProperties->bNormalMap && ResultSettings->NormalMap != nullptr)
-	{
-		const FString TexName = FString::Printf(TEXT("%s_%s"), *BaseName, *NormalTexParamName);
-		CreateTextureAsset(TexName, FTexture2DBuilder::ETextureType::NormalMap, ResultSettings->NormalMap);
-	}
+		FCreateMaterialObjectParams MaterialParams;
+		MaterialParams.TargetWorld = SourceWorld;
+		MaterialParams.StoreRelativeToObject = SourceAsset;
+		MaterialParams.BaseName = FString::Printf(TEXT("%s_Material"), *BaseName);
+		MaterialParams.MaterialToDuplicate = Material;
+		const FCreateMaterialObjectResult MaterialResult = UE::Modeling::CreateMaterialObject(GetToolManager(), MoveTemp(MaterialParams));
 
-	if (RenderCaptureProperties->bEmissiveMap && ResultSettings->EmissiveMap != nullptr)
-	{
-		const FString TexName = FString::Printf(TEXT("%s_%s"), *BaseName, *EmissiveTexParamName);
-		CreateTextureAsset(TexName, FTexture2DBuilder::ETextureType::EmissiveHDR, ResultSettings->EmissiveMap);
-	}
+		UMaterial* NewMaterial = nullptr;
+		if (ensure(MaterialResult.ResultCode == ECreateModelingObjectResult::Ok))
+		{
+			NewMaterial = CastChecked<UMaterial>(MaterialResult.NewAsset);
+			ensure(NewMaterial);
+		}
 
-	if (RenderCaptureProperties->bOpacityMap && ResultSettings->OpacityMap != nullptr)
-	{
-		const FString TexName = FString::Printf(TEXT("%s_%s"), *BaseName, *OpacityTexParamName);
-		CreateTextureAsset(TexName, FTexture2DBuilder::ETextureType::ColorLinear, ResultSettings->OpacityMap);
-	}
+		const auto TrySetTextureEditorOnly = [NewMaterial](
+			const FString& TextureName,
+			TObjectPtr<UTexture2D> Texture,
+			TObjectPtr<UTexture2D> Fallback,
+			bool bMaterialHasTexture)
+		{
+			if (NewMaterial && bMaterialHasTexture)
+			{
+				NewMaterial->SetTextureParameterValueEditorOnly(FName(TextureName), Texture ? Texture : Fallback);
+			}
+		};
 
-	if (RenderCaptureProperties->bSubsurfaceColorMap && ResultSettings->SubsurfaceColorMap != nullptr)
-	{
-		const FString TexName = FString::Printf(TEXT("%s_%s"), *BaseName, *SubsurfaceColorTexParamName);
-		CreateTextureAsset(TexName, FTexture2DBuilder::ETextureType::Color, ResultSettings->SubsurfaceColorMap);
+		// Set all computed textures or fallback to the empty texture map
+		TrySetTextureEditorOnly(BaseColorTexParamName, ResultSettings->BaseColorMap, EmptyColorMapWhite, true);
+		TrySetTextureEditorOnly(EmissiveTexParamName,  ResultSettings->EmissiveMap,  EmptyEmissiveMap,   true);
+		TrySetTextureEditorOnly(NormalTexParamName,    ResultSettings->NormalMap,    EmptyNormalMap,     true);
+		TrySetTextureEditorOnly(PackedMRSTexParamName, ResultSettings->PackedMRSMap, EmptyPackedMRSMap,  bPackedMRS);
+		TrySetTextureEditorOnly(MetallicTexParamName,  ResultSettings->MetallicMap,  EmptyMetallicMap,  !bPackedMRS);
+		TrySetTextureEditorOnly(RoughnessTexParamName, ResultSettings->RoughnessMap, EmptyRoughnessMap, !bPackedMRS);
+		TrySetTextureEditorOnly(SpecularTexParamName,  ResultSettings->SpecularMap,  EmptySpecularMap,  !bPackedMRS);
+		TrySetTextureEditorOnly(OpacityTexParamName,   ResultSettings->OpacityMap,   EmptyOpacityMap,    bSubsurfaceMaterial);
+		TrySetTextureEditorOnly(SubsurfaceColorTexParamName, ResultSettings->SubsurfaceColorMap, EmptySubsurfaceColorMap, bSubsurfaceMaterial);
+
+		if (NewMaterial)
+		{
+			// Force material update now that we have updated texture parameters
+			NewMaterial->PostEditChange();
+		}
 	}
 
-	// We need different code paths based on PackedMRS here because we don't want to uncheck the separate channels
-	// when PackedMRS is enabled to give the user a better UX (they don't have to re-check them after disabling
-	// PackedMRS). In other place we can test the PackedMRS and separate channel booleans in series and avoid the
-	// complexity of nested if statements.
-	if (RenderCaptureProperties->bPackedMRSMap && ResultSettings->PackedMRSMap != nullptr)
-	{
-		const FString TexName = FString::Printf(TEXT("%s_%s"), *BaseName, *PackedMRSTexParamName);
-		CreateTextureAsset(TexName, FTexture2DBuilder::ETextureType::ColorLinear, ResultSettings->PackedMRSMap);
-	}
-	else
+	auto CreateTextureAsset = [this, &BaseName, &SourceWorld, &SourceAsset] (
+		const FString& TexParamName,
+		FTexture2DBuilder::ETextureType Type,
+		TObjectPtr<UTexture2D> Texture,
+		TObjectPtr<UTexture2D> Fallback,
+		bool bForceFallback,
+		bool bMaterialHasTexture)
 	{
-		if (RenderCaptureProperties->bMetallicMap && ResultSettings->MetallicMap != nullptr)
+		if (!bMaterialHasTexture)
 		{
-			const FString TexName = FString::Printf(TEXT("%s_%s"), *BaseName, *MetallicTexParamName);
-			CreateTextureAsset(TexName, FTexture2DBuilder::ETextureType::Metallic, ResultSettings->MetallicMap);
+			return true;
 		}
 
-		if (RenderCaptureProperties->bRoughnessMap && ResultSettings->RoughnessMap != nullptr)
+		const bool bUseTexture = Texture && !bForceFallback;
+
+		// See :DeferredPopulateSourceData
+		if (bUseTexture)
 		{
-			const FString TexName = FString::Printf(TEXT("%s_%s"), *BaseName, *RoughnessTexParamName);
-			CreateTextureAsset(TexName, FTexture2DBuilder::ETextureType::Roughness, ResultSettings->RoughnessMap);
+			FTexture2DBuilder::CopyPlatformDataToSourceData(Texture, Type);
 		}
-
-		if (RenderCaptureProperties->bSpecularMap && ResultSettings->SpecularMap != nullptr)
+		else
 		{
-			const FString TexName = FString::Printf(TEXT("%s_%s"), *BaseName, *SpecularTexParamName);
-			CreateTextureAsset(TexName, FTexture2DBuilder::ETextureType::Specular, ResultSettings->SpecularMap);
+			FTexture2DBuilder::CopyPlatformDataToSourceData(Fallback, Type);
 		}
-	}
 
-	ensure(bCreatedAssetOK);
+		// TODO The original implementation in ApproximateActors also did the following, see WriteTextureLambda in ApproximateActorsImpl.cpp
+		//if (Type == FTexture2DBuilder::ETextureType::Roughness
+		//	|| Type == FTexture2DBuilder::ETextureType::Metallic
+		//	|| Type == FTexture2DBuilder::ETextureType::Specular)
+		//{
+		//	UE::AssetUtils::ConvertToSingleChannel(Texture);
+		//}
+
+		// We need to save the Fallback textures as well so that the generated Materials can reference them
+		FCreateTextureObjectParams TexParams;
+		TexParams.TargetWorld = SourceWorld;
+		TexParams.StoreRelativeToObject = SourceAsset;
+		TexParams.BaseName = FString::Printf(TEXT("%s_%s"), *BaseName, *TexParamName);
+		TexParams.GeneratedTransientTexture = bUseTexture ? Texture : Fallback; 
+
+		FCreateTextureObjectResult TexResult = UE::Modeling::CreateTextureObject(GetToolManager(), MoveTemp(TexParams));
+		return TexResult.IsOK();
+	};
+
+	ensure(CreateTextureAsset(BaseColorTexParamName,
+		FTexture2DBuilder::ETextureType::Color,
+		ResultSettings->BaseColorMap,
+		EmptyColorMapWhite,
+		RenderCaptureProperties->bBaseColorMap == false,
+		true));
+
+	ensure(CreateTextureAsset(NormalTexParamName,
+		FTexture2DBuilder::ETextureType::NormalMap,
+		ResultSettings->NormalMap,
+		EmptyNormalMap,
+		RenderCaptureProperties->bNormalMap == false,
+		true));
+
+	ensure(CreateTextureAsset(EmissiveTexParamName,
+		FTexture2DBuilder::ETextureType::EmissiveHDR,
+		ResultSettings->EmissiveMap,
+		EmptyEmissiveMap,
+		RenderCaptureProperties->bEmissiveMap == false,
+		true));
+
+	ensure(CreateTextureAsset(PackedMRSTexParamName,
+		FTexture2DBuilder::ETextureType::ColorLinear,
+		ResultSettings->PackedMRSMap,
+		EmptyPackedMRSMap,
+		RenderCaptureProperties->bPackedMRSMap == false,
+		bPackedMRS));
+
+	ensure(CreateTextureAsset(RoughnessTexParamName,
+		FTexture2DBuilder::ETextureType::Roughness,
+		ResultSettings->RoughnessMap,
+		EmptyRoughnessMap,
+		RenderCaptureProperties->bRoughnessMap == false,
+		!bPackedMRS));
+
+	ensure(CreateTextureAsset(MetallicTexParamName,
+		FTexture2DBuilder::ETextureType::Metallic,
+		ResultSettings->MetallicMap,
+		EmptyMetallicMap,
+		RenderCaptureProperties->bMetallicMap == false,
+		!bPackedMRS));
+
+	ensure(CreateTextureAsset(SpecularTexParamName,
+		FTexture2DBuilder::ETextureType::Specular,
+		ResultSettings->SpecularMap,
+		EmptySpecularMap,
+		RenderCaptureProperties->bSpecularMap == false,
+		!bPackedMRS));
+
+	ensure(CreateTextureAsset(OpacityTexParamName,
+		FTexture2DBuilder::ETextureType::ColorLinear,
+		ResultSettings->OpacityMap,
+		EmptyOpacityMap,
+		RenderCaptureProperties->bOpacityMap == false,
+		bSubsurfaceMaterial));
+
+	ensure(CreateTextureAsset(SubsurfaceColorTexParamName,
+		FTexture2DBuilder::ETextureType::Color,
+		ResultSettings->SubsurfaceColorMap,
+		EmptySubsurfaceColorMap,
+		RenderCaptureProperties->bSubsurfaceColorMap == false,
+		bSubsurfaceMaterial));
 
 	RecordAnalytics();
 }
@@ -668,10 +752,14 @@ bool UBakeRenderCaptureTool::ValidTargetMeshTangents()
 
 void UBakeRenderCaptureTool::InitializePreviewMaterials()
 {
+	// We will only need source data if we're saving these textures so that the material generated when the tool is
+	// Accepted can reference them :DeferredPopulateSourceData
+	constexpr bool bPopulateSourceData = false;
+
 	{
 		FTexture2DBuilder Builder;
 		Builder.Initialize(FTexture2DBuilder::ETextureType::NormalMap, FImageDimensions(16, 16));
-		Builder.Commit(false);
+		Builder.Commit(bPopulateSourceData);
 		EmptyNormalMap = Builder.GetTexture2D();
 	}
 
@@ -679,7 +767,7 @@ void UBakeRenderCaptureTool::InitializePreviewMaterials()
 		FTexture2DBuilder Builder;
 		Builder.Initialize(FTexture2DBuilder::ETextureType::Color, FImageDimensions(16, 16));
 		Builder.Clear(FColor(0,0,0));
-		Builder.Commit(false);
+		Builder.Commit(bPopulateSourceData);
 		EmptyColorMapBlack = Builder.GetTexture2D();
 	}
 
@@ -687,14 +775,14 @@ void UBakeRenderCaptureTool::InitializePreviewMaterials()
 		FTexture2DBuilder Builder;
 		Builder.Initialize(FTexture2DBuilder::ETextureType::Color, FImageDimensions(16, 16));
 		Builder.Clear(FColor::White);
-		Builder.Commit(false);
+		Builder.Commit(bPopulateSourceData);
 		EmptyColorMapWhite = Builder.GetTexture2D();
 	}
 
 	{
 		FTexture2DBuilder Builder;
 		Builder.Initialize(FTexture2DBuilder::ETextureType::EmissiveHDR, FImageDimensions(16, 16));
-		Builder.Commit(false);
+		Builder.Commit(bPopulateSourceData);
 		EmptyEmissiveMap = Builder.GetTexture2D();
 	}
 
@@ -704,7 +792,7 @@ void UBakeRenderCaptureTool::InitializePreviewMaterials()
 		// The Opacity texture is passed to the Material's Opacity pin as well as the Opacity Mask pin, so set white
 		// here so we see something when previewing the subsurface material when opacity is not baked
 		Builder.Clear(FColor::White);
-		Builder.Commit(false);
+		Builder.Commit(bPopulateSourceData);
 		EmptyOpacityMap = Builder.GetTexture2D();
 	}
 
@@ -712,7 +800,7 @@ void UBakeRenderCaptureTool::InitializePreviewMaterials()
 		FTexture2DBuilder Builder;
 		Builder.Initialize(FTexture2DBuilder::ETextureType::Color, FImageDimensions(16, 16));
 		Builder.Clear(FColor::Black);
-		Builder.Commit(false);
+		Builder.Commit(bPopulateSourceData);
 		EmptySubsurfaceColorMap = Builder.GetTexture2D();
 	}
 
@@ -720,28 +808,28 @@ void UBakeRenderCaptureTool::InitializePreviewMaterials()
 		FTexture2DBuilder Builder;
 		Builder.Initialize(FTexture2DBuilder::ETextureType::ColorLinear, FImageDimensions(16, 16));
 		Builder.Clear(FColor(0,0,0));
-		Builder.Commit(false);
+		Builder.Commit(bPopulateSourceData);
 		EmptyPackedMRSMap = Builder.GetTexture2D();
 	}
 
 	{
 		FTexture2DBuilder Builder;
 		Builder.Initialize(FTexture2DBuilder::ETextureType::Roughness, FImageDimensions(16, 16));
-		Builder.Commit(false);
+		Builder.Commit(bPopulateSourceData);
 		EmptyRoughnessMap = Builder.GetTexture2D();
 	}
 
 	{
 		FTexture2DBuilder Builder;
 		Builder.Initialize(FTexture2DBuilder::ETextureType::Metallic, FImageDimensions(16, 16));
-		Builder.Commit(false);
+		Builder.Commit(bPopulateSourceData);
 		EmptyMetallicMap = Builder.GetTexture2D();
 	}
 
 	{
 		FTexture2DBuilder Builder;
 		Builder.Initialize(FTexture2DBuilder::ETextureType::Specular, FImageDimensions(16, 16));
-		Builder.Commit(false);
+		Builder.Commit(bPopulateSourceData);
 		EmptySpecularMap = Builder.GetTexture2D();
 	}
 
@@ -1019,7 +1107,7 @@ void UBakeRenderCaptureTool::UpdateVisualization()
 	if (VisualizationProps->bPreviewAsMaterial)
 	{
 		const auto TrySetTexture =
-			[Material](const FString& TextureName, TObjectPtr<UTexture2D> Texture, TObjectPtr<UTexture2D> Fallback, bool bMaterialHasTexture = true)
+			[Material](const FString& TextureName, TObjectPtr<UTexture2D> Texture, TObjectPtr<UTexture2D> Fallback, bool bMaterialHasTexture)
 		{
 			if (bMaterialHasTexture)
 			{
@@ -1028,9 +1116,9 @@ void UBakeRenderCaptureTool::UpdateVisualization()
 		};
 
 		// Set all computed textures or fallback to the empty texture map
-		TrySetTexture(BaseColorTexParamName, ResultSettings->BaseColorMap, EmptyColorMapWhite);
-		TrySetTexture(EmissiveTexParamName,  ResultSettings->EmissiveMap,  EmptyEmissiveMap);
-		TrySetTexture(NormalTexParamName,    ResultSettings->NormalMap,    EmptyNormalMap);
+		TrySetTexture(BaseColorTexParamName, ResultSettings->BaseColorMap, EmptyColorMapWhite, true);
+		TrySetTexture(EmissiveTexParamName,  ResultSettings->EmissiveMap,  EmptyEmissiveMap,   true);
+		TrySetTexture(NormalTexParamName,    ResultSettings->NormalMap,    EmptyNormalMap,     true);
 		TrySetTexture(PackedMRSTexParamName, ResultSettings->PackedMRSMap, EmptyPackedMRSMap,  bPackedMRS);
 		TrySetTexture(RoughnessTexParamName, ResultSettings->RoughnessMap, EmptyRoughnessMap, !bPackedMRS);
 		TrySetTexture(MetallicTexParamName,  ResultSettings->MetallicMap,  EmptyMetallicMap,  !bPackedMRS);
@@ -1041,7 +1129,7 @@ void UBakeRenderCaptureTool::UpdateVisualization()
 	else
 	{
 		const auto TrySetTexture =
-			[Material, this](const FString& TextureName, TObjectPtr<UTexture2D> Texture, TObjectPtr<UTexture2D> Fallback, bool bMaterialHasTexture = true)
+			[Material, this](const FString& TextureName, TObjectPtr<UTexture2D> Texture, TObjectPtr<UTexture2D> Fallback, bool bMaterialHasTexture)
 		{
 			// Set the BaseColor texture to the MapPreview texture if it exists and use white otherwise
 			if (TextureName == Settings->MapPreview)
@@ -1066,9 +1154,9 @@ void UBakeRenderCaptureTool::UpdateVisualization()
 			}
 		};
 
-		TrySetTexture(BaseColorTexParamName, ResultSettings->BaseColorMap, EmptyColorMapWhite);
-		TrySetTexture(EmissiveTexParamName,  ResultSettings->EmissiveMap,  EmptyEmissiveMap);
-		TrySetTexture(NormalTexParamName,    ResultSettings->NormalMap,    EmptyNormalMap);
+		TrySetTexture(BaseColorTexParamName, ResultSettings->BaseColorMap, EmptyColorMapWhite, true);
+		TrySetTexture(EmissiveTexParamName,  ResultSettings->EmissiveMap,  EmptyEmissiveMap,   true);
+		TrySetTexture(NormalTexParamName,    ResultSettings->NormalMap,    EmptyNormalMap,     true);
 		TrySetTexture(PackedMRSTexParamName, ResultSettings->PackedMRSMap, EmptyPackedMRSMap,  bPackedMRS);
 		TrySetTexture(RoughnessTexParamName, ResultSettings->RoughnessMap, EmptyRoughnessMap, !bPackedMRS);
 		TrySetTexture(MetallicTexParamName,  ResultSettings->MetallicMap,  EmptyMetallicMap,  !bPackedMRS);
diff --git a/Engine/Plugins/Experimental/MeshModelingToolsetExp/Source/MeshModelingToolsExp/Public/BakeRenderCaptureTool.h b/Engine/Plugins/Experimental/MeshModelingToolsetExp/Source/MeshModelingToolsExp/Public/BakeRenderCaptureTool.h
--- a/Engine/Plugins/Experimental/MeshModelingToolsetExp/Source/MeshModelingToolsExp/Public/BakeRenderCaptureTool.h
+++ b/Engine/Plugins/Experimental/MeshModelingToolsetExp/Source/MeshModelingToolsExp/Public/BakeRenderCaptureTool.h
@@ -4,10 +4,12 @@
 
 #include "CoreMinimal.h"
 #include "InteractiveToolBuilder.h"
+#include "ModelingOperators.h"
+#include "MeshOpPreviewHelpers.h"
 #include "BaseTools/MultiSelectionMeshEditingTool.h"
+#include "BakeMeshAttributeToolCommon.h"
 #include "Baking/BakingTypes.h"
 #include "Baking/RenderCaptureFunctions.h"
-#include "BakeMeshAttributeMapsToolBase.h"
 
 // Render Capture algorithm includes
 #include "Scene/SceneCapturePhotoSet.h"
@@ -351,11 +353,11 @@ protected:
 	bool bValidTargetMeshTangents = false;
 
 	/**
-	 * Create texture assets from our result map of Texture2D
-	 * @param SourceWorld the source world to define where the texture assets will be stored.
-	 * @param SourceAsset if not null, result textures will be stored adjacent to this asset.
+	 * Create texture and material assets from our result map of Texture2D
+	 * @param SourceWorld the source world to define where the assets will be stored.
+	 * @param SourceAsset if not null, result textures/material will be stored adjacent to this asset.
 	 */
-	void CreateTextureAssets(UWorld* SourceWorld, UObject* SourceAsset);
+	void CreateAssets(UWorld* SourceWorld, UObject* SourceAsset);
 
 	// The baking background compute operation
 	TUniquePtr<TGenericDataBackgroundCompute<UE::Geometry::FMeshMapBaker>> Compute = nullptr;
diff --git a/Engine/Plugins/Runtime/MeshModelingToolset/Source/ModelingComponents/Private/ModelingObjectsCreationAPI.cpp b/Engine/Plugins/Runtime/MeshModelingToolset/Source/ModelingComponents/Private/ModelingObjectsCreationAPI.cpp
--- a/Engine/Plugins/Runtime/MeshModelingToolset/Source/ModelingComponents/Private/ModelingObjectsCreationAPI.cpp
+++ b/Engine/Plugins/Runtime/MeshModelingToolset/Source/ModelingComponents/Private/ModelingObjectsCreationAPI.cpp
@@ -107,6 +107,31 @@ FCreateTextureObjectResult UE::Modeling::CreateTextureObject(UInteractiveToolMan
 
 
 
+FCreateMaterialObjectResult UE::Modeling::CreateMaterialObject(UInteractiveToolManager* ToolManager, FCreateMaterialObjectParams&& CreateMaterialParams)
+{
+	if (ensure(ToolManager))
+	{
+		UModelingObjectsCreationAPI* UseAPI = ToolManager->GetContextObjectStore()->FindContext<UModelingObjectsCreationAPI>();
+		if (UseAPI)
+		{
+			if (UseAPI->HasMoveVariants())
+			{
+				return UseAPI->CreateMaterialObject(MoveTemp(CreateMaterialParams));
+			}
+			else
+			{
+				return UseAPI->CreateMaterialObject(CreateMaterialParams);
+			}
+		}
+	}
+	return FCreateMaterialObjectResult{ ECreateModelingObjectResult::Failed_NoAPIFound };
+}
+
+
+
+
+
+
 
 
 FString UE::Modeling::GetComponentAssetBaseName(UPrimitiveComponent* Component, bool bRemoveAutoGeneratedSuffixes)
diff --git a/Engine/Plugins/Runtime/MeshModelingToolset/Source/ModelingComponents/Public/ModelingObjectsCreationAPI.h b/Engine/Plugins/Runtime/MeshModelingToolset/Source/ModelingComponents/Public/ModelingObjectsCreationAPI.h
--- a/Engine/Plugins/Runtime/MeshModelingToolset/Source/ModelingComponents/Public/ModelingObjectsCreationAPI.h
+++ b/Engine/Plugins/Runtime/MeshModelingToolset/Source/ModelingComponents/Public/ModelingObjectsCreationAPI.h
@@ -35,7 +35,8 @@ enum class ECreateModelingObjectResult : uint8
 	Failed_InvalidMesh,
 	Failed_InvalidTexture,
 	Failed_AssetCreationFailed,
-    Failed_ActorCreationFailed
+	Failed_ActorCreationFailed,
+	Failed_InvalidMaterial,
 };
 
 /**
@@ -297,7 +298,7 @@ struct MODELINGCOMPONENTS_API FCreateTextureObjectResult
 	UPROPERTY(Category = "CreateTextureObjectResult", VisibleAnywhere)
 	ECreateModelingObjectResult ResultCode = ECreateModelingObjectResult::Ok;
 
-	/** A pointer to a newly-created Asset for the texture object, if applicable (eg StaticMeshAsset) */
+	/** A pointer to a newly-created Asset for the texture object */
 	UPROPERTY(Category = "CreateTextureObjectResult", VisibleAnywhere)
 	TObjectPtr<UObject> NewAsset = nullptr;
 
@@ -310,6 +311,76 @@ struct MODELINGCOMPONENTS_API FCreateTextureObjectResult
 
 
 
+/**
+ * FCreateMaterialObjectParams is a collection of input data intended to be passed to
+ * UModelingObjectsCreationAPI::CreateMaterialObject().
+ */
+USTRUCT(Blueprintable)
+struct MODELINGCOMPONENTS_API FCreateMaterialObjectParams
+{
+	GENERATED_BODY()
+
+	//
+	// Base data
+	//
+
+	/** 
+	 * The World/Level the new Material object should be created in (if known).
+	 * Note that Material generally do not exist as objects in a Level. 
+	 * However, it may be necessary to store the texture in a path relative to the
+	 * level (for example if the level is in a Plugin, this would be necessary in-Editor)
+	 */
+	UPROPERTY(Category = "CreateMaterialObjectParams", EditAnywhere)
+	TObjectPtr<UWorld> TargetWorld = nullptr;
+
+	/** An object to store the Material relative to. */
+	UPROPERTY(Category = "CreateMaterialObjectParams", EditAnywhere)
+	TObjectPtr<UObject> StoreRelativeToObject = nullptr;
+
+	/** The base name of the new Material object */
+	UPROPERTY(Category = "CreateMaterialObjectParams", EditAnywhere)
+	FString BaseName;
+
+	//
+	// input data
+	//
+
+	/** 
+	 * The parent UMaterial of this material will be duplicated to create the new UMaterial Asset.
+	 */
+	UPROPERTY(Category = "CreateMaterialObjectParams", EditAnywhere)
+	TObjectPtr<UMaterialInterface> MaterialToDuplicate = nullptr;
+};
+
+
+/**
+ * FCreateMaterialObjectResult is returned by UModelingObjectsCreationAPI::CreateMaterialObject()
+ * to indicate success/failure and provide information about created texture objects
+ */
+USTRUCT(BlueprintType)
+struct MODELINGCOMPONENTS_API FCreateMaterialObjectResult
+{
+	GENERATED_BODY()
+
+	/** Success/Failure status for the requested operation */
+	UPROPERTY(Category = "CreateMaterialObjectResult", VisibleAnywhere)
+	ECreateModelingObjectResult ResultCode = ECreateModelingObjectResult::Ok;
+
+	/** A pointer to a newly-created Asset for the material object */
+	UPROPERTY(Category = "CreateTextureObjectResult", VisibleAnywhere)
+	TObjectPtr<UObject> NewAsset = nullptr;
+
+
+	bool IsOK() const { return ResultCode == ECreateModelingObjectResult::Ok; }
+};
+
+
+
+
+
+
+
+
 /**
  * UModelingObjectsCreationAPI is a base interface for functions that can be used to
  * create various types of objects from Modeling Tools, or other sources. The "type" is
@@ -349,6 +420,13 @@ public:
 	UFUNCTION(BlueprintCallable, Category = "Modeling Objects")
 	virtual FCreateTextureObjectResult CreateTextureObject(const FCreateTextureObjectParams& CreateTexParams) { return FCreateTextureObjectResult{ ECreateModelingObjectResult::Failed_Unknown }; }
 
+	/**
+	 * Create a new material object based on the data in CreateMaterialParams
+	 * @return a results data structure, containing a result code and information about any new objects created
+	 */
+	UFUNCTION(BlueprintCallable, Category = "Modeling Objects")
+	virtual FCreateMaterialObjectResult CreateMaterialObject(const FCreateMaterialObjectParams& CreateMaterialParams) { return FCreateMaterialObjectResult{ ECreateModelingObjectResult::Failed_Unknown }; }
+
 
 	//
 	// Non-UFunction variants that support std::move operators
@@ -362,6 +440,7 @@ public:
 
 	virtual FCreateMeshObjectResult CreateMeshObject(FCreateMeshObjectParams&& CreateMeshParams) { return FCreateMeshObjectResult{ ECreateModelingObjectResult::Failed_Unknown }; }
 	virtual FCreateTextureObjectResult CreateTextureObject(FCreateTextureObjectParams&& CreateTexParams) { return FCreateTextureObjectResult{ ECreateModelingObjectResult::Failed_Unknown }; }
+	virtual FCreateMaterialObjectResult CreateMaterialObject(FCreateMaterialObjectParams&& CreateMaterialParams) { return FCreateMaterialObjectResult{ ECreateModelingObjectResult::Failed_Unknown }; }
 
 };
 
@@ -392,6 +471,15 @@ MODELINGCOMPONENTS_API FCreateMeshObjectResult CreateMeshObject(UInteractiveTool
 MODELINGCOMPONENTS_API FCreateTextureObjectResult CreateTextureObject(UInteractiveToolManager* ToolManager, FCreateTextureObjectParams&& CreateTexParams);
 
 
+/**
+ * Create a new material object based on the data in CreateMaterialParams.
+ * This is a convenience function that will try to locate a UModelingObjectsCreationAPI instance in the ToolManager's ContextObjectStore,
+ * and then call UModelingObjectsCreationAPI::CreateMaterialObject()
+ * @return a results data structure, containing a result code and information about any new objects created
+ */
+MODELINGCOMPONENTS_API FCreateMaterialObjectResult CreateMaterialObject(UInteractiveToolManager* ToolManager, FCreateMaterialObjectParams&& CreateMaterialParams);
+
+
 
 /**
  * Strip the appended auto-generated _UUID suffix on the given string, if we can detect one
diff --git a/Engine/Plugins/Runtime/MeshModelingToolset/Source/ModelingComponentsEditorOnly/Private/EditorModelingObjectsCreationAPI.cpp b/Engine/Plugins/Runtime/MeshModelingToolset/Source/ModelingComponentsEditorOnly/Private/EditorModelingObjectsCreationAPI.cpp
--- a/Engine/Plugins/Runtime/MeshModelingToolset/Source/ModelingComponentsEditorOnly/Private/EditorModelingObjectsCreationAPI.cpp
+++ b/Engine/Plugins/Runtime/MeshModelingToolset/Source/ModelingComponentsEditorOnly/Private/EditorModelingObjectsCreationAPI.cpp
@@ -7,6 +7,7 @@
 
 #include "AssetUtils/CreateStaticMeshUtil.h"
 #include "AssetUtils/CreateTexture2DUtil.h"
+#include "AssetUtils/CreateMaterialUtil.h"
 #include "Physics/ComponentCollisionUtil.h"
 
 #include "ConversionUtils/DynamicMeshToVolume.h"
@@ -96,6 +97,12 @@ FCreateTextureObjectResult UEditorModelingObjectsCreationAPI::CreateTextureObjec
 {
 	FCreateTextureObjectParams LocalParams = CreateTexParams;
 	return CreateTextureObject(MoveTemp(LocalParams));
+}
+
+FCreateMaterialObjectResult UEditorModelingObjectsCreationAPI::CreateMaterialObject(const FCreateMaterialObjectParams& CreateMaterialParams)
+{
+	FCreateMaterialObjectParams LocalParams = CreateMaterialParams;
+	return CreateMaterialObject(MoveTemp(LocalParams));
 
 }
 
@@ -286,21 +293,17 @@ FCreateMeshObjectResult UEditorModelingObjectsCreationAPI::CreateDynamicMeshActo
 
 FCreateMeshObjectResult UEditorModelingObjectsCreationAPI::CreateStaticMeshAsset(FCreateMeshObjectParams&& CreateMeshParams)
 {
-	if (!ensure(CreateMeshParams.TargetWorld)) { return FCreateMeshObjectResult{ ECreateModelingObjectResult::Failed_InvalidWorld }; }
-
 	UE::AssetUtils::FStaticMeshAssetOptions AssetOptions;
-
-	if (GetNewAssetPathNameCallback.IsBound())
-	{
-		AssetOptions.NewAssetPath = GetNewAssetPathNameCallback.Execute(CreateMeshParams.BaseName, CreateMeshParams.TargetWorld, FString());
-		if (AssetOptions.NewAssetPath.Len() == 0)
-		{
-			return FCreateMeshObjectResult{ ECreateModelingObjectResult::Cancelled };
-		}
-	}
-	else
+	
+	ECreateModelingObjectResult AssetPathResult = GetNewAssetPath(
+		AssetOptions.NewAssetPath,
+		CreateMeshParams.BaseName,
+		nullptr,
+		CreateMeshParams.TargetWorld);
+		
+	if (AssetPathResult != ECreateModelingObjectResult::Ok)
 	{
-		AssetOptions.NewAssetPath = "/Game/" + CreateMeshParams.BaseName;
+		return FCreateMeshObjectResult{ AssetPathResult };
 	}
 
 	AssetOptions.NumSourceModels = 1;
@@ -404,36 +407,17 @@ FCreateMeshObjectResult UEditorModelingObjectsCreationAPI::CreateStaticMeshAsset
 
 FCreateTextureObjectResult UEditorModelingObjectsCreationAPI::CreateTextureObject(FCreateTextureObjectParams&& CreateTexParams)
 {
-	FString RelativeToObjectFolder;
-	if (CreateTexParams.StoreRelativeToObject != nullptr)
-	{
-		// find path to asset
-		UPackage* AssetOuterPackage = CastChecked<UPackage>(CreateTexParams.StoreRelativeToObject->GetOuter());
-		if (ensure(AssetOuterPackage))
-		{
-			FString AssetPackageName = AssetOuterPackage->GetName();
-			RelativeToObjectFolder = FPackageName::GetLongPackagePath(AssetPackageName);
-		}
-	}
-	else
-	{
-		if (!ensure(CreateTexParams.TargetWorld)) { return FCreateTextureObjectResult{ ECreateModelingObjectResult::Failed_InvalidWorld }; }
-	}
-
 	UE::AssetUtils::FTexture2DAssetOptions AssetOptions;
 
-	if (GetNewAssetPathNameCallback.IsBound())
+	ECreateModelingObjectResult AssetPathResult = GetNewAssetPath(
+		AssetOptions.NewAssetPath,
+		CreateTexParams.BaseName,
+		CreateTexParams.StoreRelativeToObject,
+		CreateTexParams.TargetWorld);
+
+	if (AssetPathResult != ECreateModelingObjectResult::Ok)
 	{
-		AssetOptions.NewAssetPath = GetNewAssetPathNameCallback.Execute(CreateTexParams.BaseName, CreateTexParams.TargetWorld, RelativeToObjectFolder);
-		if (AssetOptions.NewAssetPath.Len() == 0)
-		{
-			return FCreateTextureObjectResult{ ECreateModelingObjectResult::Cancelled };
-		}
-	}
-	else
-	{
-		FString UseBaseFolder = (RelativeToObjectFolder.Len() > 0) ? RelativeToObjectFolder : TEXT("/Game");
-		AssetOptions.NewAssetPath = FPaths::Combine(UseBaseFolder, CreateTexParams.BaseName);
+		return FCreateTextureObjectResult{ AssetPathResult };
 	}
 
 	// currently we cannot create a new texture without an existing generated texture to store
@@ -459,5 +443,91 @@ FCreateTextureObjectResult UEditorModelingObjectsCreationAPI::CreateTextureObjec
 	OnModelingTextureCreated.Broadcast(ResultOut);
 
 	return ResultOut;
+}
+
+
+
+
+FCreateMaterialObjectResult UEditorModelingObjectsCreationAPI::CreateMaterialObject(FCreateMaterialObjectParams&& CreateMaterialParams)
+{
+	UE::AssetUtils::FMaterialAssetOptions AssetOptions;
+
+	ECreateModelingObjectResult AssetPathResult = GetNewAssetPath(
+		AssetOptions.NewAssetPath,
+		CreateMaterialParams.BaseName,
+		CreateMaterialParams.StoreRelativeToObject,
+		CreateMaterialParams.TargetWorld);
+
+	if (AssetPathResult != ECreateModelingObjectResult::Ok)
+	{
+		return FCreateMaterialObjectResult{ AssetPathResult };
+	}
+
+	// Currently we cannot create a new material without duplicating an existing material
+	if (!ensure(CreateMaterialParams.MaterialToDuplicate))
+	{
+		return FCreateMaterialObjectResult{ ECreateModelingObjectResult::Failed_InvalidMaterial };
+	}
+
+	UE::AssetUtils::FMaterialAssetResults ResultData;
+	UE::AssetUtils::ECreateMaterialResult AssetResult = UE::AssetUtils::CreateDuplicateMaterial(
+		CreateMaterialParams.MaterialToDuplicate,
+		AssetOptions,
+		ResultData);
+
+	if (AssetResult != UE::AssetUtils::ECreateMaterialResult::Ok)
+	{
+		return FCreateMaterialObjectResult{ ECreateModelingObjectResult::Failed_AssetCreationFailed };
+	}
+
+	// emit result
+	FCreateMaterialObjectResult ResultOut;
+	ResultOut.ResultCode = ECreateModelingObjectResult::Ok;
+	ResultOut.NewAsset = ResultData.NewMaterial;
+
+	OnModelingMaterialCreated.Broadcast(ResultOut);
+
+	return ResultOut;
+}
+
+
+ECreateModelingObjectResult UEditorModelingObjectsCreationAPI::GetNewAssetPath(
+	FString& OutNewAssetPath,
+	const FString& BaseName,
+	const UObject* StoreRelativeToObject,
+	const UWorld* TargetWorld)
+{
+	FString RelativeToObjectFolder;
+	if (StoreRelativeToObject != nullptr)
+	{
+		// find path to asset
+		UPackage* AssetOuterPackage = CastChecked<UPackage>(StoreRelativeToObject->GetOuter());
+		if (ensure(AssetOuterPackage))
+		{
+			FString AssetPackageName = AssetOuterPackage->GetName();
+			RelativeToObjectFolder = FPackageName::GetLongPackagePath(AssetPackageName);
+		}
+	}
+	else
+	{
+		if (!ensure(TargetWorld)) {
+			return ECreateModelingObjectResult::Failed_InvalidWorld;
+		}
+	}
+
+	if (GetNewAssetPathNameCallback.IsBound())
+	{
+		OutNewAssetPath = GetNewAssetPathNameCallback.Execute(BaseName, TargetWorld, RelativeToObjectFolder);
+		if (OutNewAssetPath.Len() == 0)
+		{
+			return ECreateModelingObjectResult::Cancelled;
+		}
+	}
+	else
+	{
+		FString UseBaseFolder = (RelativeToObjectFolder.Len() > 0) ? RelativeToObjectFolder : TEXT("/Game");
+		OutNewAssetPath = FPaths::Combine(UseBaseFolder, BaseName);
+	}
 
+	return ECreateModelingObjectResult::Ok;
 }
diff --git a/Engine/Plugins/Runtime/MeshModelingToolset/Source/ModelingComponentsEditorOnly/Public/EditorModelingObjectsCreationAPI.h b/Engine/Plugins/Runtime/MeshModelingToolset/Source/ModelingComponentsEditorOnly/Public/EditorModelingObjectsCreationAPI.h
--- a/Engine/Plugins/Runtime/MeshModelingToolset/Source/ModelingComponentsEditorOnly/Public/EditorModelingObjectsCreationAPI.h
+++ b/Engine/Plugins/Runtime/MeshModelingToolset/Source/ModelingComponentsEditorOnly/Public/EditorModelingObjectsCreationAPI.h
@@ -11,8 +11,9 @@ class UInteractiveToolsContext;
 
 /**
  * Implementation of UModelingObjectsCreationAPI suitable for use in UE Editor.
- * - CreateMeshObject() currently creates a StaticMesh Asset/Actor
+ * - CreateMeshObject() currently creates a StaticMesh Asset/Actor, a Volume Actor or a DynamicMesh Actor
  * - CreateTextureObject() currently creates a UTexture2D Asset
+ * - CreateMaterialObject() currently creates a UMaterial Asset
  * 
  * This is intended to be registered in the ToolsContext ContextObjectStore.
  * Static utility functions ::Register() / ::Find() / ::Deregister() can be used to do this in a consistent way.
@@ -22,6 +23,7 @@ class UInteractiveToolsContext;
  *    things like pop up an interactive path-selection dialog, use project-defined paths, etc
  *  - OnModelingMeshCreated is broadcast for each new created mesh object
  *  - OnModelingTextureCreated is broadcast for each new created texture object
+ *  - OnModelingMaterialCreated is broadcast for each new created material object
  */
 UCLASS()
 class MODELINGCOMPONENTSEDITORONLY_API UEditorModelingObjectsCreationAPI : public UModelingObjectsCreationAPI
@@ -35,6 +37,8 @@ public:
 	// UFUNCTION(BlueprintCallable, Category = "Modeling Objects")
 	virtual FCreateTextureObjectResult CreateTextureObject(const FCreateTextureObjectParams& CreateTexParams) override;
 
+	// UFUNCTION(BlueprintCallable, Category = "Modeling Objects")
+	virtual FCreateMaterialObjectResult CreateMaterialObject(const FCreateMaterialObjectParams& CreateMaterialParams) override;
 
 	//
 	// Non-UFunction variants that support std::move operators
@@ -46,6 +50,8 @@ public:
 
 	virtual FCreateTextureObjectResult CreateTextureObject(FCreateTextureObjectParams&& CreateTexParams) override;
 
+	virtual FCreateMaterialObjectResult CreateMaterialObject(FCreateMaterialObjectParams&& CreateMaterialParams) override;
+
 
 	//
 	// Callbacks that editor can hook into to handle asset creation
@@ -60,6 +66,9 @@ public:
 	DECLARE_MULTICAST_DELEGATE_OneParam(FModelingTextureCreatedSignature, const FCreateTextureObjectResult& CreatedInfo);
 	FModelingTextureCreatedSignature OnModelingTextureCreated;
 
+	DECLARE_MULTICAST_DELEGATE_OneParam(FModelingMaterialCreatedSignature, const FCreateMaterialObjectResult& CreatedInfo);
+	FModelingMaterialCreatedSignature OnModelingMaterialCreated;
+
 
 
 	//
@@ -78,6 +87,7 @@ public:
 	FCreateMeshObjectResult CreateStaticMeshAsset(FCreateMeshObjectParams&& CreateMeshParams);
 	FCreateMeshObjectResult CreateVolume(FCreateMeshObjectParams&& CreateMeshParams);
 	FCreateMeshObjectResult CreateDynamicMeshActor(FCreateMeshObjectParams&& CreateMeshParams);
+	ECreateModelingObjectResult GetNewAssetPath(FString& OutNewAssetPath, const FString& BaseName, const UObject* StoreRelativeToObject, const UWorld* World);
 
 	TArray<UMaterialInterface*> FilterMaterials(const TArray<UMaterialInterface*>& MaterialsIn);
 };
