Fix Horde issue 306851

#horde 306851

diff --git a/Engine/Plugins/Experimental/MeshModelingToolsetExp/Source/MeshModelingToolsExp/Private/BakeRenderCaptureTool.cpp b/Engine/Plugins/Experimental/MeshModelingToolsetExp/Source/MeshModelingToolsExp/Private/BakeRenderCaptureTool.cpp
--- a/Engine/Plugins/Experimental/MeshModelingToolsetExp/Source/MeshModelingToolsExp/Private/BakeRenderCaptureTool.cpp
+++ b/Engine/Plugins/Experimental/MeshModelingToolsetExp/Source/MeshModelingToolsExp/Private/BakeRenderCaptureTool.cpp
@@ -425,6 +425,7 @@ void UBakeRenderCaptureTool::CreateAssets(UWorld* SourceWorld, UObject* SourceAs
 	const bool bPackedMRS = ResultSettings->PackedMRSMap != nullptr;
 	const bool bSubsurfaceMaterial = ResultSettings->SubsurfaceColorMap || ResultSettings->OpacityMap;
 
+#if WITH_EDITOR
 	{
 		// Choose the material
 		TObjectPtr<UMaterialInstanceDynamic> Material = ResultSettings->PackedMRSMap ? PreviewMaterialPackedRC : PreviewMaterialRC;
@@ -478,6 +479,7 @@ void UBakeRenderCaptureTool::CreateAssets(UWorld* SourceWorld, UObject* SourceAs
 			NewMaterial->PostEditChange();
 		}
 	}
+#endif // WITH_EDITOR
 
 	auto CreateTextureAsset = [this, &BaseName, &SourceWorld, &SourceAsset] (
 		const FString& TexParamName,
