Modeling Tools: minor changes in prepration for future CL. Add UE::Geometry::Lerp for TVector4 types. Refactor code MeshGroupPaintTool code for polylasso post-process to static function UPolyLassoMarqueeMechanic::ApproximateSelfClipPolyline(). Make brush flow rate configurable via overridable UMeshSculptToolBase::GetStampTemporalFlowRate().
#rnx

diff --git a/Engine/Plugins/Experimental/MeshModelingToolsetExp/Source/MeshModelingToolsExp/Private/MeshGroupPaintTool.cpp b/Engine/Plugins/Experimental/MeshModelingToolsetExp/Source/MeshModelingToolsExp/Private/MeshGroupPaintTool.cpp
--- a/Engine/Plugins/Experimental/MeshModelingToolsetExp/Source/MeshModelingToolsExp/Private/MeshGroupPaintTool.cpp
+++ b/Engine/Plugins/Experimental/MeshModelingToolsetExp/Source/MeshModelingToolsExp/Private/MeshGroupPaintTool.cpp
@@ -21,8 +21,6 @@
 #include "Selections/MeshVertexSelection.h"
 #include "Polygroups/PolygroupUtil.h"
 #include "Polygon2.h"
-#include "Intersection/IntrLine2Line2.h"
-#include "Intersection/IntrSegment2Segment2.h"
 
 #include "Changes/MeshVertexChange.h"
 #include "Changes/MeshPolygroupChange.h"
@@ -654,161 +652,11 @@ bool UMeshGroupPaintTool::SyncMeshWithGroupBuffer(FDynamicMesh3* Mesh)
 
 
 
-template<typename RealType>
-static bool FindPolylineSelfIntersection(
-	const TArray<UE::Math::TVector2<RealType>>& Polyline, 
-	UE::Math::TVector2<RealType>& IntersectionPointOut, 
-	FIndex2i& IntersectionIndexOut,
-	bool bParallel = true)
-{
-	int32 N = Polyline.Num();
-	std::atomic<bool> bSelfIntersects(false);
-	ParallelFor(N - 1, [&](int32 i)
-	{
-		TSegment2<RealType> SegA(Polyline[i], Polyline[i + 1]);
-		for (int32 j = i + 2; j < N - 1 && bSelfIntersects == false; ++j)
-		{
-			TSegment2<RealType> SegB(Polyline[j], Polyline[j + 1]);
-			if (SegA.Intersects(SegB) && bSelfIntersects == false)		
-			{
-				bool ExpectedValue = false;
-				if (std::atomic_compare_exchange_strong(&bSelfIntersects, &ExpectedValue, true))
-				{
-					TIntrSegment2Segment2<RealType> Intersection(SegA, SegB);
-					Intersection.Find();
-					IntersectionPointOut = Intersection.Point0;
-					IntersectionIndexOut = FIndex2i(i, j);
-					return;
-				}
-			}
-		}
-	}, (bParallel) ? EParallelForFlags::None : EParallelForFlags::ForceSingleThread );
-
-	return bSelfIntersects;
-}
-
-
-
-template<typename RealType>
-static bool FindPolylineSegmentIntersection(
-	const TArray<UE::Math::TVector2<RealType>>& Polyline,
-	const TSegment2<RealType>& Segment,
-	UE::Math::TVector2<RealType>& IntersectionPointOut,
-	int& IntersectionIndexOut)
-{
-
-	int32 N = Polyline.Num();
-	for (int32 i = 0; i < N-1; ++i)
-	{
-		TSegment2<RealType> PolySeg(Polyline[i], Polyline[i + 1]);
-		if (Segment.Intersects(PolySeg))
-		{
-			TIntrSegment2Segment2<RealType> Intersection(Segment, PolySeg);
-			Intersection.Find();
-			IntersectionPointOut = Intersection.Point0;
-			IntersectionIndexOut = i;
-			return true;
-		}
-	}
-	return false;
-}
-
-
-
-bool ApproxSelfClipPolyline(TArray<FVector2f>& Polyline)
-{
-	int32 N = Polyline.Num();
-
-	// handle already-closed polylines
-	if (Distance(Polyline[0], Polyline[N-1]) < 0.0001f)
-	{
-		return true;
-	}
-
-	FVector2f IntersectPoint;
-	FIndex2i IntersectionIndex(-1, -1);
-	bool bSelfIntersects = FindPolylineSelfIntersection(Polyline, IntersectPoint, IntersectionIndex);
-	if (bSelfIntersects)
-	{
-		TArray<FVector2f> NewPolyline;
-		NewPolyline.Add(IntersectPoint);
-		for (int32 i = IntersectionIndex.A; i <= IntersectionIndex.B; ++i)
-		{
-			NewPolyline.Add(Polyline[i]);
-		}
-		NewPolyline.Add(IntersectPoint);
-		Polyline = MoveTemp(NewPolyline);
-		return true;
-	}
-
-
-	FVector2f StartDirOut = UE::Geometry::Normalized(Polyline[0] - Polyline[1]);
-	FLine2f StartLine(Polyline[0], StartDirOut);
-	FVector2f EndDirOut = UE::Geometry::Normalized(Polyline[N - 1] - Polyline[N - 2]);
-	FLine2f EndLine(Polyline[N - 1], EndDirOut);
-	FIntrLine2Line2f LineIntr(StartLine, EndLine);
-	bool bIntersects = false;
-	if (LineIntr.Find())
-	{
-		bIntersects = LineIntr.IsSimpleIntersection() && (LineIntr.Segment1Parameter > 0) && (LineIntr.Segment2Parameter > 0);
-		if (bIntersects)
-		{
-			Polyline.Add(StartLine.PointAt(LineIntr.Segment1Parameter));
-			Polyline.Add(StartLine.Origin);
-			return true;
-		}
-	}
-
-
-	FAxisAlignedBox2f Bounds;
-	for (const FVector2f& P : Polyline)
-	{
-		Bounds.Contain(P);
-	}
-	float Size = Bounds.DiagonalLength();
-
-	FVector2f StartPos = Polyline[0] + 0.001f * StartDirOut;
-	if (FindPolylineSegmentIntersection(Polyline, FSegment2f(StartPos, StartPos + 2*Size*StartDirOut), IntersectPoint, IntersectionIndex.A))
-	{
-		//TArray<FVector2f> NewPolyline;
-		//for (int32 i = 0; i <= IntersectionIndex.A; ++i)
-		//{
-		//	NewPolyline.Add(Polyline[i]);
-		//}
-		//NewPolyline.Add(IntersectPoint);
-		//NewPolyline.Add(Polyline[0]);
-		//Polyline = MoveTemp(NewPolyline);
-		return true;
-	}
-
-	FVector2f EndPos = Polyline[N-1] + 0.001f * EndDirOut;
-	if (FindPolylineSegmentIntersection(Polyline, FSegment2f(EndPos, EndPos + 2*Size*EndDirOut), IntersectPoint, IntersectionIndex.A))
-	{
-		//TArray<FVector2f> NewPolyline;
-		//NewPolyline.Add(IntersectPoint);
-		//for (int32 i = IntersectionIndex.A+1; i < N; ++i)
-		//{
-		//	NewPolyline.Add(Polyline[i]);
-		//}
-		//NewPolyline.Add(Polyline[0]);
-		//NewPolyline.Add(IntersectPoint);
-		//Polyline = MoveTemp(NewPolyline);
-		return true;
-	}
-
-	return false;
-}
-
-
 
 void UMeshGroupPaintTool::OnPolyLassoFinished(const FCameraPolyLasso& Lasso, bool bCanceled)
 {
 	// construct polyline
-	TArray<FVector2f> Polyline;
-	for (FVector2D Pos : Lasso.Polyline)
-	{
-		Polyline.Add((FVector2f)Pos);
-	}
+	TArray<FVector2D> Polyline = Lasso.Polyline;
 	int32 N = Polyline.Num();
 	if (N < 2)
 	{
@@ -818,17 +666,17 @@ void UMeshGroupPaintTool::OnPolyLassoFinished(const FCameraPolyLasso& Lasso, boo
 	// Try to clip polyline to be closed, or closed-enough for winding evaluation to work.
 	// If that returns false, the polyline is "too open". In that case we will extend
 	// outwards from the endpoints and then try to create a closed very large polygon
-	if (ApproxSelfClipPolyline(Polyline) == false)
+	if (UPolyLassoMarqueeMechanic::ApproximateSelfClipPolyline(Polyline) == false)
 	{
-		FVector2f StartDirOut = UE::Geometry::Normalized(Polyline[0] - Polyline[1]);
-		FLine2f StartLine(Polyline[0], StartDirOut);
-		FVector2f EndDirOut = UE::Geometry::Normalized(Polyline[N-1] - Polyline[N-2]);
-		FLine2f EndLine(Polyline[N-1], EndDirOut);
+		FVector2d StartDirOut = UE::Geometry::Normalized(Polyline[0] - Polyline[1]);
+		FLine2d StartLine(Polyline[0], StartDirOut);
+		FVector2d EndDirOut = UE::Geometry::Normalized(Polyline[N - 1] - Polyline[N - 2]);
+		FLine2d EndLine(Polyline[N - 1], EndDirOut);
 
 		// if we did not intersect, we are in ambiguous territory. Check if a segment along either end-direction
 		// intersects the polyline. If it does, we have something like a spiral and will be OK. 
 		// If not, make a closed polygon by interpolating outwards from each endpoint, and then in perp-directions.
-		FPolygon2f Polygon(Polyline);
+		FPolygon2d Polygon(Polyline);
 		float PerpSign = Polygon.IsClockwise() ? -1.0 : 1.0;
 
 		Polyline.Insert(StartLine.PointAt(10000.0f), 0);
@@ -836,7 +684,7 @@ void UMeshGroupPaintTool::OnPolyLassoFinished(const FCameraPolyLasso& Lasso, boo
 
 		Polyline.Add(EndLine.PointAt(10000.0f));
 		Polyline.Add(Polyline.Last() + 1000 * PerpSign * UE::Geometry::PerpCW(EndDirOut));
-		FVector2f StartPos = Polyline[0];
+		FVector2d StartPos = Polyline[0];
 		Polyline.Add(StartPos);		// close polyline (cannot use Polyline[0] in case Add resizes!)
 	}
 
@@ -850,14 +698,14 @@ void UMeshGroupPaintTool::OnPolyLassoFinished(const FCameraPolyLasso& Lasso, boo
 		if (Mesh->IsVertex(vid))
 		{
 			FVector3d WorldPos = CurTargetTransform.TransformPosition(Mesh->GetVertex(vid));
-			FVector2f PlanePos = (FVector2f)Lasso.GetProjectedPoint((FVector)WorldPos);
+			FVector2d PlanePos = (FVector2d)Lasso.GetProjectedPoint((FVector)WorldPos);
 
 			double WindingSum = 0;
-			FVector2f a = Polyline[0] - PlanePos, b = FVector2f::Zero();
+			FVector2d a = Polyline[0] - PlanePos, b = FVector2d::Zero();
 			for (int32 i = 1; i < N; ++i)
 			{
 				b = Polyline[i] - PlanePos;
-				WindingSum += (double)FMathf::Atan2(a.X*b.Y - a.Y*b.X, a.X*b.X + a.Y*b.Y);
+				WindingSum += (double)FMathd::Atan2(a.X*b.Y - a.Y*b.X, a.X*b.X + a.Y*b.Y);
 				a = b;
 			}
 			WindingSum /= FMathd::TwoPi;
diff --git a/Engine/Plugins/Experimental/MeshModelingToolsetExp/Source/MeshModelingToolsExp/Private/Sculpting/MeshSculptToolBase.cpp b/Engine/Plugins/Experimental/MeshModelingToolsetExp/Source/MeshModelingToolsExp/Private/Sculpting/MeshSculptToolBase.cpp
--- a/Engine/Plugins/Experimental/MeshModelingToolsetExp/Source/MeshModelingToolsExp/Private/Sculpting/MeshSculptToolBase.cpp
+++ b/Engine/Plugins/Experimental/MeshModelingToolsetExp/Source/MeshModelingToolsExp/Private/Sculpting/MeshSculptToolBase.cpp
@@ -736,22 +736,28 @@ void UMeshSculptToolBase::UpdateHoverStamp(const FFrame3d& StampFrameWorld)
 	HoverStamp.LocalFrame.Transform(CurTargetTransform.InverseUnsafe());
 }
 
+float UMeshSculptToolBase::GetStampTemporalFlowRate() const
+{
+	return BrushProperties->FlowRate;
+}
+
 void UMeshSculptToolBase::UpdateStampPendingState()
 {
 	if (InStroke() == false) return;
 
 	bool bFlowStampPending = false;
-	if (BrushProperties->FlowRate >= 1.0)
+	float UseStampFlowRate = GetStampTemporalFlowRate();
+	if (UseStampFlowRate >= 1.0)
 	{
 		bFlowStampPending = true;
 	}
-	else if (BrushProperties->FlowRate == 0.0)
+	else if (UseStampFlowRate == 0.0)
 	{
 		bFlowStampPending = (LastFlowTimeStamp++ == 0);
 	}
 	else
 	{
-		double dt = (1.0 - BrushProperties->FlowRate);
+		double dt = (1.0 - UseStampFlowRate);
 		int FlowTimestamp = (int)(ActiveStrokeTime / dt);
 		if (FlowTimestamp > LastFlowTimeStamp)
 		{
diff --git a/Engine/Plugins/Experimental/MeshModelingToolsetExp/Source/MeshModelingToolsExp/Public/Sculpting/MeshSculptToolBase.h b/Engine/Plugins/Experimental/MeshModelingToolsetExp/Source/MeshModelingToolsExp/Public/Sculpting/MeshSculptToolBase.h
--- a/Engine/Plugins/Experimental/MeshModelingToolsetExp/Source/MeshModelingToolsExp/Public/Sculpting/MeshSculptToolBase.h
+++ b/Engine/Plugins/Experimental/MeshModelingToolsetExp/Source/MeshModelingToolsExp/Public/Sculpting/MeshSculptToolBase.h
@@ -144,7 +144,6 @@ public:
 
 	UPROPERTY(meta = (TransientToolProperty))
 	bool bShowSpacing = true;
-
 };
 
 
@@ -485,6 +484,8 @@ protected:
 	virtual void UpdateHoverStamp(const FFrame3d& StampFrameWorld);
 	virtual bool IsStampPending() const { return bIsStampPending; }
 	virtual const FRay& GetPendingStampRayWorld() const { return PendingStampRay;  }
+	// Temporal Flow Rate defines the frequency of stamp placement. 1 is max rate, 0 is no stamps. Defaults to BrushProperties->FlowRate, but subclasses can re-use that setting for other things
+	virtual float GetStampTemporalFlowRate() const;
 
 
 	//
diff --git a/Engine/Plugins/Runtime/MeshModelingToolset/Source/ModelingComponents/Private/Mechanics/PolyLassoMarqueeMechanic.cpp b/Engine/Plugins/Runtime/MeshModelingToolset/Source/ModelingComponents/Private/Mechanics/PolyLassoMarqueeMechanic.cpp
--- a/Engine/Plugins/Runtime/MeshModelingToolset/Source/ModelingComponents/Private/Mechanics/PolyLassoMarqueeMechanic.cpp
+++ b/Engine/Plugins/Runtime/MeshModelingToolset/Source/ModelingComponents/Private/Mechanics/PolyLassoMarqueeMechanic.cpp
@@ -8,6 +8,13 @@
 #include "ToolSceneQueriesUtil.h"
 #include "TransformTypes.h"
 
+#include "SegmentTypes.h"
+#include "LineTypes.h"
+#include "Intersection/IntrLine2Line2.h"
+#include "Intersection/IntrSegment2Segment2.h"
+#include "Async/ParallelFor.h"
+
+
 #include UE_INLINE_GENERATED_CPP_BY_NAME(PolyLassoMarqueeMechanic)
 
 using namespace UE::Geometry;
@@ -308,5 +315,156 @@ void UPolyLassoMarqueeMechanic::DrawHUD(FCanvas* Canvas, IToolsContextRenderAPI*
 	}
 }
 
+
+
+
+namespace UELocal
+{
+
+
+static bool FindPolylineSelfIntersection(
+	const TArray<FVector2D>& Polyline, 
+	FVector2D& IntersectionPointOut, 
+	FIndex2i& IntersectionIndexOut,
+	bool bParallel = true)
+{
+	int32 N = Polyline.Num();
+	std::atomic<bool> bSelfIntersects(false);
+	ParallelFor(N - 1, [&](int32 i)
+	{
+		FSegment2d SegA(Polyline[i], Polyline[i + 1]);
+		for (int32 j = i + 2; j < N - 1 && bSelfIntersects == false; ++j)
+		{
+			FSegment2d SegB(Polyline[j], Polyline[j + 1]);
+			if (SegA.Intersects(SegB) && bSelfIntersects == false)		
+			{
+				bool ExpectedValue = false;
+				if (std::atomic_compare_exchange_strong(&bSelfIntersects, &ExpectedValue, true))
+				{
+					FIntrSegment2Segment2d Intersection(SegA, SegB);
+					Intersection.Find();
+					IntersectionPointOut = Intersection.Point0;
+					IntersectionIndexOut = FIndex2i(i, j);
+					return;
+				}
+			}
+		}
+	}, (bParallel) ? EParallelForFlags::None : EParallelForFlags::ForceSingleThread );
+
+	return bSelfIntersects;
+}
+
+
+
+static bool FindPolylineSegmentIntersection(
+	const TArray<FVector2D>& Polyline,
+	const UE::Geometry::FSegment2d& Segment,
+	FVector2D& IntersectionPointOut,
+	int& IntersectionIndexOut)
+{
+	int32 N = Polyline.Num();
+	for (int32 i = 0; i < N - 1; ++i)
+	{
+		FSegment2d PolySeg(Polyline[i], Polyline[i + 1]);
+		if (Segment.Intersects(PolySeg))
+		{
+			FIntrSegment2Segment2d Intersection(Segment, PolySeg);
+			Intersection.Find();
+			IntersectionPointOut = Intersection.Point0;
+			IntersectionIndexOut = i;
+			return true;
+		}
+	}
+	return false;
+}
+
+}
+
+
+bool UPolyLassoMarqueeMechanic::ApproximateSelfClipPolyline(TArray<FVector2D>& Polyline)
+{
+	int32 N = Polyline.Num();
+
+	// handle already-closed polylines
+	if (Distance(Polyline[0], Polyline[N - 1]) < 0.0001)
+	{
+		return true;
+	}
+
+	FVector2d IntersectPoint;
+	FIndex2i IntersectionIndex(-1, -1);
+	bool bSelfIntersects = UELocal::FindPolylineSelfIntersection(Polyline, IntersectPoint, IntersectionIndex);
+	if (bSelfIntersects)
+	{
+		TArray<FVector2D> NewPolyline;
+		NewPolyline.Add(IntersectPoint);
+		for (int32 i = IntersectionIndex.A; i <= IntersectionIndex.B; ++i)
+		{
+			NewPolyline.Add(Polyline[i]);
+		}
+		NewPolyline.Add(IntersectPoint);
+		Polyline = MoveTemp(NewPolyline);
+		return true;
+	}
+
+
+	FVector2d StartDirOut = UE::Geometry::Normalized(Polyline[0] - Polyline[1]);
+	FLine2d StartLine(Polyline[0], StartDirOut);
+	FVector2d EndDirOut = UE::Geometry::Normalized(Polyline[N - 1] - Polyline[N - 2]);
+	FLine2d EndLine(Polyline[N - 1], EndDirOut);
+	FIntrLine2Line2d LineIntr(StartLine, EndLine);
+	bool bIntersects = false;
+	if (LineIntr.Find())
+	{
+		bIntersects = LineIntr.IsSimpleIntersection() && (LineIntr.Segment1Parameter > 0) && (LineIntr.Segment2Parameter > 0);
+		if (bIntersects)
+		{
+			Polyline.Add(StartLine.PointAt(LineIntr.Segment1Parameter));
+			Polyline.Add(StartLine.Origin);
+			return true;
+		}
+	}
+
+
+	FAxisAlignedBox2d Bounds;
+	for (const FVector2d& P : Polyline)
+	{
+		Bounds.Contain(P);
+	}
+	double Size = Bounds.DiagonalLength();
+
+	FVector2d StartPos = Polyline[0] + 0.001 * StartDirOut;
+	if (UELocal::FindPolylineSegmentIntersection(Polyline, FSegment2d(StartPos, StartPos + 2 * Size * StartDirOut), IntersectPoint, IntersectionIndex.A))
+	{
+		//TArray<FVector2f> NewPolyline;
+		//for (int32 i = 0; i <= IntersectionIndex.A; ++i)
+		//{
+		//	NewPolyline.Add(Polyline[i]);
+		//}
+		//NewPolyline.Add(IntersectPoint);
+		//NewPolyline.Add(Polyline[0]);
+		//Polyline = MoveTemp(NewPolyline);
+		return true;
+	}
+
+	FVector2d EndPos = Polyline[N - 1] + 0.001 * EndDirOut;
+	if (UELocal::FindPolylineSegmentIntersection(Polyline, FSegment2d(EndPos, EndPos + 2 * Size * EndDirOut), IntersectPoint, IntersectionIndex.A))
+	{
+		//TArray<FVector2f> NewPolyline;
+		//NewPolyline.Add(IntersectPoint);
+		//for (int32 i = IntersectionIndex.A+1; i < N; ++i)
+		//{
+		//	NewPolyline.Add(Polyline[i]);
+		//}
+		//NewPolyline.Add(Polyline[0]);
+		//NewPolyline.Add(IntersectPoint);
+		//Polyline = MoveTemp(NewPolyline);
+		return true;
+	}
+
+	return false;
+}
+
+
 #undef LOCTEXT_NAMESPACE
 
diff --git a/Engine/Plugins/Runtime/MeshModelingToolset/Source/ModelingComponents/Public/Mechanics/PolyLassoMarqueeMechanic.h b/Engine/Plugins/Runtime/MeshModelingToolset/Source/ModelingComponents/Public/Mechanics/PolyLassoMarqueeMechanic.h
--- a/Engine/Plugins/Runtime/MeshModelingToolset/Source/ModelingComponents/Public/Mechanics/PolyLassoMarqueeMechanic.h
+++ b/Engine/Plugins/Runtime/MeshModelingToolset/Source/ModelingComponents/Public/Mechanics/PolyLassoMarqueeMechanic.h
@@ -181,4 +181,9 @@ private:
 
 	FCameraPolyLasso CurrentPolyLasso;
 	bool WillCloseCurrentLasso(const FInputDeviceRay& DevicePos) const;
+
+
+public:
+	// utility function clients can use to clean up poly-lasso if it is not closed (perhaps should be folded into mechanic behavior)
+	static bool ApproximateSelfClipPolyline(TArray<FVector2D>& Polyline);
 };
diff --git a/Engine/Source/Runtime/GeometryCore/Public/VectorTypes.h b/Engine/Source/Runtime/GeometryCore/Public/VectorTypes.h
--- a/Engine/Source/Runtime/GeometryCore/Public/VectorTypes.h
+++ b/Engine/Source/Runtime/GeometryCore/Public/VectorTypes.h
@@ -405,6 +405,17 @@ constexpr FColor ToFColor(const UE::Math::TVector4<T>& Vector)
 }
 
 
+template<typename T>
+UE::Math::TVector4<T> Lerp(const UE::Math::TVector4<T>& A, const UE::Math::TVector4<T>& B, T Alpha)
+{
+	T OneMinusAlpha = (T)1 - Alpha;
+	return UE::Math::TVector4<T>(
+		OneMinusAlpha * A.X + Alpha * B.X,
+		OneMinusAlpha * A.Y + Alpha * B.Y,
+		OneMinusAlpha * A.Z + Alpha * B.Z,
+		OneMinusAlpha * A.W + Alpha * B.W);
+}
+
 } // end namespace UE::Geometry
 
 namespace Math
