CombineMeshInstances: Add support for 2D retriangulation of planar areas, after coplanar face merging. This can further reduce triangle count as the planar-simplification applied after coplanar merging often cannot remove vertices on region borders, limiting it's effectiveness.
Also tweaked geometric tolerances recently added to RemoveHiddenFaces_ExteriorVisibilit() as they seem to be too aggressive (need to add real-world dimensions here).
Updated internal usage to enable this new capability
#rnx

diff --git a/Engine/Plugins/Experimental/MeshModelingToolsetExp/Source/GeometryProcessingAdapters/Private/CombineMeshInstancesImpl.cpp b/Engine/Plugins/Experimental/MeshModelingToolsetExp/Source/GeometryProcessingAdapters/Private/CombineMeshInstancesImpl.cpp
--- a/Engine/Plugins/Experimental/MeshModelingToolsetExp/Source/GeometryProcessingAdapters/Private/CombineMeshInstancesImpl.cpp
+++ b/Engine/Plugins/Experimental/MeshModelingToolsetExp/Source/GeometryProcessingAdapters/Private/CombineMeshInstancesImpl.cpp
@@ -1421,8 +1421,8 @@ static void RemoveHiddenFaces_ExteriorVisibility(
 
 	// geometric magic numbers used below that have been slightly tuned...
 	double GlancingAngleDotTolerance = FMathd::Cos(85.0 * FMathd::DegToRad);
-	const double TriScalingAlpha = 0.95;
-	const double BaryCoordsThreshold = 0.01;
+	const double TriScalingAlpha = 0.999;
+	const double BaryCoordsThreshold = 0.001;
 
 
 	auto FindHitTriangleTest = [&](FVector3d TargetPosition, FVector3d TargetNormal, FVector3d FarPosition) -> int
@@ -1643,6 +1643,273 @@ struct FMergeTriInfo
 };
 
 
+// Assuming SourcePartMesh is epsilon-planar, it's border polygons can be projected to a plane and
+// remeshed using 2D triangulation to get the minimal triangle count. And once it is polygons,
+// they can be boolean-unioned, topologically-closed, small holes can be removed, etc
+// TODO: this code is very similar to ComputeSweptSolidApproximation and it would be nice if they could be combined
+static void ComputePlanarPolygonApproximation(
+	const FDynamicMesh3& SourcePartMesh,
+	FDynamicMesh3& NewPlanarMesh,
+	FVector3d Direction,
+	double MergeOffset = 0.1,
+	double SimplifyTolerance = 1.0,
+	double MinHoleArea = 1.0)
+{
+	check(SourcePartMesh.IsCompactT());
+	FFrame3d ProjectFrame(SourcePartMesh.GetTriCentroid(0), Direction);
+
+	FMeshBoundaryLoops Loops(&SourcePartMesh);
+	FPlanarComplexd PlanarComplex;
+	for (FEdgeLoop& Loop : Loops.Loops)
+	{
+		TArray<FVector3d> Vertices;
+		Loop.GetVertices<FVector3d>(Vertices);
+		FPolygon2d Polygon;
+		for (FVector3d V : Vertices)
+		{
+			FVector LocalV = ProjectFrame.ToFramePoint(V);
+			Polygon.AppendVertex( FVector2d(LocalV.X, LocalV.Y) );
+		}
+		Polygon.Reverse();		// mesh orientation comes out backwards...
+		PlanarComplex.Polygons.Add(MoveTemp(Polygon));
+	}
+	PlanarComplex.bTrustOrientations = true;		// have to do this or overlapping projections will create holes
+	PlanarComplex.FindSolidRegions();
+	TArray<FGeneralPolygon2d> Polygons = PlanarComplex.ConvertOutputToGeneralPolygons();
+
+	if (Polygons.Num() == 0)
+	{
+		NewPlanarMesh = SourcePartMesh;
+		return;
+	}
+
+	double UnionMergeOffset = 0.1;
+	if (Polygons.Num() > 1)
+	{
+		// nudge all polygons outwards to ensure that when we boolean union exactly-coincident polygons
+		// they intersect a bit, otherwise we may end up with zero-area cracks/holes
+		if (UnionMergeOffset > 0)
+		{
+			for (FGeneralPolygon2d& Polygon : Polygons)
+			{
+				Polygon.VtxNormalOffset(UnionMergeOffset);
+			}
+		}
+
+		TArray<FGeneralPolygon2d> ResultPolygons;
+		PolygonsUnion(Polygons, ResultPolygons, true);
+		Polygons = MoveTemp(ResultPolygons);
+
+		if (UnionMergeOffset > 0)
+		{
+			for (FGeneralPolygon2d& Polygon : Polygons)
+			{
+				Polygon.VtxNormalOffset(-UnionMergeOffset);	// undo offset
+			}
+		}
+	}
+
+	// can optionally try to reduce polygon complexity by topological closure (dilate/erode)
+	if (MergeOffset > 0)
+	{
+		TArray<FGeneralPolygon2d> TmpPolygons;
+		PolygonsOffsets(MergeOffset, -MergeOffset,
+			Polygons, TmpPolygons, true, 1.0,
+			EPolygonOffsetJoinType::Square,
+			EPolygonOffsetEndType::Polygon);
+
+		Polygons = MoveTemp(TmpPolygons);
+	}
+
+	FConstrainedDelaunay2d Triangulator;
+	for (FGeneralPolygon2d& Polygon : Polygons)
+	{
+		if (SimplifyTolerance > 0)
+		{
+			Polygon.Simplify(SimplifyTolerance, SimplifyTolerance * 0.25);		// 0.25 is kind of arbitrary here...
+		}
+		if (MinHoleArea > 0)
+		{
+			Polygon.FilterHoles([&](const FPolygon2d& HolePoly) { return HolePoly.Area() < MinHoleArea; });
+		}
+		Triangulator.Add(Polygon);
+	}
+
+	Triangulator.Triangulate([&Polygons](const TArray<FVector2d>& Vertices, FIndex3i Tri)
+	{
+		FVector2d Point = (Vertices[Tri.A] + Vertices[Tri.B] + Vertices[Tri.C]) / 3.0;
+		for (const FGeneralPolygon2d& Polygon : Polygons)
+		{
+			if (Polygon.Contains(Point))
+			{
+				return true;
+			}
+		}
+		return false;
+	});
+
+	FFlatTriangulationMeshGenerator TriangulationMeshGen;
+	TriangulationMeshGen.Vertices2D = Triangulator.Vertices;
+	TriangulationMeshGen.Triangles2D = Triangulator.Triangles;
+	FDynamicMesh3 PolygonsMesh(&TriangulationMeshGen.Generate());
+
+	if (PolygonsMesh.TriangleCount() < 2)
+	{
+		NewPlanarMesh = SourcePartMesh;
+		return;
+	}
+
+	MeshTransforms::FrameCoordsToWorld(PolygonsMesh, ProjectFrame);
+
+	NewPlanarMesh = MoveTemp(PolygonsMesh);
+}
+
+
+
+// Find sets of triangles that lie in the same 3D plane on TargetMesh,
+// and then extract out those areas, pull out the boundary polygons,
+// union them together, and do a 2D polygon-with-holes triangulation.
+static void RetriangulatePlanarFacePolygons(FDynamicMesh3& TargetMesh)
+{
+	TArray<FFrame3d> PlaneSet;
+
+	const double AngleDotTol = 0.99;
+	// this distance probably should not be hardcoded...
+	const double DistanceTol = 0.1;
+
+	if (TargetMesh.IsCompactT() == false)
+	{
+		TargetMesh.CompactInPlace();
+	}
+
+	TArray<int32> TriPlaneID;
+	TriPlaneID.Init(-1, TargetMesh.MaxTriangleID());
+
+	auto SamePlaneCheck = [&](FFrame3d& Frame, FVector3d TriNormal, FVector3d TriCentroid)
+	{
+		FVector3d PlaneNormal = Frame.Z();
+		if (PlaneNormal.Dot(TriNormal) < AngleDotTol) return false;
+		FVector3d LocalVec = TriCentroid - Frame.Origin;
+		double SignedDist = LocalVec.Dot(PlaneNormal);
+		if (FMathd::Abs(SignedDist) > DistanceTol) return false;
+		return true;
+	};
+
+	// accumulate set of unique planes. Would be nice if plane search could be
+	// done more efficiently but there is not an obvious spatial data structure...
+	for (int32 tid : TargetMesh.TriangleIndicesItr())
+	{
+		FVector3d TriNormal, Centroid; double Area;
+		TargetMesh.GetTriInfo(tid, TriNormal, Area, Centroid);
+		FFrame3d TriPlane(Centroid, TriNormal);
+
+		bool bFound = false;
+		for (int32 k = 0; k < PlaneSet.Num() && bFound == false; ++k)
+		{
+			if (SamePlaneCheck(PlaneSet[k], TriNormal, Centroid))
+			{
+				TriPlaneID[tid] = k;
+				bFound = true;
+			}
+		}
+
+		if (bFound == false)
+		{
+			TriPlaneID[tid] = PlaneSet.Num();
+			PlaneSet.Add(TriPlane);
+		}
+	}
+	if (PlaneSet.Num() < 2)
+	{
+		return;
+	}
+
+	// if we have vertex colors, transfer them to new meshes by finding a value at nearest vertex
+	FDynamicMeshColorOverlay* SourceColors = TargetMesh.Attributes()->PrimaryColors();
+	FDynamicMeshAABBTree3 TargetMeshSpatial(&TargetMesh, (SourceColors != nullptr));
+
+	// TODO: we do not actually have to split here, we can just send the triangle ROI to ComputePlanarPolygonApproximation and
+	// it can use RegionBoundaryLoops. Then can delete and append new mesh if it is better (although messier to parallelize)
+	TArray<FDynamicMesh3> SplitMeshes;
+	FDynamicMeshEditor::SplitMesh(&TargetMesh, SplitMeshes, [&](int32 tid) { return TriPlaneID[tid]; });
+
+	for (FDynamicMesh3& Mesh : SplitMeshes)		// ought to be trivially parallelizable...
+	{
+		if (Mesh.TriangleCount() <= 2) continue;
+
+		const double MergeOffset = 1.0;
+		const double SimplifyTolerance = 1.0;
+		const double MinHoleArea = 10.0;
+
+		FDynamicMesh3 NewPlanarMesh;
+		ComputePlanarPolygonApproximation(Mesh, NewPlanarMesh,
+			Mesh.GetTriNormal(0), MergeOffset, SimplifyTolerance, MinHoleArea);
+
+		if (NewPlanarMesh.TriangleCount() < Mesh.TriangleCount())
+		{
+			Mesh = MoveTemp(NewPlanarMesh);
+
+			Mesh.EnableAttributes();
+			Mesh.Attributes()->SetNumUVLayers(0);
+			FMeshNormals::InitializeOverlayToPerVertexNormals(Mesh.Attributes()->PrimaryNormals(), false);
+
+			// project source colors to new mesh vertices
+			if (SourceColors)
+			{
+				Mesh.Attributes()->EnablePrimaryColors();
+				FDynamicMeshColorOverlay* SetColors = Mesh.Attributes()->PrimaryColors();
+				TArray<int32> VertexToElementMap; 
+				VertexToElementMap.SetNum(Mesh.MaxVertexID());
+				for (int32 vid : Mesh.VertexIndicesItr())
+				{
+					double NearestDistSqr = TNumericLimits<double>::Max();
+					int32 TargetVID = TargetMeshSpatial.FindNearestVertex(Mesh.GetVertex(vid), NearestDistSqr);
+					FVector4f UseColor = FVector4f::Zero();
+					SourceColors->EnumerateVertexElements(TargetVID, [&](int TriID, int ElemID, const FVector4f& ElemColor) { UseColor = ElemColor; return false; }, false);
+					VertexToElementMap[vid] = SetColors->AppendElement(UseColor);
+				}
+				for (int32 tid : Mesh.TriangleIndicesItr())
+				{
+					FIndex3i Triangle = Mesh.GetTriangle(tid);
+					SetColors->SetTriangle(tid, FIndex3i(VertexToElementMap[Triangle.A], VertexToElementMap[Triangle.B], VertexToElementMap[Triangle.C]));
+				}
+			}
+		}
+	}
+
+	FDynamicMesh3 NewMesh;
+	NewMesh.EnableMatchingAttributes(TargetMesh);
+
+	FDynamicMeshEditor Editor(&NewMesh);
+	for (FDynamicMesh3& Mesh : SplitMeshes)
+	{
+		FMeshIndexMappings Mappings;
+		Editor.AppendMesh(&Mesh, Mappings);
+	}
+
+	if (NewMesh.TriangleCount() == 0)
+	{
+		return;
+	}
+
+	// currently assuming input mesh has been split by MaterialID, so no projection here,
+	// just use any MaterialID
+	if (const FDynamicMeshMaterialAttribute* SourceMaterialIDs = TargetMesh.Attributes()->GetMaterialID())
+	{
+		if (FDynamicMeshMaterialAttribute* TargetMaterialIDs = NewMesh.Attributes()->GetMaterialID())
+		{
+			int32 ConstantMaterialID = SourceMaterialIDs->GetValue(0);
+			for (int32 tid : NewMesh.TriangleIndicesItr())
+			{
+				TargetMaterialIDs->SetValue(tid, ConstantMaterialID);
+			}
+		}
+	}
+
+	TargetMesh = MoveTemp(NewMesh);
+}
+
+
 //
 // After hidden face removal, a mesh can often be optimized to at least save some vertices (by welding open borders),
 // and then in some cases, now-connected triangle areas can be retriangulated to require fewer triangles.
@@ -1656,6 +1923,7 @@ static void PostProcessHiddenFaceRemovedMesh(
 	FDynamicMesh3& TargetMesh, 
 	double Tolerance,
 	bool bTryToMergeFaces,
+	bool bApplyPlanarRetriangulation,
 	TFunctionRef<FIndex3i(const FDynamicMesh3& Mesh, int32 TriangleID)> GetTriangleGroupingIDFunc)
 {
 	bool bVerbose = CVarGeometryCombineMeshInstancesVerbose.GetValueOnAnyThread();
@@ -1751,6 +2019,13 @@ static void PostProcessHiddenFaceRemovedMesh(
 
 		Simplifier.CollapseMode = FQEMSimplification::ESimplificationCollapseModes::AverageVertexPosition;
 		Simplifier.SimplifyToMinimalPlanar( 0.01 );
+
+		// minimal-planar simplification is good, but for planar areas, we can go further and extract 2D
+		// polygons that can be delaunay-triangulated. This is currently somewhat expensive...
+		if (bApplyPlanarRetriangulation)
+		{
+			RetriangulatePlanarFacePolygons(SubRegionMesh);
+		}
 	}
 
 	TargetMesh.Clear();
@@ -1777,6 +2052,14 @@ static void PostProcessHiddenFaceRemovedMesh(
 		Welder.Apply();
 	}
 
+
+	// make sure we have necessary attribute sets
+	if (TargetMesh.Attributes()->NumUVLayers() == 0)
+	{
+		TargetMesh.Attributes()->SetNumUVLayers(1);
+	}
+
+
 	TargetMesh.CompactInPlace();
 
 	if ( bVerbose )
@@ -2217,6 +2500,7 @@ void ComputeHiddenRemovalForLOD(
 		PostProcessHiddenFaceRemovedMesh(MeshLOD, 
 			CombineOptions.SimplifyBaseTolerance,
 			CombineOptions.bMergeCoplanarFaces && LODIndex >= CombineOptions.MergeCoplanarFacesStartLOD,
+			CombineOptions.PlanarPolygonRetriangulationStartLOD >= 0 && LODIndex >= CombineOptions.PlanarPolygonRetriangulationStartLOD,
 			GroupingIDFunc );
 	}
 }
diff --git a/Engine/Source/Developer/GeometryProcessingInterfaces/Public/GeometryProcessingInterfaces/CombineMeshInstances.h b/Engine/Source/Developer/GeometryProcessingInterfaces/Public/GeometryProcessingInterfaces/CombineMeshInstances.h
--- a/Engine/Source/Developer/GeometryProcessingInterfaces/Public/GeometryProcessingInterfaces/CombineMeshInstances.h
+++ b/Engine/Source/Developer/GeometryProcessingInterfaces/Public/GeometryProcessingInterfaces/CombineMeshInstances.h
@@ -152,11 +152,21 @@ public:
 		// opening angle used to detect/assign sharp edges
 		double HardNormalAngleDeg = 15.0;
 
-
+		// Attempt to merge/weld coplanar areas after hidden removal, and then further simplify those merged areas
+		// Coplanar merging is never applied between areas with different Materials. 
+		// TriangleGroupingIDFunc below can be used to further demarcate important internal shape boundaries.
 		bool bMergeCoplanarFaces = true;
+		// LOD level at which coplanar face merging is applied. 
 		int32 MergeCoplanarFacesStartLOD = 1;
+		// TriangleGroupingIDFunc allows client to specify external adjacency relationships between triangles via an integer ID tuple.
+		// Adjacent triangles will only be considered for Coplanar Merging if they have the same FIndex3i ID.
 		TFunction<UE::Geometry::FIndex3i(const UE::Geometry::FDynamicMesh3& Mesh, int32 TriangleID)> TriangleGroupingIDFunc;
 
+		// LOD level to attempt extraction and retriangulation of detected planar polygonal areas, after removing hidden faces.
+		// This occurs after coplanar face merging, and may further merge adjacent coplanar faces.
+		// Can significantly reduce triangle count, but attributes on the interiors of polygonal areas will be completely discarded
+		int32 PlanarPolygonRetriangulationStartLOD = -1;
+
 		// If enabled, attempt to retriangulate planar areas of Source LODs to remove redundant coplanar geometry
 		bool bRetriangulateSourceLODs = true;
 		// which Source LOD to start retriangulating at
