ModelingTools: add MeshStatistics display to DisplaceMeshTool. Remove AdvancedDisplay from ProjectionDistance and bProjectionInWorldSpace in Bake tools, these are critical properties and should not be hidden-by-default.

diff --git a/Engine/Plugins/Experimental/MeshModelingToolsetExp/Source/MeshModelingToolsExp/Private/DisplaceMeshTool.cpp b/Engine/Plugins/Experimental/MeshModelingToolsetExp/Source/MeshModelingToolsExp/Private/DisplaceMeshTool.cpp
--- a/Engine/Plugins/Experimental/MeshModelingToolsetExp/Source/MeshModelingToolsExp/Private/DisplaceMeshTool.cpp
+++ b/Engine/Plugins/Experimental/MeshModelingToolsetExp/Source/MeshModelingToolsExp/Private/DisplaceMeshTool.cpp
@@ -21,6 +21,7 @@
 #include "Operations/UniformTessellate.h"
 #include "Operations/SelectiveTessellate.h"
 #include "Materials/MaterialInterface.h"
+#include "Properties/MeshStatisticsProperties.h"
 
 // needed to disable normals recalculation on the underlying asset
 #include "AssetUtils/MeshDescriptionUtil.h"
@@ -1062,6 +1063,10 @@ void UDisplaceMeshTool::Setup()
 
 	Subdivider = MakeUnique<FSubdivideMeshOpFactory>(OriginalMesh, SubParameters, CommonProperties->SubdivisionType);
 
+	MeshStatistics = NewObject<UMeshStatisticsProperties>(this);
+	MeshStatistics->Update(OriginalMesh);
+	AddToolPropertySource(MeshStatistics);
+
 	StartComputation();
 
 	SetToolDisplayName(LOCTEXT("ToolName", "Displace"));
@@ -1420,6 +1425,8 @@ void UDisplaceMeshTool::AdvanceComputation()
 		SubdividedMesh = TSharedPtr<FDynamicMesh3, ESPMode::ThreadSafe>(SubOp->ExtractResult().Release());
 		VerticesToDisplace = TSharedPtr<TArray<int>, ESPMode::ThreadSafe>(SubOpDownCast->ExtractVerticesToDisplace().Release());
 
+		MeshStatistics->Update(*SubdividedMesh);
+
 		delete SubdivideTask;
 		SubdivideTask = nullptr;
 	}
@@ -1437,6 +1444,8 @@ void UDisplaceMeshTool::AdvanceComputation()
 	if (DisplaceTask && DisplaceTask->IsDone())
 	{
 		TUniquePtr<FDynamicMesh3> DisplacedMesh = DisplaceTask->GetTask().ExtractOperator()->ExtractResult();
+		MeshStatistics->Update(*SubdividedMesh);
+
 		delete DisplaceTask;
 		DisplaceTask = nullptr;
 		DynamicMeshComponent->ClearOverrideRenderMaterial();
diff --git a/Engine/Plugins/Experimental/MeshModelingToolsetExp/Source/MeshModelingToolsExp/Public/BakeMeshAttributeToolCommon.h b/Engine/Plugins/Experimental/MeshModelingToolsetExp/Source/MeshModelingToolsExp/Public/BakeMeshAttributeToolCommon.h
--- a/Engine/Plugins/Experimental/MeshModelingToolsetExp/Source/MeshModelingToolsExp/Public/BakeMeshAttributeToolCommon.h
+++ b/Engine/Plugins/Experimental/MeshModelingToolsetExp/Source/MeshModelingToolsExp/Public/BakeMeshAttributeToolCommon.h
@@ -103,13 +103,13 @@ public:
 
 	/** Maximum allowed distance for the projection from target mesh to source mesh for the sample to be considered valid.
 	 * This is only relevant if a separate source mesh is provided. */
-	UPROPERTY(EditAnywhere, Category = BakeInput, AdvancedDisplay, meta = (ClampMin = "0.001",
+	UPROPERTY(EditAnywhere, Category = BakeInput, meta = (ClampMin = "0.001",
 		EditCondition = "SourceStaticMesh != nullptr || SourceSkeletalMesh != nullptr || SourceDynamicMesh != nullptr", HideEditConditionToggle))
 	float ProjectionDistance = 3.0;
 
 	/** If true, uses the world space positions for the projection from target mesh to source mesh, otherwise it uses their object space positions.
 	 * This is only relevant if a separate source mesh is provided. */
-	UPROPERTY(EditAnywhere, Category = BakeInput, AdvancedDisplay, meta = (
+	UPROPERTY(EditAnywhere, Category = BakeInput, meta = (
 		EditCondition = "SourceStaticMesh != nullptr || SourceSkeletalMesh != nullptr || SourceDynamicMesh != nullptr", HideEditConditionToggle))
 	bool bProjectionInWorldSpace = false;
 
diff --git a/Engine/Plugins/Experimental/MeshModelingToolsetExp/Source/MeshModelingToolsExp/Public/DisplaceMeshTool.h b/Engine/Plugins/Experimental/MeshModelingToolsetExp/Source/MeshModelingToolsExp/Public/DisplaceMeshTool.h
--- a/Engine/Plugins/Experimental/MeshModelingToolsetExp/Source/MeshModelingToolsExp/Public/DisplaceMeshTool.h
+++ b/Engine/Plugins/Experimental/MeshModelingToolsetExp/Source/MeshModelingToolsExp/Public/DisplaceMeshTool.h
@@ -18,6 +18,7 @@
 struct FMeshDescription;
 class UCurveFloat;
 class UDynamicMeshComponent;
+class UMeshStatisticsProperties;
 
 UENUM()
 enum class EDisplaceMeshToolDisplaceType : uint8
@@ -342,6 +343,10 @@ public:
 	UPROPERTY()
 	TObjectPtr<UCurveFloat> ActiveContrastCurveTarget = nullptr;
 
+	UPROPERTY()
+	TObjectPtr<UMeshStatisticsProperties> MeshStatistics;
+
+
 private:
 	void StartComputation();
 	void AdvanceComputation();
