ModelingTools: add new Vertex Painting Tool

Add UMeshVertexPaintTool, subclass of UMeshSculptToolBase. Paints Mesh Vertex Colors (ie Asset colors, not instance colors). Provides many interactions similar to GroupPaintTool (ie brush paint, click-to-fill, poly-lasso, volumetric vs connected brush ROI, UV/Normal seams, visibility filter, etc. Supports painting split-vertex colors, or smooth per-vertex colors (bHardEdges setting). Supports various brush types (paint, smooth, "blend" = average split values). Painting supports custom color-blend modes (ie "over" mode, does stroke compositing like pixel painting tools). Can show vertex color under cursor. Supports showing vertex color or original materials. Ability to filter channels for visibility/painting. Actions buttons to fill current-color/black/white. "Utility" panel for common operations (blend all seams, fill channels, invert channels, copy channel to channel, swap channels, copy from weight map, copy to all/single LOD).

Internally supports seam painting by splitting all colors on startup (ie to create per-triangle-vertex color) and then welding on shutdown.

Current auto-RGBA conversions are inverted, internally the tool works in Linear colors, preview DynamicMeshComponent is configured to convert to SRGB in color buffers (similar to static mesh).

Vertex color brush ops added in MeshVertexPaintBrushOps.h, deriving from new FMeshVertexColorBrushOp. Various color-blending functions implemented there, should be moved to centralized place eventually.

Expose MeshVertexPaintTool in ModelingMode
Add FModelingToolsColorChannelFilter, for toggling R/G/B/A filtering state, and a details customization that shows these four checkboxes in a horizontal row. Use in MeshVertexPaintTool, Register in Modeling Mode.

diff --git a/Engine/Plugins/Editor/ModelingToolsEditorMode/Source/ModelingToolsEditorMode/Private/DetailsCustomizations/MeshVertexSculptToolCustomizations.cpp b/Engine/Plugins/Editor/ModelingToolsEditorMode/Source/ModelingToolsEditorMode/Private/DetailsCustomizations/MeshVertexSculptToolCustomizations.cpp
--- a/Engine/Plugins/Editor/ModelingToolsEditorMode/Source/ModelingToolsEditorMode/Private/DetailsCustomizations/MeshVertexSculptToolCustomizations.cpp
+++ b/Engine/Plugins/Editor/ModelingToolsEditorMode/Source/ModelingToolsEditorMode/Private/DetailsCustomizations/MeshVertexSculptToolCustomizations.cpp
@@ -30,6 +30,11 @@ TSharedRef<IDetailCustomization> FSculptBrushPropertiesDetails::MakeInstance()
 
 void FSculptBrushPropertiesDetails::CustomizeDetails(IDetailLayoutBuilder& DetailBuilder)
 {
+	TArray<TWeakObjectPtr<UObject>> ObjectsBeingCustomized;
+	DetailBuilder.GetObjectsBeingCustomized(ObjectsBeingCustomized);
+	check(ObjectsBeingCustomized.Num() > 0);
+	USculptBrushProperties* BrushProperties = Cast<USculptBrushProperties>(ObjectsBeingCustomized[0]);
+
 	TSharedPtr<IPropertyHandle> FlowRateHandle = DetailBuilder.GetProperty(GET_MEMBER_NAME_CHECKED(USculptBrushProperties, FlowRate), USculptBrushProperties::StaticClass());
 	ensure(FlowRateHandle->IsValidHandle());
 
@@ -37,20 +42,43 @@ void FSculptBrushPropertiesDetails::CustomizeDetails(IDetailLayoutBuilder& Detai
 	ensure(SpacingHandle->IsValidHandle());
 	SpacingHandle->MarkHiddenByCustomization();
 
+	TSharedPtr<IPropertyHandle> FalloffHandle = DetailBuilder.GetProperty(GET_MEMBER_NAME_CHECKED(USculptBrushProperties, BrushFalloffAmount), USculptBrushProperties::StaticClass());
+	ensure(FalloffHandle->IsValidHandle());
+	FalloffHandle->MarkHiddenByCustomization();
+
 	TSharedPtr<SDynamicNumericEntry::FDataSource> FlowRateSource = SDynamicNumericEntry::MakeSimpleDataSource(
 		FlowRateHandle, TInterval<float>(0.0f, 1.0f), TInterval<float>(0.0f, 1.0f));
-	TSharedPtr<SDynamicNumericEntry::FDataSource> SpacingSource = SDynamicNumericEntry::MakeSimpleDataSource(
-		SpacingHandle, TInterval<float>(0.0f, 1000.0f), TInterval<float>(0.0f, 4.0f));
 
-	DetailBuilder.EditDefaultProperty(FlowRateHandle)->CustomWidget()
-		.OverrideResetToDefault(FResetToDefaultOverride::Hide())
-		.WholeRowContent()
-	[
-		MakeTwoWidgetDetailRowHBox(
-			MakeFixedWidthLabelSliderHBox(FlowRateHandle, FlowRateSource, FSculptToolsUIConstants::SculptShortLabelWidth),
-			MakeFixedWidthLabelSliderHBox(SpacingHandle, SpacingSource, FSculptToolsUIConstants::SculptShortLabelWidth)
-		)
-	];
+	if ( BrushProperties->bShowSpacing )
+	{
+		TSharedPtr<SDynamicNumericEntry::FDataSource> SpacingSource = SDynamicNumericEntry::MakeSimpleDataSource(
+			SpacingHandle, TInterval<float>(0.0f, 1000.0f), TInterval<float>(0.0f, 4.0f));
+
+		DetailBuilder.EditDefaultProperty(FlowRateHandle)->CustomWidget()
+			.OverrideResetToDefault(FResetToDefaultOverride::Hide())
+			.WholeRowContent()
+		[
+			MakeTwoWidgetDetailRowHBox(
+				MakeFixedWidthLabelSliderHBox(FlowRateHandle, FlowRateSource, FSculptToolsUIConstants::SculptShortLabelWidth),
+				MakeFixedWidthLabelSliderHBox(SpacingHandle, SpacingSource, FSculptToolsUIConstants::SculptShortLabelWidth)
+			)
+		];
+	}
+	else   // if bShowFalloff
+	{
+		TSharedPtr<SDynamicNumericEntry::FDataSource> FalloffSource = SDynamicNumericEntry::MakeSimpleDataSource(
+			FalloffHandle, TInterval<float>(0.0f, 1.0f), TInterval<float>(0.0f, 1.0f));
+
+		DetailBuilder.EditDefaultProperty(FlowRateHandle)->CustomWidget()
+			.OverrideResetToDefault(FResetToDefaultOverride::Hide())
+			.WholeRowContent()
+			[
+				MakeTwoWidgetDetailRowHBox(
+					MakeFixedWidthLabelSliderHBox(FalloffHandle, FalloffSource, FSculptToolsUIConstants::SculptShortLabelWidth),
+					MakeFixedWidthLabelSliderHBox(FlowRateHandle, FlowRateSource, FSculptToolsUIConstants::SculptShortLabelWidth)
+				)
+			];
+	}
 
 	TSharedPtr<IPropertyHandle> LazynessHandle = DetailBuilder.GetProperty(GET_MEMBER_NAME_CHECKED(USculptBrushProperties, Lazyness), USculptBrushProperties::StaticClass());
 	ensure(LazynessHandle->IsValidHandle());
diff --git a/Engine/Plugins/Editor/ModelingToolsEditorMode/Source/ModelingToolsEditorMode/Private/DetailsCustomizations/ModelingToolsColorChannelFilterCustomization.cpp b/Engine/Plugins/Editor/ModelingToolsEditorMode/Source/ModelingToolsEditorMode/Private/DetailsCustomizations/ModelingToolsColorChannelFilterCustomization.cpp
new file mode 100644
--- /dev/null
+++ b/Engine/Plugins/Editor/ModelingToolsEditorMode/Source/ModelingToolsEditorMode/Private/DetailsCustomizations/ModelingToolsColorChannelFilterCustomization.cpp
@@ -0,0 +1,118 @@
+// Copyright Epic Games, Inc. All Rights Reserved.
+
+#include "DetailsCustomizations/ModelingToolPropertyCustomizations.h"
+
+#include "DetailWidgetRow.h"
+
+#include "PropertySets/ColorChannelFilterPropertyType.h"
+#include "Widgets/SBoxPanel.h"
+
+
+#define LOCTEXT_NAMESPACE "ModelingToolsColorChannelFilterCustomization"
+
+
+TSharedRef<IPropertyTypeCustomization> FModelingToolsColorChannelFilterCustomization::MakeInstance()
+{
+	return MakeShareable(new FModelingToolsColorChannelFilterCustomization);
+}
+
+void FModelingToolsColorChannelFilterCustomization::CustomizeHeader(TSharedRef<IPropertyHandle> StructPropertyHandle, FDetailWidgetRow& HeaderRow, IPropertyTypeCustomizationUtils& StructCustomizationUtils)
+{
+	const float XYZPadding = 10.f;
+
+	TSharedPtr<IPropertyHandle> bColorChannelR = StructPropertyHandle->GetChildHandle(GET_MEMBER_NAME_CHECKED(FModelingToolsColorChannelFilter, bRed));
+	bColorChannelR->MarkHiddenByCustomization();
+
+	TSharedPtr<IPropertyHandle> bColorChannelG = StructPropertyHandle->GetChildHandle(GET_MEMBER_NAME_CHECKED(FModelingToolsColorChannelFilter, bGreen));
+	bColorChannelG->MarkHiddenByCustomization();
+
+	TSharedPtr<IPropertyHandle> bColorChannelB = StructPropertyHandle->GetChildHandle(GET_MEMBER_NAME_CHECKED(FModelingToolsColorChannelFilter, bBlue));
+	bColorChannelB->MarkHiddenByCustomization();
+
+	TSharedPtr<IPropertyHandle> bColorChannelA = StructPropertyHandle->GetChildHandle(GET_MEMBER_NAME_CHECKED(FModelingToolsColorChannelFilter, bAlpha));
+	bColorChannelA->MarkHiddenByCustomization();
+
+	HeaderRow.NameContent()
+		[
+			StructPropertyHandle->CreatePropertyNameWidget()
+		]
+	.ValueContent()
+		.MinDesiredWidth(125.f)
+		[
+			SNew(SHorizontalBox)
+			+ SHorizontalBox::Slot()
+			.Padding(0.f, 0.f, XYZPadding, 0.f)
+			.AutoWidth()
+			[
+				SNew(SHorizontalBox)
+				+ SHorizontalBox::Slot()
+				.AutoWidth()
+				[
+					bColorChannelR->CreatePropertyNameWidget()
+				]
+				+ SHorizontalBox::Slot()
+				.AutoWidth()
+				[
+					bColorChannelR->CreatePropertyValueWidget()
+				]
+			]
+
+			+ SHorizontalBox::Slot()
+			.Padding(0.f, 0.f, XYZPadding, 0.f)
+			.AutoWidth()
+			[
+				SNew(SHorizontalBox)
+				+ SHorizontalBox::Slot()
+				.AutoWidth()
+				[
+					bColorChannelG->CreatePropertyNameWidget()
+				]
+				+ SHorizontalBox::Slot()
+				.AutoWidth()
+				[
+					bColorChannelG->CreatePropertyValueWidget()
+				]
+			]
+
+			+ SHorizontalBox::Slot()
+			.Padding(0.f, 0.f, XYZPadding, 0.f)
+			.AutoWidth()
+			[
+				SNew(SHorizontalBox)
+				+ SHorizontalBox::Slot()
+				.AutoWidth()
+				[
+					bColorChannelB->CreatePropertyNameWidget()
+				]
+				+ SHorizontalBox::Slot()
+				.AutoWidth()
+				[
+					bColorChannelB->CreatePropertyValueWidget()
+				]
+			]
+
+			+ SHorizontalBox::Slot()
+			.Padding(0.f, 0.f, XYZPadding, 0.f)
+			.AutoWidth()
+			[
+				SNew(SHorizontalBox)
+				+ SHorizontalBox::Slot()
+				.AutoWidth()
+				[
+					bColorChannelA->CreatePropertyNameWidget()
+				]
+				+ SHorizontalBox::Slot()
+				.AutoWidth()
+				[
+					bColorChannelA->CreatePropertyValueWidget()
+				]
+			]
+		];
+}
+
+void FModelingToolsColorChannelFilterCustomization::CustomizeChildren(TSharedRef<IPropertyHandle> StructPropertyHandle, class IDetailChildrenBuilder& ChildBuilder, IPropertyTypeCustomizationUtils& StructCustomizationUtils)
+{
+}
+
+
+#undef LOCTEXT_NAMESPACE
diff --git a/Engine/Plugins/Editor/ModelingToolsEditorMode/Source/ModelingToolsEditorMode/Private/ModelingToolsActions.cpp b/Engine/Plugins/Editor/ModelingToolsEditorMode/Source/ModelingToolsEditorMode/Private/ModelingToolsActions.cpp
--- a/Engine/Plugins/Editor/ModelingToolsEditorMode/Source/ModelingToolsEditorMode/Private/ModelingToolsActions.cpp
+++ b/Engine/Plugins/Editor/ModelingToolsEditorMode/Source/ModelingToolsEditorMode/Private/ModelingToolsActions.cpp
@@ -6,6 +6,7 @@
 #include "EditMeshMaterialsTool.h"
 #include "MeshVertexSculptTool.h"
 #include "MeshGroupPaintTool.h"
+#include "MeshVertexPaintTool.h"
 #include "MeshAttributePaintTool.h"
 #include "MeshInspectorTool.h"
 #include "CubeGridTool.h"
@@ -92,6 +93,7 @@ void FModelingToolActionCommands::RegisterAllToolActions()
 	FSculptToolActionCommands::Register();
 	FVertexSculptToolActionCommands::Register();
 	FMeshGroupPaintToolActionCommands::Register();
+	FMeshVertexPaintToolActionCommands::Register();
 	FMeshAttributePaintToolActionCommands::Register();
 	FDrawPolygonToolActionCommands::Register();
 	FTransformToolActionCommands::Register();
@@ -109,6 +111,7 @@ void FModelingToolActionCommands::UnregisterAllToolActions()
 	FSculptToolActionCommands::Unregister();
 	FVertexSculptToolActionCommands::Unregister();
 	FMeshGroupPaintToolActionCommands::Unregister();
+	FMeshVertexPaintToolActionCommands::Unregister();
 	FMeshAttributePaintToolActionCommands::Unregister();
 	FDrawPolygonToolActionCommands::Unregister();
 	FTransformToolActionCommands::Unregister();
@@ -143,6 +146,10 @@ void FModelingToolActionCommands::UpdateToolCommandBinding(UInteractiveTool* Too
 	{
 		UPDATE_BINDING(FMeshGroupPaintToolActionCommands);
 	}
+	else if (ExactCast<UMeshVertexPaintTool>(Tool) != nullptr)
+	{
+		UPDATE_BINDING(FMeshVertexPaintToolActionCommands);
+	}
 	else if (ExactCast<UMeshAttributePaintTool>(Tool) != nullptr)
 	{
 		UPDATE_BINDING(FMeshAttributePaintToolActionCommands);
@@ -199,6 +206,7 @@ void CommandsClassName::GetToolDefaultObjectList(TArray<UInteractiveTool*>& Tool
 DEFINE_TOOL_ACTION_COMMANDS(FSculptToolActionCommands, "ModelingToolsSculptTool", "Modeling Tools - Sculpt Tool", UDynamicMeshSculptTool);
 DEFINE_TOOL_ACTION_COMMANDS(FVertexSculptToolActionCommands, "ModelingToolsVertexSculptTool", "Modeling Tools - Vertex Sculpt Tool", UMeshVertexSculptTool);
 DEFINE_TOOL_ACTION_COMMANDS(FMeshGroupPaintToolActionCommands, "ModelingToolsMeshGroupPaintTool", "Modeling Tools - Group Paint Tool", UMeshGroupPaintTool);
+DEFINE_TOOL_ACTION_COMMANDS(FMeshVertexPaintToolActionCommands, "ModelingToolsMeshVertexPaintTool", "Modeling Tools - Vertex Paint Tool", UMeshVertexPaintTool);
 DEFINE_TOOL_ACTION_COMMANDS(FMeshAttributePaintToolActionCommands, "ModelingToolsMeshAttributePaintTool", "Modeling Tools - Attribute Paint Tool", UMeshAttributePaintTool);
 DEFINE_TOOL_ACTION_COMMANDS(FTransformToolActionCommands, "ModelingToolsTransformTool", "Modeling Tools - Transform Tool", UTransformMeshesTool);
 DEFINE_TOOL_ACTION_COMMANDS(FDrawPolygonToolActionCommands, "ModelingToolsDrawPolygonTool", "Modeling Tools - Draw Polygon Tool", UDrawPolygonTool);
diff --git a/Engine/Plugins/Editor/ModelingToolsEditorMode/Source/ModelingToolsEditorMode/Private/ModelingToolsEditorMode.cpp b/Engine/Plugins/Editor/ModelingToolsEditorMode/Source/ModelingToolsEditorMode/Private/ModelingToolsEditorMode.cpp
--- a/Engine/Plugins/Editor/ModelingToolsEditorMode/Source/ModelingToolsEditorMode/Private/ModelingToolsEditorMode.cpp
+++ b/Engine/Plugins/Editor/ModelingToolsEditorMode/Source/ModelingToolsEditorMode/Private/ModelingToolsEditorMode.cpp
@@ -87,6 +87,7 @@
 #include "LatticeDeformerTool.h"
 #include "SeamSculptTool.h"
 #include "MeshGroupPaintTool.h"
+#include "MeshVertexPaintTool.h"
 #include "TransferMeshTool.h"
 #include "ConvertMeshesTool.h"
 #include "SplitMeshesTool.h"
@@ -619,6 +620,7 @@ void UModelingToolsEditorMode::Enter()
 	MeshGroupPaintToolBuilder->StylusAPI = StylusStateTracker.Get();
 #endif
 	RegisterTool(ToolManagerCommands.BeginMeshGroupPaintTool, TEXT("BeginMeshGroupPaintTool"), MeshGroupPaintToolBuilder);
+	RegisterTool(ToolManagerCommands.BeginMeshVertexPaintTool, TEXT("BeginMeshVertexPaintTool"), NewObject<UMeshVertexPaintToolBuilder>());
 
 	RegisterTool(ToolManagerCommands.BeginPolyEditTool, TEXT("BeginPolyEditTool"), NewObject<UEditMeshPolygonsToolBuilder>());
 	UEditMeshPolygonsToolBuilder* TriEditBuilder = NewObject<UEditMeshPolygonsToolBuilder>();
diff --git a/Engine/Plugins/Editor/ModelingToolsEditorMode/Source/ModelingToolsEditorMode/Private/ModelingToolsEditorModeModule.cpp b/Engine/Plugins/Editor/ModelingToolsEditorMode/Source/ModelingToolsEditorMode/Private/ModelingToolsEditorModeModule.cpp
--- a/Engine/Plugins/Editor/ModelingToolsEditorMode/Source/ModelingToolsEditorMode/Private/ModelingToolsEditorModeModule.cpp
+++ b/Engine/Plugins/Editor/ModelingToolsEditorMode/Source/ModelingToolsEditorMode/Private/ModelingToolsEditorModeModule.cpp
@@ -17,6 +17,7 @@
 #include "DetailsCustomizations/MeshTopologySelectionMechanicCustomization.h"
 
 #include "PropertySets/AxisFilterPropertyType.h"
+#include "PropertySets/ColorChannelFilterPropertyType.h"
 #include "Selection/MeshTopologySelectionMechanic.h"
 #include "MeshVertexSculptTool.h"
 #include "BakeMeshAttributeMapsTool.h"
@@ -78,6 +79,8 @@ void FModelingToolsEditorModeModule::OnPostEngineInit()
 	/// Sculpt
 	PropertyModule.RegisterCustomPropertyTypeLayout("ModelingToolsAxisFilter", FOnGetPropertyTypeCustomizationInstance::CreateStatic(&FModelingToolsAxisFilterCustomization::MakeInstance));
 	PropertiesToUnregisterOnShutdown.Add(FModelingToolsAxisFilter::StaticStruct()->GetFName());
+	PropertyModule.RegisterCustomPropertyTypeLayout("ModelingToolsColorChannelFilter", FOnGetPropertyTypeCustomizationInstance::CreateStatic(&FModelingToolsColorChannelFilterCustomization::MakeInstance));
+	PropertiesToUnregisterOnShutdown.Add(FModelingToolsColorChannelFilter::StaticStruct()->GetFName());
 	PropertyModule.RegisterCustomPropertyTypeLayout("BrushToolRadius", FOnGetPropertyTypeCustomizationInstance::CreateStatic(&FModelingToolsBrushSizeCustomization::MakeInstance));
 	PropertiesToUnregisterOnShutdown.Add(FBrushToolRadius::StaticStruct()->GetFName());
 	PropertyModule.RegisterCustomClassLayout("SculptBrushProperties", FOnGetDetailCustomizationInstance::CreateStatic(&FSculptBrushPropertiesDetails::MakeInstance));
diff --git a/Engine/Plugins/Editor/ModelingToolsEditorMode/Source/ModelingToolsEditorMode/Private/ModelingToolsEditorModeToolkit.cpp b/Engine/Plugins/Editor/ModelingToolsEditorMode/Source/ModelingToolsEditorMode/Private/ModelingToolsEditorModeToolkit.cpp
--- a/Engine/Plugins/Editor/ModelingToolsEditorMode/Source/ModelingToolsEditorMode/Private/ModelingToolsEditorModeToolkit.cpp
+++ b/Engine/Plugins/Editor/ModelingToolsEditorMode/Source/ModelingToolsEditorMode/Private/ModelingToolsEditorModeToolkit.cpp
@@ -430,7 +430,9 @@ void FModelingToolsEditorModeToolkit::RegisterPalettes()
 		Commands.BeginAttributeEditorTool,
 		Commands.BeginEditMeshMaterialsTool,
 
+		Commands.BeginMeshVertexPaintTool,
 		Commands.BeginMeshAttributePaintTool,
+
 		Commands.BeginPhysicsInspectorTool,
 		Commands.BeginSetCollisionGeometryTool,
 		Commands.BeginExtractCollisionGeometryTool
diff --git a/Engine/Plugins/Editor/ModelingToolsEditorMode/Source/ModelingToolsEditorMode/Private/ModelingToolsManagerActions.cpp b/Engine/Plugins/Editor/ModelingToolsEditorMode/Source/ModelingToolsEditorMode/Private/ModelingToolsManagerActions.cpp
--- a/Engine/Plugins/Editor/ModelingToolsEditorMode/Source/ModelingToolsEditorMode/Private/ModelingToolsManagerActions.cpp
+++ b/Engine/Plugins/Editor/ModelingToolsEditorMode/Source/ModelingToolsEditorMode/Private/ModelingToolsManagerActions.cpp
@@ -224,6 +224,7 @@ void FModelingToolsManagerCommands::RegisterCommands()
 		REGISTER_MODELING_TOOL_COMMAND(BeginAttributeEditorTool, "AttrEd", "Edit/configure Mesh attributes");
 		REGISTER_MODELING_TOOL_COMMAND(BeginPolyGroupsTool, "GrpGen", "Generate new PolyGroups");
 		REGISTER_MODELING_TOOL_COMMAND(BeginMeshGroupPaintTool, "GrpPnt", "Paint new PolyGroups");
+		REGISTER_MODELING_TOOL_COMMAND(BeginMeshVertexPaintTool, "VtxPnt", "Paint Mesh Vertices");
 		REGISTER_MODELING_TOOL_COMMAND(BeginMeshAttributePaintTool, "MapPnt", "Paint attribute maps");
 		REGISTER_MODELING_TOOL_COMMAND(BeginEditMeshMaterialsTool, "MatEd", "Assign materials to selected triangles");
 	}
@@ -234,6 +235,7 @@ void FModelingToolsManagerCommands::RegisterCommands()
 		REGISTER_MODELING_TOOL_COMMAND(BeginEditTangentsTool, "Tangents", "Recompute Tangents");
 		REGISTER_MODELING_TOOL_COMMAND(BeginPolyGroupsTool, "Generate PolyGroups", "Generate new PolyGroups");
 		REGISTER_MODELING_TOOL_COMMAND(BeginMeshGroupPaintTool, "Paint PolyGroups", "Paint new PolyGroups");
+		REGISTER_MODELING_TOOL_COMMAND(BeginMeshVertexPaintTool, "Paint Vertex Colors", "Paint Mesh Vertex Colors");
 		REGISTER_MODELING_TOOL_COMMAND(BeginAttributeEditorTool, "Edit Attributes", "Edit/configure Mesh attributes");
 		REGISTER_MODELING_TOOL_COMMAND(BeginMeshAttributePaintTool, "Paint Maps", "Paint attribute maps");
 		REGISTER_MODELING_TOOL_COMMAND(BeginEditMeshMaterialsTool, "Edit Materials", "Assign materials to selected triangles");
diff --git a/Engine/Plugins/Editor/ModelingToolsEditorMode/Source/ModelingToolsEditorMode/Public/DetailsCustomizations/ModelingToolPropertyCustomizations.h b/Engine/Plugins/Editor/ModelingToolsEditorMode/Source/ModelingToolsEditorMode/Public/DetailsCustomizations/ModelingToolPropertyCustomizations.h
--- a/Engine/Plugins/Editor/ModelingToolsEditorMode/Source/ModelingToolsEditorMode/Public/DetailsCustomizations/ModelingToolPropertyCustomizations.h
+++ b/Engine/Plugins/Editor/ModelingToolsEditorMode/Source/ModelingToolsEditorMode/Public/DetailsCustomizations/ModelingToolPropertyCustomizations.h
@@ -16,6 +16,18 @@ public:
 	virtual void CustomizeChildren(TSharedRef<IPropertyHandle> StructPropertyHandle, class IDetailChildrenBuilder& ChildBuilder, IPropertyTypeCustomizationUtils& StructCustomizationUtils) override;
 };
 
+
+/** Details customization for FModelingToolsColorChannelFilter struct - R/G/B/A booleans are laid out in a single horizontal row */
+class FModelingToolsColorChannelFilterCustomization : public IPropertyTypeCustomization
+{
+public:
+	static TSharedRef<IPropertyTypeCustomization> MakeInstance();
+
+	virtual void CustomizeHeader(TSharedRef<IPropertyHandle> StructPropertyHandle, class FDetailWidgetRow& HeaderRow, IPropertyTypeCustomizationUtils& StructCustomizationUtils) override;
+	virtual void CustomizeChildren(TSharedRef<IPropertyHandle> StructPropertyHandle, class IDetailChildrenBuilder& ChildBuilder, IPropertyTypeCustomizationUtils& StructCustomizationUtils) override;
+};
+
+
 #if UE_ENABLE_INCLUDE_ORDER_DEPRECATED_IN_5_2
 #include "Input/Reply.h"
 #include "Widgets/Input/SButton.h"
diff --git a/Engine/Plugins/Editor/ModelingToolsEditorMode/Source/ModelingToolsEditorMode/Public/ModelingToolsActions.h b/Engine/Plugins/Editor/ModelingToolsEditorMode/Source/ModelingToolsEditorMode/Public/ModelingToolsActions.h
--- a/Engine/Plugins/Editor/ModelingToolsEditorMode/Source/ModelingToolsEditorMode/Public/ModelingToolsActions.h
+++ b/Engine/Plugins/Editor/ModelingToolsEditorMode/Source/ModelingToolsEditorMode/Public/ModelingToolsActions.h
@@ -79,6 +79,7 @@ public:\
 DECLARE_TOOL_ACTION_COMMANDS(FSculptToolActionCommands);
 DECLARE_TOOL_ACTION_COMMANDS(FVertexSculptToolActionCommands);
 DECLARE_TOOL_ACTION_COMMANDS(FMeshGroupPaintToolActionCommands);
+DECLARE_TOOL_ACTION_COMMANDS(FMeshVertexPaintToolActionCommands);
 DECLARE_TOOL_ACTION_COMMANDS(FMeshAttributePaintToolActionCommands);
 DECLARE_TOOL_ACTION_COMMANDS(FTransformToolActionCommands);
 DECLARE_TOOL_ACTION_COMMANDS(FDrawPolygonToolActionCommands);
diff --git a/Engine/Plugins/Editor/ModelingToolsEditorMode/Source/ModelingToolsEditorMode/Public/ModelingToolsManagerActions.h b/Engine/Plugins/Editor/ModelingToolsEditorMode/Source/ModelingToolsEditorMode/Public/ModelingToolsManagerActions.h
--- a/Engine/Plugins/Editor/ModelingToolsEditorMode/Source/ModelingToolsEditorMode/Public/ModelingToolsManagerActions.h
+++ b/Engine/Plugins/Editor/ModelingToolsEditorMode/Source/ModelingToolsEditorMode/Public/ModelingToolsManagerActions.h
@@ -112,6 +112,7 @@ public:
 	TSharedPtr<FUICommandInfo> BeginMeshToVolumeTool;
 	TSharedPtr<FUICommandInfo> BeginVolumeToMeshTool;
 	TSharedPtr<FUICommandInfo> BeginMeshGroupPaintTool;
+	TSharedPtr<FUICommandInfo> BeginMeshVertexPaintTool;
 
 	TSharedPtr<FUICommandInfo> BeginPhysicsInspectorTool;
 	TSharedPtr<FUICommandInfo> BeginSetCollisionGeometryTool;
diff --git a/Engine/Plugins/Experimental/MeshModelingToolsetExp/Source/MeshModelingToolsExp/Private/MeshVertexPaintTool.cpp b/Engine/Plugins/Experimental/MeshModelingToolsetExp/Source/MeshModelingToolsExp/Private/MeshVertexPaintTool.cpp
new file mode 100644
--- /dev/null
+++ b/Engine/Plugins/Experimental/MeshModelingToolsetExp/Source/MeshModelingToolsExp/Private/MeshVertexPaintTool.cpp
@@ -0,0 +1,2060 @@
+// Copyright Epic Games, Inc. All Rights Reserved.
+
+#include "MeshVertexPaintTool.h"
+#include "Engine/Engine.h"
+#include "Engine/World.h"
+#include "InteractiveToolManager.h"
+#include "ToolTargetManager.h"
+#include "InteractiveGizmoManager.h"
+#include "Drawing/MeshElementsVisualizer.h"
+#include "Async/ParallelFor.h"
+#include "Async/Async.h"
+#include "SceneView.h"
+#include "ToolSetupUtil.h"
+#include "ModelingToolTargetUtil.h"
+#include "MeshWeights.h"
+#include "DynamicMesh/MeshNormals.h"
+#include "DynamicMesh/MeshIndexUtil.h"
+#include "Util/BufferUtil.h"
+#include "Util/ColorConstants.h"
+#include "Selections/MeshConnectedComponents.h"
+#include "Selections/MeshFaceSelection.h"
+#include "Selections/MeshVertexSelection.h"
+#include "Polygroups/PolygroupUtil.h"
+#include "Polygon2.h"
+#include "DynamicMesh/Operations/SplitAttributeWelder.h"
+#include "MeshDescriptionToDynamicMesh.h"
+#include "DynamicMeshToMeshDescription.h"
+
+#include "Changes/MeshVertexChange.h"
+#include "Changes/MeshPolygroupChange.h"
+#include "Changes/BasicChanges.h"
+
+#include "Sculpting/MeshVertexPaintBrushOps.h"
+#include "Sculpting/StampFalloffs.h"
+#include "Sculpting/MeshSculptUtil.h"
+#include "WeightMapUtil.h"
+#include "WeightMapTypes.h"
+
+#include "CanvasTypes.h"
+#include "CanvasItem.h"
+#include "Engine/Engine.h"  // for GEngine->GetSmallFont()
+#include "SceneView.h"
+#include "Materials/Material.h"
+#include "Components/MeshComponent.h"
+
+#include UE_INLINE_GENERATED_CPP_BY_NAME(MeshVertexPaintTool)
+
+using namespace UE::Geometry;
+
+#define LOCTEXT_NAMESPACE "UMeshVertexPaintTool"
+
+namespace
+{
+	// probably should be something defined for the whole tool framework...
+#if WITH_EDITOR
+	static EAsyncExecution VertexPaintToolAsyncExecTarget = EAsyncExecution::LargeThreadPool;
+#else
+	static EAsyncExecution VertexPaintToolAsyncExecTarget = EAsyncExecution::ThreadPool;
+#endif
+}
+
+
+/*
+ * ToolBuilder
+ */
+UMeshSurfacePointTool* UMeshVertexPaintToolBuilder::CreateNewTool(const FToolBuilderState& SceneState) const
+{
+	UMeshVertexPaintTool* SculptTool = NewObject<UMeshVertexPaintTool>(SceneState.ToolManager);
+	SculptTool->SetWorld(SceneState.World);
+	return SculptTool;
+}
+
+bool UMeshVertexPaintToolBuilder::CanBuildTool(const FToolBuilderState& SceneState) const
+{
+	return UMeshSurfacePointMeshEditingToolBuilder::CanBuildTool(SceneState) &&
+		SceneState.TargetManager->CountSelectedAndTargetableWithPredicate(SceneState, GetTargetRequirements(),
+			[](UActorComponent& Component) { return Cast<UMeshComponent>(&Component) != nullptr; }) > 0;
+}
+
+
+/*
+ * Properties
+ */
+void UMeshVertexPaintToolActionPropertySet::PostAction(EMeshVertexPaintToolActions Action)
+{
+	if (ParentTool.IsValid())
+	{
+		ParentTool->RequestAction(Action);
+	}
+}
+
+
+
+
+
+/*
+ * Tool
+ */
+
+void UMeshVertexPaintTool::Setup()
+{
+	// abort for volumes?
+	//if (ActiveColorOverlay == nullptr)
+	//{
+	//	GetToolManager()->PostActiveToolShutdownRequest(this, EToolShutdownType::Cancel, true,
+	//		LOCTEXT("NoColorsShutdown", "Target Mesh does not have or support Vertex Colors"));
+	//	return;
+	//}
+
+
+	UMeshSculptToolBase::Setup();
+
+	SetToolDisplayName(LOCTEXT("ToolName", "Paint Vertex Colors"));
+
+	// create dynamic mesh component to use for live preview
+	FActorSpawnParameters SpawnInfo;
+	PreviewMeshActor = TargetWorld->SpawnActor<AInternalToolFrameworkActor>(FVector::ZeroVector, FRotator::ZeroRotator, SpawnInfo);
+	DynamicMeshComponent = NewObject<UDynamicMeshComponent>(PreviewMeshActor);
+	InitializeSculptMeshComponent(DynamicMeshComponent, PreviewMeshActor);
+
+	// assign materials
+	FComponentMaterialSet MaterialSet = UE::ToolTarget::GetMaterialSet(Target);
+	for (int k = 0; k < MaterialSet.Materials.Num(); ++k)
+	{
+		DynamicMeshComponent->SetMaterial(k, MaterialSet.Materials[k]);
+	}
+
+	DynamicMeshComponent->SetInvalidateProxyOnChangeEnabled(false);
+	OnDynamicMeshComponentChangedHandle = DynamicMeshComponent->OnMeshChanged.AddUObject(this, &UMeshVertexPaintTool::OnDynamicMeshComponentChanged);
+
+	FDynamicMesh3* Mesh = GetSculptMesh();
+	ActiveColorOverlay = (Mesh->Attributes() != nullptr) ? Mesh->Attributes()->PrimaryColors() : nullptr;
+
+	if (ActiveColorOverlay == nullptr)
+	{
+		if (Mesh->Attributes() == nullptr)
+		{
+			Mesh->EnableAttributes();
+		}
+		if (Mesh->Attributes()->PrimaryColors() == nullptr)
+		{
+			Mesh->Attributes()->EnablePrimaryColors();
+			ActiveColorOverlay = Mesh->Attributes()->PrimaryColors();
+			ActiveColorOverlay->CreateFromPredicate(
+				[](int, int, int) { return false; }, 1.0f);
+		}
+	}
+	else
+	{
+		ActiveColorOverlay->SplitVerticesWithPredicate(
+			[](int ElementIdx, int TriID) { return true; },
+			[this](int ElementIdx, int TriID, float* FillVect)
+			{
+				FVector4f CurValue = ActiveColorOverlay->GetElement(ElementIdx);
+				FillVect[0] = CurValue.X; FillVect[1] = CurValue.Y; FillVect[2] = CurValue.Z; FillVect[3] = CurValue.W;
+			});
+	}
+
+	StrokeInitialColorBuffer.SetNum(ActiveColorOverlay->MaxElementID());
+	StrokeAccumColorBuffer.SetNum(ActiveColorOverlay->MaxElementID());
+
+	// UMeshSculptToolBase::InitializeSculptMeshComponent() should provide some way for this level to control how 
+	// colors are interpreted. Actually calls GetDynamicMeshCopy() which may or may not convert to SRGB...
+	// 
+	// code above auto-converted to SRGB so we will convert back to linear now...
+	for (int32 ElementID : ActiveColorOverlay->ElementIndicesItr())
+	{
+		FVector4f Color = ActiveColorOverlay->GetElement(ElementID);
+		LinearColors::SRGBToLinear(Color);
+		ActiveColorOverlay->SetElement(ElementID, Color);
+	}
+	DynamicMeshComponent->ColorSpaceMode = EDynamicMeshVertexColorTransformMode::LinearToSRGB;
+
+	//Mesh->EnableVertexColors(FVector3f::One());
+	FAxisAlignedBox3d Bounds = Mesh->GetBounds(true);
+
+	TFuture<void> PrecomputeFuture = Async(VertexPaintToolAsyncExecTarget, [&]()
+	{
+		PrecomputeFilterData();
+	});
+
+	TFuture<void> OctreeFuture = Async(VertexPaintToolAsyncExecTarget, [&]()
+	{
+		// initialize dynamic octree
+		if (Mesh->TriangleCount() > 100000)
+		{
+			Octree.RootDimension = Bounds.MaxDim() / 10.0;
+			Octree.SetMaxTreeDepth(4);
+		}
+		else
+		{
+			Octree.RootDimension = Bounds.MaxDim();
+			Octree.SetMaxTreeDepth(8);
+		}
+		Octree.Initialize(Mesh);
+	});
+
+	// initialize render decomposition
+	TUniquePtr<FMeshRenderDecomposition> Decomp = MakeUnique<FMeshRenderDecomposition>();
+	FMeshRenderDecomposition::BuildChunkedDecomposition(Mesh, &MaterialSet, *Decomp);
+	Decomp->BuildAssociations(Mesh);
+	//UE_LOG(LogTemp, Warning, TEXT("Decomposition has %d groups"), Decomp->Num());
+	DynamicMeshComponent->SetExternalDecomposition(MoveTemp(Decomp));
+
+	// initialize brush radius range interval, brush properties
+	UMeshSculptToolBase::InitializeBrushSizeRange(Bounds);
+
+	// Set up control points mechanic
+	PolyLassoMechanic = NewObject<UPolyLassoMarqueeMechanic>(this);
+	PolyLassoMechanic->Setup(this);
+	PolyLassoMechanic->SetIsEnabled(false);
+	PolyLassoMechanic->SpacingTolerance = 10.0f;
+	PolyLassoMechanic->OnDrawPolyLassoFinished.AddUObject(this, &UMeshVertexPaintTool::OnPolyLassoFinished);
+
+	PolygroupLayerProperties = NewObject<UPolygroupLayersProperties>(this);
+	PolygroupLayerProperties->RestoreProperties(this, TEXT("MeshVertexPaintTool"));
+	PolygroupLayerProperties->InitializeGroupLayers(GetSculptMesh());
+	PolygroupLayerProperties->WatchProperty(PolygroupLayerProperties->ActiveGroupLayer, [&](FName) { OnSelectedGroupLayerChanged(); });
+	UpdateActiveGroupLayer();
+
+	BasicProperties = NewObject<UVertexPaintBasicProperties>(this);
+	BasicProperties->RestoreProperties(this);
+	BasicProperties->WatchProperty(BasicProperties->SubToolType,
+		[this](EMeshVertexPaintInteractionType NewType) { UpdateSubToolType(NewType); });
+	BasicProperties->WatchProperty(BasicProperties->ChannelFilter,
+		[this](const FModelingToolsColorChannelFilter&) { OnChannelFilterModified(); });
+	BasicProperties->WatchProperty(BasicProperties->SecondaryActionType,
+		[this](EMeshVertexPaintSecondaryActionType NewType) { UpdateSecondaryBrushType(NewType); });
+
+	// initialize other properties
+	FilterProperties = NewObject<UVertexPaintBrushFilterProperties>(this);
+	FilterProperties->RestoreProperties(this);
+	FilterProperties->WatchProperty(FilterProperties->MaterialMode, [this](EMeshVertexPaintMaterialMode) { UpdateVertexPaintMaterialMode(); });
+
+	InitializeIndicator();
+
+	AddToolPropertySource(BasicProperties);
+	AddToolPropertySource(UMeshSculptToolBase::BrushProperties);
+
+	// initialize our properties
+	UMeshSculptToolBase::BrushProperties->bShowPerBrushProps = false;
+	UMeshSculptToolBase::BrushProperties->bShowFalloff = true;
+	UMeshSculptToolBase::BrushProperties->bShowLazyness = true;
+	UMeshSculptToolBase::BrushProperties->bShowFlowRate = true;
+	UMeshSculptToolBase::BrushProperties->bShowSpacing = false;
+	UMeshSculptToolBase::BrushProperties->BrushFalloffAmount = 0.66f;
+	CalculateBrushRadius();
+
+	PaintBrushOpProperties = NewObject<UVertexColorPaintBrushOpProps>(this);
+	RegisterBrushType((int32)EMeshVertexPaintBrushType::Paint, LOCTEXT("Paint", "Paint"),
+		MakeUnique<FLambdaMeshSculptBrushOpFactory>([this]() { return MakeUnique<UE::Geometry::FVertexPaintBrushOp>(); }),
+		PaintBrushOpProperties);
+
+	// secondary brushes
+	EraseBrushOpProperties = NewObject<UVertexColorPaintBrushOpProps>(this);
+	RegisterSecondaryBrushType((int32)EMeshVertexPaintBrushType::Erase, LOCTEXT("Erase", "Erase"),
+		MakeUnique<FLambdaMeshSculptBrushOpFactory>([this]() { return MakeUnique<UE::Geometry::FVertexPaintBrushOp>(); }),
+		EraseBrushOpProperties);
+
+	RegisterSecondaryBrushType((int32)EMeshVertexPaintBrushType::Soften, LOCTEXT("Soften", "Soften"),
+		MakeUnique<TBasicMeshSculptBrushOpFactory<FVertexColorSoftenBrushOp>>(),
+		NewObject<UVertexColorSoftenBrushOpProps>(this));
+
+	RegisterSecondaryBrushType((int32)EMeshVertexPaintBrushType::Smooth, LOCTEXT("Smooth", "Smooth"),
+		MakeUnique<TBasicMeshSculptBrushOpFactory<FVertexColorSmoothBrushOp>>(),
+		NewObject<UVertexColorSmoothBrushOpProps>(this));
+
+	// falloffs
+	RegisterStandardFalloffTypes();
+
+
+	QuickActions = NewObject<UMeshVertexPaintToolQuickActions>(this);
+	QuickActions->Initialize(this);
+
+	UtilityActions = NewObject<UMeshVertexPaintToolUtilityActions>(this);
+	UtilityActions->Initialize(this);
+
+	AddToolPropertySource(FilterProperties);
+	AddToolPropertySource(UMeshSculptToolBase::ViewProperties);
+	AddToolPropertySource(QuickActions);
+	AddToolPropertySource(UtilityActions);
+	AddToolPropertySource(PolygroupLayerProperties);
+
+	AddToolPropertySource(UMeshSculptToolBase::GizmoProperties);
+	SetToolPropertySourceEnabled(UMeshSculptToolBase::GizmoProperties, false);
+
+
+	// register watchers
+	BasicProperties->WatchProperty(BasicProperties->PrimaryBrushType,
+		[this](EMeshVertexPaintBrushType NewType) { UpdateBrushType(NewType); });
+
+	// must call before updating brush type so that we register all brush properties?
+	UMeshSculptToolBase::OnCompleteSetup();
+
+	UpdateBrushType(BasicProperties->PrimaryBrushType);
+	UpdateSecondaryBrushType(BasicProperties->SecondaryActionType);
+	SetPrimaryFalloffType(EMeshSculptFalloffType::Smooth);
+
+	MeshElementsDisplay = NewObject<UMeshElementsVisualizer>(this);
+	MeshElementsDisplay->CreateInWorld(DynamicMeshComponent->GetWorld(), DynamicMeshComponent->GetComponentTransform());
+	if (ensure(MeshElementsDisplay->Settings))
+	{
+		MeshElementsDisplay->Settings->bShowColorSeams = false;		// default color seams to false
+		MeshElementsDisplay->Settings->RestoreProperties(this, TEXT("MeshVertexPaintTool"));
+		AddToolPropertySource(MeshElementsDisplay->Settings);
+	}
+	MeshElementsDisplay->SetMeshAccessFunction([this](UMeshElementsVisualizer::ProcessDynamicMeshFunc ProcessFunc) {
+		ProcessFunc(*GetSculptMesh());
+	});
+
+	// initialize weightmap list in Utility actions
+	TArray<FName> TargetWeightMaps;
+	UE::WeightMaps::FindVertexWeightMaps(UE::ToolTarget::GetMeshDescription(Target), TargetWeightMaps);
+	for (FName Name : TargetWeightMaps)
+	{
+		UtilityActions->WeightMapsList.Add(Name.ToString());
+	}
+	if (UtilityActions->WeightMapsList.Num() > 0)
+	{
+		UtilityActions->WeightMap = FName(UtilityActions->WeightMapsList[0]);
+	}
+
+	// initialize LOD names list in Utility actions
+	this->SourceLOD = UE::ToolTarget::GetTargetMeshDescriptionLOD(Target, this->bTargetSupportsLODs);
+	if (bTargetSupportsLODs)
+	{
+		this->AvailableLODs = UE::ToolTarget::GetMeshDescriptionLODs(Target, bTargetSupportsLODs, false, true);
+		this->AvailableLODs.Remove(SourceLOD);
+	}
+	if (this->bTargetSupportsLODs && AvailableLODs.Num() > 0)
+	{
+		for ( int32 k = 0; k < AvailableLODs.Num(); ++k)
+		{
+			EMeshLODIdentifier LOD = AvailableLODs[k];
+			if (LOD != SourceLOD)
+			{
+				if (LOD == EMeshLODIdentifier::HiResSource)
+				{
+					UtilityActions->LODNamesList.Add(TEXT("HiRes"));
+				}
+				else
+				{
+					UtilityActions->LODNamesList.Add(FString::Printf(TEXT("LOD %d"), static_cast<int32>(LOD)));
+				}
+			}
+			
+		}
+		UtilityActions->CopyToLODName = UtilityActions->LODNamesList[0];
+	}
+
+
+	// disable view properties
+	SetViewPropertiesEnabled(false);
+	UpdateWireframeVisibility(false);
+	UpdateFlatShadingSetting(true);
+	UpdateVertexPaintMaterialMode();
+
+	// configure panels
+	UpdateSubToolType(BasicProperties->SubToolType);
+	OnChannelFilterModified();
+
+	PrecomputeFuture.Wait();
+	OctreeFuture.Wait();
+}
+
+void UMeshVertexPaintTool::Shutdown(EToolShutdownType ShutdownType)
+{
+	if (DynamicMeshComponent != nullptr)
+	{
+		DynamicMeshComponent->OnMeshChanged.Remove(OnDynamicMeshComponentChangedHandle);
+	}
+
+	if (ensure(MeshElementsDisplay->Settings))
+	{
+		MeshElementsDisplay->Settings->SaveProperties(this, TEXT("MeshVertexPaintTool"));
+	}
+	MeshElementsDisplay->Disconnect();
+
+	BasicProperties->SaveProperties(this);
+	FilterProperties->SaveProperties(this);
+	PolygroupLayerProperties->SaveProperties(this, TEXT("MeshVertexPaintTool"));
+
+	if (PreviewMeshActor != nullptr)
+	{
+		PreviewMeshActor->Destroy();
+		PreviewMeshActor = nullptr;
+	}
+
+	if (ShutdownType == EToolShutdownType::Accept)
+	{
+		// weld colors to remove seams where possible
+		FDynamicMesh3* Mesh = GetSculptMesh();
+		FDynamicMeshColorOverlay* ColorOverlay = GetActiveColorOverlay();
+		for (int32 vid : Mesh->VertexIndicesItr())
+		{
+			FSplitAttributeWelder::WeldSplitColors(vid, *ColorOverlay, FMathd::ZeroTolerance);
+		}
+		Mesh->CompactInPlace();		// to clean up colors we just welded
+
+		// TODO can override UMeshSculptToolBase::CommitResult() that ::Shutdown calls to avoid this conversion...
+
+		// shutdown call below wants colors to be in SRGB so it can convert back to linear...
+		for (int32 ElementID : ActiveColorOverlay->ElementIndicesItr())
+		{
+			FVector4f Color = ColorOverlay->GetElement(ElementID);
+			LinearColors::LinearToSRGB(Color);
+			ColorOverlay->SetElement(ElementID, Color);
+		}
+	}
+
+	UMeshSculptToolBase::Shutdown(ShutdownType);
+}
+
+
+void UMeshVertexPaintTool::CommitResult(UBaseDynamicMeshComponent* Component, bool bModifiedTopology)
+{
+	GetToolManager()->BeginUndoTransaction(LOCTEXT("VertexPaintToolTransactionName", "Paint Vertex Colors"));
+	Component->ProcessMesh([&](const FDynamicMesh3& CurMesh)
+	{
+		UE::ToolTarget::CommitDynamicMeshUpdate(Target, CurMesh, true);
+	});
+	GetToolManager()->EndUndoTransaction();
+}
+
+
+void UMeshVertexPaintTool::RegisterActions(FInteractiveToolActionSet& ActionSet)
+{
+	UMeshSculptToolBase::RegisterActions(ActionSet);
+
+	ActionSet.RegisterAction(this, (int32)EStandardToolActions::BaseClientDefinedActionID + 500,
+		TEXT("PickColorUnderCursor"),
+		LOCTEXT("PickColorUnderCursor", "Pick Color"),
+		LOCTEXT("PickColorUnderCursorTooltip", "Switch the Paint Color to the color currently under the cursor"),
+		EModifierKey::Shift, EKeys::G,
+		[this]() { bPendingPickColor = true; });
+
+	ActionSet.RegisterAction(this, (int32)EStandardToolActions::BaseClientDefinedActionID + 501,
+		TEXT("PickEraseColorUnderCursor"),
+		LOCTEXT("PickEraseColorUnderCursor", "Pick Erase Color"),
+		LOCTEXT("PickEraseColorUnderCursorTooltip", "Switch the Erase Color to the color currently under the cursor"),
+		EModifierKey::Control | EModifierKey::Shift, EKeys::G,
+		[this]() { bPendingPickEraseColor = true; });
+
+	//ActionSet.RegisterAction(this, (int32)EStandardToolActions::BaseClientDefinedActionID + 502,
+	//	TEXT("ToggleFrozenGroup"),
+	//	LOCTEXT("ToggleFrozenGroup", "Toggle Group Frozen State"),
+	//	LOCTEXT("ToggleFrozenGroupTooltip", "Toggle Group Frozen State"),
+	//	EModifierKey::Shift, EKeys::F,
+	//	[this]() { bPendingToggleFreezeGroup = true; });
+
+
+	//ActionSet.RegisterAction(this, (int32)EStandardToolActions::BaseClientDefinedActionID + 503,
+	//	TEXT("CreateNewGroup"),
+	//	LOCTEXT("CreateNewGroup", "New Group"),
+	//	LOCTEXT("CreateNewGroupTooltip", "Allocate a new Polygroup and set as Current"),
+	//	EModifierKey::Shift, EKeys::Q,
+	//	[this]() { AllocateNewGroupAndSetAsCurrentAction(); });
+};
+
+
+TUniquePtr<FMeshSculptBrushOp>& UMeshVertexPaintTool::GetActiveBrushOp()
+{
+	if (GetInEraseStroke())
+	{
+		return SecondaryBrushOp;
+	}
+	else
+	{
+		return PrimaryBrushOp;
+	}
+}
+
+
+void UMeshVertexPaintTool::OnPropertyModified(UObject* PropertySet, FProperty* Property)
+{
+	CalculateBrushRadius();
+}
+
+
+
+bool UMeshVertexPaintTool::IsInBrushSubMode() const
+{
+	return BasicProperties->SubToolType == EMeshVertexPaintInteractionType::Brush
+		|| BasicProperties->SubToolType == EMeshVertexPaintInteractionType::Fill
+		|| BasicProperties->SubToolType == EMeshVertexPaintInteractionType::GroupFill;
+}
+
+
+void UMeshVertexPaintTool::OnBeginStroke(const FRay& WorldRay)
+{
+	UpdateBrushPosition(WorldRay);
+
+	if (PaintBrushOpProperties)
+	{
+		PaintBrushOpProperties->Color = BasicProperties->PaintColor;
+		PaintBrushOpProperties->BlendMode = (EVertexColorPaintBrushOpBlendMode)(int32)BasicProperties->BlendMode;
+	}
+	if (EraseBrushOpProperties)
+	{
+		EraseBrushOpProperties->Color = BasicProperties->EraseColor;
+		EraseBrushOpProperties->BlendMode = EVertexColorPaintBrushOpBlendMode::Lerp;
+	}
+
+	TUniquePtr<FMeshSculptBrushOp>& UseBrushOp = GetActiveBrushOp();
+	if (GetInSmoothingStroke() 
+		&& BasicProperties->SecondaryActionType == EMeshVertexPaintSecondaryActionType::Smooth)
+	{
+		UseBrushOp->PropertySet->SetStrength(0.01 + BasicProperties->SmoothStrength / 2.0);
+	}
+	else
+	{
+		UseBrushOp->PropertySet->SetStrength(
+			FMathf::Pow(UMeshSculptToolBase::BrushProperties->FlowRate, 2.0f));
+	}
+	UseBrushOp->PropertySet->SetFalloff(UMeshSculptToolBase::BrushProperties->BrushFalloffAmount);
+
+	// initialize first "Last Stamp", so that we can assume all stamps in stroke have a valid previous stamp
+	LastStamp.WorldFrame = GetBrushFrameWorld();
+	LastStamp.LocalFrame = GetBrushFrameLocal();
+	LastStamp.Radius = GetCurrentBrushRadius();
+	LastStamp.Falloff = GetCurrentBrushFalloff();
+	LastStamp.Direction = GetInInvertStroke() ? -1.0 : 1.0;
+	LastStamp.Depth = GetCurrentBrushDepth();
+	LastStamp.Power = GetActivePressure() * GetCurrentBrushStrength();
+	LastStamp.TimeStamp = FDateTime::Now();
+
+	FSculptBrushOptions SculptOptions;
+	//SculptOptions.bPreserveUVFlow = false; // FilterProperties->bPreserveUVFlow;
+	SculptOptions.ConstantReferencePlane = GetCurrentStrokeReferencePlane();
+
+	// to mix colors during a stroke like a proper painting tool, we need to accumulate stroke in a separate buffer 
+	// and each frame blend it with the initial color buffer.
+	// (this should be done async after previous stroke ended...except initial buffer needs to be updated on other operations then)
+	if (GetInSmoothingStroke() == false && PaintBrushOpProperties->BlendMode == EVertexColorPaintBrushOpBlendMode::Mix)
+	{
+		FDynamicMeshColorOverlay* ColorOverlay = GetActiveColorOverlay();
+		int32 NumElements = StrokeAccumColorBuffer.Num();
+		for (int32 k = 0; k < NumElements; ++k)
+		{
+			StrokeAccumColorBuffer[k] = FVector4f::Zero();
+			StrokeInitialColorBuffer[k] = ColorOverlay->IsElement(k) ? ColorOverlay->GetElement(k) : FVector4f::One();
+		}
+	}
+
+	UseBrushOp->ConfigureOptions(SculptOptions);
+	UseBrushOp->BeginStroke(GetSculptMesh(), LastStamp, VertexROI);
+
+	AccumulatedTriangleROI.Reset();
+
+	// begin change here? or wait for first stamp?
+	BeginChange();
+}
+
+void UMeshVertexPaintTool::OnEndStroke()
+{
+	GetActiveBrushOp()->EndStroke(GetSculptMesh(), LastStamp, VertexROI);
+
+	// close change record
+	EndChange();
+}
+
+
+
+
+void UMeshVertexPaintTool::UpdateROI(const FSculptBrushStamp& BrushStamp)
+{
+	TRACE_CPUPROFILER_EVENT_SCOPE(MeshVPaint_UpdateROI);
+
+	const FVector3d& BrushPos = BrushStamp.LocalFrame.Origin;
+	const FDynamicMesh3* Mesh = GetSculptMesh();
+	float RadiusSqr = GetCurrentBrushRadius() * GetCurrentBrushRadius();
+	FAxisAlignedBox3d BrushBox(
+		BrushPos - GetCurrentBrushRadius() * FVector3d::One(),
+		BrushPos + GetCurrentBrushRadius() * FVector3d::One());
+
+	TriangleROI.Reset();
+
+	int32 CenterTID = GetBrushTriangleID();
+	if (Mesh->IsTriangle(CenterTID))
+	{
+		TriangleROI.Add(CenterTID);
+	}
+
+	FVector3d CenterNormal = Mesh->IsTriangle(CenterTID) ? TriNormals[CenterTID] : FVector3d::One();		// One so that normal check always passes
+
+	bool bVolumetric = (FilterProperties->BrushAreaMode == EMeshVertexPaintBrushAreaType::Volumetric);
+	bool bUseAngleThreshold = (bVolumetric == false) && (FilterProperties->AngleThreshold < 180.0f);
+	double DotAngleThreshold = FMathd::Cos(FilterProperties->AngleThreshold * FMathd::DegToRad);
+	bool bStopAtUVSeams = FilterProperties->bUVSeams;
+	bool bStopAtNormalSeams = FilterProperties->bNormalSeams;
+
+	auto CheckEdgeCriteria = [&](int32 t1, int32 t2) -> bool
+	{
+		if (bUseAngleThreshold == false || CenterNormal.Dot(TriNormals[t2]) > DotAngleThreshold)
+		{
+			int32 eid = Mesh->FindEdgeFromTriPair(t1, t2);
+			if (bStopAtUVSeams == false || UVSeamEdges[eid] == false)
+			{
+				if (bStopAtNormalSeams == false || NormalSeamEdges[eid] == false)
+				{
+					return true;
+				}
+			}
+		}
+		return false;
+	};
+
+	bool bFill = (BasicProperties->SubToolType == EMeshVertexPaintInteractionType::Fill);
+	bool bGroupFill = (BasicProperties->SubToolType == EMeshVertexPaintInteractionType::GroupFill);
+
+	if (bVolumetric)
+	{
+		Octree.RangeQuery(BrushBox,
+			[&](int TriIdx) {
+
+			if ((Mesh->GetTriCentroid(TriIdx) - BrushPos).SquaredLength() < RadiusSqr)
+			{
+				TriangleROI.Add(TriIdx);
+			}
+		});
+	}
+	else
+	{
+		if (Mesh->IsTriangle(CenterTID))
+		{
+			TArray<int32> StartROI;
+			StartROI.Add(CenterTID);
+			FMeshConnectedComponents::GrowToConnectedTriangles(Mesh, StartROI, TriangleROI, &TempROIBuffer,
+				[&](int t1, int t2) 
+			{ 
+				if ((Mesh->GetTriCentroid(t2) - BrushPos).SquaredLength() < RadiusSqr)
+				{
+					return CheckEdgeCriteria(t1, t2);
+				}
+				return false;
+			});
+		}
+	}
+
+	if (bFill)
+	{
+		TArray<int32> StartROI;
+		for (int32 tid : TriangleROI)
+		{
+			StartROI.Add(tid);
+		}
+		FMeshConnectedComponents::GrowToConnectedTriangles(Mesh, StartROI, TriangleROI, &TempROIBuffer,
+														   [&](int t1, int t2)
+		{
+			return CheckEdgeCriteria(t1, t2);
+		});
+	}
+	else if (bGroupFill)
+	{
+		TArray<int32> StartROI;
+		TSet<int32> FillGroups;
+		for (int32 tid : TriangleROI)
+		{
+			StartROI.Add(tid);
+			FillGroups.Add(ActiveGroupSet->GetGroup(tid));
+		}
+		FMeshConnectedComponents::GrowToConnectedTriangles(Mesh, StartROI, TriangleROI, &TempROIBuffer,
+			[&](int t1, int t2)
+		{
+			return (FillGroups.Contains(ActiveGroupSet->GetGroup(t2)));
+		});
+	}
+
+
+	// apply visibility filter
+	if (FilterProperties->VisibilityFilter != EMeshVertexPaintVisibilityType::None)
+	{
+		TArray<int32> ResultBuffer;
+		ApplyVisibilityFilter(TriangleROI, TempROIBuffer, ResultBuffer);
+	}
+
+	// construct ROI vertex set
+	TempVertexSet.Reset();
+	for (int32 tid : TriangleROI)
+	{
+		FIndex3i Tri = Mesh->GetTriangle(tid);
+		TempVertexSet.Add(Tri.A);  TempVertexSet.Add(Tri.B);  TempVertexSet.Add(Tri.C);
+	}
+	VertexROI.SetNum(0, false);
+	BufferUtil::AppendElements(VertexROI, TempVertexSet);
+
+	// construct ROI triangle and group buffers
+	ROITriangleBuffer.Reserve(TriangleROI.Num());
+	ROITriangleBuffer.SetNum(0, false);
+	for (int32 tid : TriangleROI)
+	{
+		ROITriangleBuffer.Add(tid);
+	}
+}
+
+bool UMeshVertexPaintTool::UpdateStampPosition(const FRay& WorldRay)
+{
+	CalculateBrushRadius();
+
+	TUniquePtr<FMeshSculptBrushOp>& UseBrushOp = GetActiveBrushOp();
+
+	ESculptBrushOpTargetType TargetType = UseBrushOp->GetBrushTargetType();
+	switch (TargetType)
+	{
+	case ESculptBrushOpTargetType::SculptMesh:
+	case ESculptBrushOpTargetType::TargetMesh:
+		UpdateBrushPositionOnSculptMesh(WorldRay, true);
+		break;
+	case ESculptBrushOpTargetType::ActivePlane:
+		check(false);
+		UpdateBrushPositionOnActivePlane(WorldRay);
+		break;
+	}
+
+	if (UseBrushOp->GetAlignStampToView())
+	{
+		AlignBrushToView();
+	}
+
+	CurrentStamp = LastStamp;
+	CurrentStamp.DeltaTime = FMathd::Min((FDateTime::Now() - LastStamp.TimeStamp).GetTotalSeconds(), 1.0);
+	CurrentStamp.WorldFrame = GetBrushFrameWorld();
+	CurrentStamp.LocalFrame = GetBrushFrameLocal();
+	CurrentStamp.Power = GetActivePressure() * GetCurrentBrushStrength();
+
+	CurrentStamp.PrevLocalFrame = LastStamp.LocalFrame;
+	CurrentStamp.PrevWorldFrame = LastStamp.WorldFrame;
+
+	FVector3d MoveDelta = CurrentStamp.LocalFrame.Origin - CurrentStamp.PrevLocalFrame.Origin;
+	if (UseBrushOp->IgnoreZeroMovements() && MoveDelta.SquaredLength() < FMathd::ZeroTolerance)
+	{
+		return false;
+	}
+
+	return true;
+}
+
+
+bool UMeshVertexPaintTool::ApplyStamp()
+{
+	TRACE_CPUPROFILER_EVENT_SCOPE(MeshVPaint_ApplyStamp);
+
+	TUniquePtr<FMeshSculptBrushOp>& UseBrushOp = GetActiveBrushOp();
+
+	// yuck
+	FMeshVertexColorBrushOp* VertexColorBrushOp = (FMeshVertexColorBrushOp*)UseBrushOp.Get();
+
+	// convert triangle ROI to element IDs, populating ROIElementSet, ROIElementBuffer, ROIColorBuffer
+	{
+		TRACE_CPUPROFILER_EVENT_SCOPE(MeshVPaint_ApplyStamp_ElementROI);
+		InitializeElementROIFromTriangleROI(ROITriangleBuffer, true);
+	}
+	
+	FDynamicMesh3* Mesh = GetSculptMesh();
+	{
+		TRACE_CPUPROFILER_EVENT_SCOPE(MeshVPaint_ApplyStamp_ApplyStampToVertexColors);
+		VertexColorBrushOp->ApplyStampToVertexColors(Mesh, GetActiveColorOverlay(),
+			StrokeInitialColorBuffer, StrokeAccumColorBuffer,
+			CurrentStamp, ROIElementBuffer, ROIColorBuffer);
+	}
+
+	bool bUpdated = SyncMeshWithColorBuffer(Mesh);
+
+	LastStamp = CurrentStamp;
+	LastStamp.TimeStamp = FDateTime::Now();
+
+	return bUpdated;
+}
+
+
+
+
+bool UMeshVertexPaintTool::SyncMeshWithColorBuffer(FDynamicMesh3* Mesh)
+{
+	TRACE_CPUPROFILER_EVENT_SCOPE(MeshVPaint_SyncMeshWithColorBuffer);
+
+	int NumModified = 0;
+	const int32 NumElems = ROIElementBuffer.Num();
+
+	if (FDynamicMeshAttributeSet* AttributeSet = Mesh->Attributes())
+	{
+		if (FDynamicMeshColorOverlay* ColorAttrib = AttributeSet->PrimaryColors())
+		{
+			for (int32 k = 0; k < NumElems; ++k)
+			{
+				// check if color has changed? why bother?
+				int32 ElementID = ROIElementBuffer[k];
+				FVector4f CurColor = ColorAttrib->GetElement(ElementID);
+				FVector4f NewColor = ROIColorBuffer[k];
+				ApplyChannelFilter(CurColor, NewColor);
+				ActiveChangeBuilder->UpdateValue(ElementID, CurColor, NewColor);
+				ColorAttrib->SetElement(ElementID, NewColor);
+				NumModified++;
+			}
+
+		}
+	}
+
+	return (NumModified > 0);
+}
+
+
+
+
+
+void UMeshVertexPaintTool::OnPolyLassoFinished(const FCameraPolyLasso& Lasso, bool bCanceled)
+{
+	// construct polyline
+	TArray<FVector2D> Polyline = Lasso.Polyline;
+	int32 N = Polyline.Num();
+	if (N < 2)
+	{
+		return;
+	}
+
+	// Try to clip polyline to be closed, or closed-enough for winding evaluation to work.
+	// If that returns false, the polyline is "too open". In that case we will extend
+	// outwards from the endpoints and then try to create a closed very large polygon
+	if (UPolyLassoMarqueeMechanic::ApproximateSelfClipPolyline(Polyline) == false)
+	{
+		FVector2d StartDirOut = UE::Geometry::Normalized(Polyline[0] - Polyline[1]);
+		FLine2d StartLine(Polyline[0], StartDirOut);
+		FVector2d EndDirOut = UE::Geometry::Normalized(Polyline[N-1] - Polyline[N-2]);
+		FLine2d EndLine(Polyline[N-1], EndDirOut);
+
+		// if we did not intersect, we are in ambiguous territory. Check if a segment along either end-direction
+		// intersects the polyline. If it does, we have something like a spiral and will be OK. 
+		// If not, make a closed polygon by interpolating outwards from each endpoint, and then in perp-directions.
+		FPolygon2d Polygon(Polyline);
+		float PerpSign = Polygon.IsClockwise() ? -1.0 : 1.0;
+
+		Polyline.Insert(StartLine.PointAt(10000.0f), 0);
+		Polyline.Insert(Polyline[0] + 1000 * PerpSign * UE::Geometry::PerpCW(StartDirOut), 0);
+
+		Polyline.Add(EndLine.PointAt(10000.0f));
+		Polyline.Add(Polyline.Last() + 1000 * PerpSign * UE::Geometry::PerpCW(EndDirOut));
+		FVector2d StartPos = Polyline[0];
+		Polyline.Add(StartPos);		// close polyline (cannot use Polyline[0] in case Add resizes!)
+	}
+
+	N = Polyline.Num();
+
+	// project each mesh vertex to view plane and evaluate winding integral of polyline
+	const FDynamicMesh3* Mesh = GetSculptMesh();
+	TempROIBuffer.SetNum(Mesh->MaxVertexID());
+	ParallelFor(Mesh->MaxVertexID(), [&](int32 vid)
+	{
+		if (Mesh->IsVertex(vid))
+		{
+			FVector3d WorldPos = CurTargetTransform.TransformPosition(Mesh->GetVertex(vid));
+			FVector2d PlanePos = (FVector2d)Lasso.GetProjectedPoint((FVector)WorldPos);
+
+			double WindingSum = 0;
+			FVector2d a = Polyline[0] - PlanePos, b = FVector2d::Zero();
+			for (int32 i = 1; i < N; ++i)
+			{
+				b = Polyline[i] - PlanePos;
+				WindingSum += (double)FMathd::Atan2(a.X*b.Y - a.Y*b.X, a.X*b.X + a.Y*b.Y);
+				a = b;
+			}
+			WindingSum /= FMathd::TwoPi;
+			bool bInside = FMathd::Abs(WindingSum) > 0.3;
+			TempROIBuffer[vid] = bInside ? 1 : 0;
+		}
+		else
+		{
+			TempROIBuffer[vid] = -1;
+		}
+	});
+
+	// convert to vertex selection, and then select fully-enclosed faces
+	FMeshVertexSelection VertexSelection(Mesh);
+	VertexSelection.SelectByVertexID([&](int32 vid) { return TempROIBuffer[vid] == 1; });
+ 	FMeshFaceSelection FaceSelection(Mesh, VertexSelection, FilterProperties->MinTriVertCount);
+	if (FaceSelection.Num() == 0)
+	{
+		return;
+	}
+
+	FLinearColor SetColor = GetInEraseStroke() ? BasicProperties->EraseColor : BasicProperties->PaintColor;
+	SetTrianglesToVertexColor(FaceSelection.AsSet(), SetColor);
+}
+
+// could get rid of this...
+void UMeshVertexPaintTool::SetTrianglesToVertexColor(const TSet<int32>& Triangles, const FLinearColor& ToColor)
+{
+	TempROIBuffer.Reset();
+	for (int32 tid : Triangles)
+	{
+		TempROIBuffer.Add(tid);
+	}
+	SetTrianglesToVertexColor(TempROIBuffer, ToColor);
+}
+
+
+void UMeshVertexPaintTool::SetTrianglesToVertexColor(const TArray<int32>& Triangles, const FLinearColor& ToColor)
+{
+	FDynamicMesh3* Mesh = GetSculptMesh();
+	FDynamicMeshColorOverlay* ColorOverlay = GetActiveColorOverlay();
+
+	const TArray<int32>* UseTriangles = &Triangles;
+	TArray<int32> VisibleTriangles;
+	if (HaveVisibilityFilter())
+	{
+		VisibleTriangles.Reserve(Triangles.Num());
+		ApplyVisibilityFilter(Triangles, VisibleTriangles);
+		UseTriangles = &VisibleTriangles;
+	}
+
+	// convert triangle ROI to element IDs, populating ROIElementSet
+	InitializeElementROIFromTriangleROI(*UseTriangles, false);
+
+	if (ROIElementSet.Num() > 0)
+	{
+		BeginChange();
+		for (int32 ElementID : ROIElementSet)
+		{
+			FVector4f CurColor = ColorOverlay->GetElement(ElementID);
+			FVector4f NewColor = (FVector4f)ToColor;
+			ApplyChannelFilter(CurColor, NewColor);
+			ActiveChangeBuilder->UpdateValue(ElementID, CurColor, NewColor);
+			ColorOverlay->SetElement(ElementID, NewColor);
+		}
+		DynamicMeshComponent->FastNotifyTriangleVerticesUpdated(*UseTriangles, EMeshRenderAttributeFlags::VertexColors);
+		GetToolManager()->PostInvalidation();
+		EndChange();
+	}
+}
+
+
+
+void UMeshVertexPaintTool::InitializeElementROIFromTriangleROI(const TArray<int32>& Triangles, bool bInitializeFlatBuffers)
+{
+	ROIElementBuffer.Reset();
+	ROIColorBuffer.Reset();
+	ROIElementSet.Reset();
+
+	FDynamicMesh3* Mesh = GetSculptMesh();
+	FDynamicMeshColorOverlay* ColorOverlay = GetActiveColorOverlay();
+
+	if (FilterProperties->bHardEdges)
+	{
+		for (int32 tid : Triangles)
+		{
+			if (ColorOverlay->IsSetTriangle(tid))
+			{
+				FIndex3i TriElementIDs = ColorOverlay->GetTriangle(tid);
+				ROIElementSet.Add(TriElementIDs.A);
+				ROIElementSet.Add(TriElementIDs.B);
+				ROIElementSet.Add(TriElementIDs.C);
+			}
+		}
+	}
+	else
+	{
+		// Need to accumulate all Element IDs attached to all vertices of triangles.
+		// Perhaps there is a faster way to do this?
+		TArray<int> TempElementIDs;
+		TempVertexSet.Reset();
+		for (int32 tid : Triangles)
+		{
+			FIndex3i TriVertIDs = Mesh->GetTriangle(tid);
+			for (int32 j = 0; j < 3; ++j)
+			{
+				int32 vid = TriVertIDs[j];
+				if (TempVertexSet.Contains(vid) == false)
+				{
+					TempVertexSet.Add(vid);
+					TempElementIDs.Reset();
+					ColorOverlay->GetVertexElements(vid, TempElementIDs);
+					for (int32 ElementID : TempElementIDs)
+					{
+						ROIElementSet.Add(ElementID);
+					}
+				}
+			}
+		}
+	}
+
+	if (bInitializeFlatBuffers)
+	{
+		ROIElementBuffer = ROIElementSet.Array();
+		ROIColorBuffer.SetNum(ROIElementBuffer.Num());
+		for (int32 k = 0; k < ROIElementBuffer.Num(); ++k)
+		{
+			ROIColorBuffer[k] = ColorOverlay->GetElement(ROIElementBuffer[k]);
+		}
+	}
+}
+
+
+bool UMeshVertexPaintTool::HaveVisibilityFilter() const
+{
+	return FilterProperties->VisibilityFilter != EMeshVertexPaintVisibilityType::None;
+}
+
+
+void UMeshVertexPaintTool::ApplyVisibilityFilter(TSet<int32>& Triangles, TArray<int32>& ROIBuffer, TArray<int32>& OutputBuffer)
+{
+	ROIBuffer.SetNum(0, false);
+	ROIBuffer.Reserve(Triangles.Num());
+	for (int32 tid : Triangles)
+	{
+		ROIBuffer.Add(tid);
+	}
+	
+	OutputBuffer.Reset();
+	ApplyVisibilityFilter(TempROIBuffer, OutputBuffer);
+
+	Triangles.Reset();
+	for (int32 tid : OutputBuffer)
+	{
+		TriangleROI.Add(tid);
+	}
+}
+
+void UMeshVertexPaintTool::ApplyVisibilityFilter(const TArray<int32>& Triangles, TArray<int32>& VisibleTriangles)
+{
+	if (!HaveVisibilityFilter())
+	{
+		VisibleTriangles = Triangles;
+		return;
+	}
+
+	FViewCameraState StateOut;
+	GetToolManager()->GetContextQueriesAPI()->GetCurrentViewState(StateOut);
+	FTransform3d LocalToWorld = UE::ToolTarget::GetLocalToWorldTransform(Target);
+	FVector3d LocalEyePosition(LocalToWorld.InverseTransformPosition(StateOut.Position));
+
+	const FDynamicMesh3* Mesh = GetSculptMesh();
+
+	int32 NumTriangles = Triangles.Num();
+
+	VisibilityFilterBuffer.SetNum(NumTriangles, false);
+	ParallelFor(NumTriangles, [&](int32 idx)
+	{
+		VisibilityFilterBuffer[idx] = true;
+		FVector3d Centroid = Mesh->GetTriCentroid(Triangles[idx]);
+		FVector3d FaceNormal = Mesh->GetTriNormal(Triangles[idx]);
+		if (FaceNormal.Dot((Centroid - LocalEyePosition)) > 0)
+		{
+			VisibilityFilterBuffer[idx] = false;
+		}
+		if (FilterProperties->VisibilityFilter == EMeshVertexPaintVisibilityType::Unoccluded)
+		{
+			int32 HitTID = Octree.FindNearestHitObject(FRay3d(LocalEyePosition, UE::Geometry::Normalized(Centroid - LocalEyePosition)));
+			if (HitTID != Triangles[idx])
+			{
+				VisibilityFilterBuffer[idx] = false;
+			}
+		}
+	});
+
+	VisibleTriangles.Reset();
+	for (int32 k = 0; k < NumTriangles; ++k)
+	{
+		if (VisibilityFilterBuffer[k])
+		{
+			VisibleTriangles.Add(Triangles[k]);
+		}
+	}
+}
+
+
+
+int32 UMeshVertexPaintTool::FindHitSculptMeshTriangle(const FRay3d& LocalRay)
+{
+	if (!IsInBrushSubMode())
+	{
+		return IndexConstants::InvalidID;
+	}
+
+	if (GetBrushCanHitBackFaces())
+	{
+		return Octree.FindNearestHitObject(LocalRay);
+	}
+	else
+	{
+		FDynamicMesh3* Mesh = GetSculptMesh();
+
+		FViewCameraState StateOut;
+		GetToolManager()->GetContextQueriesAPI()->GetCurrentViewState(StateOut);
+		FVector3d LocalEyePosition(CurTargetTransform.InverseTransformPosition((FVector3d)StateOut.Position));
+		int HitTID = Octree.FindNearestHitObject(LocalRay,
+			[this, Mesh, &LocalEyePosition](int TriangleID) {
+			FVector3d Normal, Centroid;
+			double Area;
+			Mesh->GetTriInfo(TriangleID, Normal, Area, Centroid);
+			return Normal.Dot((Centroid - LocalEyePosition)) < 0;
+		});
+		return HitTID;
+	}
+}
+
+int32 UMeshVertexPaintTool::FindHitTargetMeshTriangle(const FRay3d& LocalRay)
+{
+	check(false);
+	return IndexConstants::InvalidID;
+}
+
+
+
+bool UMeshVertexPaintTool::UpdateBrushPosition(const FRay& WorldRay)
+{
+	TUniquePtr<FMeshSculptBrushOp>& UseBrushOp = GetActiveBrushOp();
+
+	bool bHit = false; 
+	ESculptBrushOpTargetType TargetType = UseBrushOp->GetBrushTargetType();
+	switch (TargetType)
+	{
+	case ESculptBrushOpTargetType::SculptMesh:
+	case ESculptBrushOpTargetType::TargetMesh:
+		bHit = UpdateBrushPositionOnSculptMesh(WorldRay, false);
+		break;
+	case ESculptBrushOpTargetType::ActivePlane:
+		check(false);
+		bHit = UpdateBrushPositionOnSculptMesh(WorldRay, false);
+		break;
+	}
+
+	if (bHit && UseBrushOp->GetAlignStampToView())
+	{
+		AlignBrushToView();
+	}
+
+	return bHit;
+}
+
+
+
+
+bool UMeshVertexPaintTool::OnUpdateHover(const FInputDeviceRay& DevicePos)
+{
+	PendingStampType = BasicProperties->PrimaryBrushType;
+
+	if(ensure(InStroke() == false))
+	{
+		UpdateBrushPosition(DevicePos.WorldRay);
+
+		// update strength and falloff for current brush
+		TUniquePtr<FMeshSculptBrushOp>& UseBrushOp = GetActiveBrushOp();
+		UseBrushOp->PropertySet->SetStrength(
+			FMathf::Pow(UMeshSculptToolBase::BrushProperties->FlowRate, 2.0f) );
+		UseBrushOp->PropertySet->SetFalloff(UMeshSculptToolBase::BrushProperties->BrushFalloffAmount);
+
+	}
+	return true;
+}
+
+void UMeshVertexPaintTool::DrawHUD(FCanvas* Canvas, IToolsContextRenderAPI* RenderAPI)
+{
+	if (PolyLassoMechanic)
+	{
+		// because the actual group change is deferred until mouse release, color the lasso to let the user know whether it will erase
+		PolyLassoMechanic->LineColor = GetInEraseStroke() ? FLinearColor::Red : FLinearColor::Green;
+		PolyLassoMechanic->DrawHUD(Canvas, RenderAPI);
+	}
+
+
+	float DPIScale = Canvas->GetDPIScale();
+	UFont* UseFont = GEngine->GetSmallFont();
+	FViewCameraState CamState = RenderAPI->GetCameraState();
+	const FSceneView* SceneView = RenderAPI->GetSceneView();
+	FVector3d LocalEyePosition(CurTargetTransform.InverseTransformPosition((FVector3d)CamState.Position));
+
+	FDynamicMesh3* Mesh = GetSculptMesh();
+
+	if (FilterProperties->bShowHitColor)
+	{
+		FRay3d LocalRay(LocalEyePosition, Normalized(HoverStamp.LocalFrame.Origin - LocalEyePosition));
+		int CursorHitTID = Octree.FindNearestHitObject(LocalRay);
+		if (Mesh->IsTriangle(CursorHitTID))
+		{
+			FIntrRay3Triangle3d Intersection = TMeshQueries<FDynamicMesh3>::RayTriangleIntersection(*GetSculptMesh(), CursorHitTID, LocalRay);
+			FVector3f BaryCoords = (FVector3f)Intersection.TriangleBaryCoords;
+			FVector4f InterpColor;
+			GetActiveColorOverlay()->GetTriBaryInterpolate<float>(CursorHitTID, &BaryCoords.X, &InterpColor.X);
+			FVector2D CursorPixelPos;
+			SceneView->WorldToPixel(HoverStamp.WorldFrame.Origin, CursorPixelPos);
+			FString CursorString = FString::Printf(TEXT("%.3f %.3f %.3f %.3f"), InterpColor.X, InterpColor.Y, InterpColor.Z, InterpColor.W);
+			Canvas->DrawShadowedString(CursorPixelPos.X / (double)DPIScale, CursorPixelPos.Y / (double)DPIScale, *CursorString, UseFont, FLinearColor::White);
+		}
+	}
+
+}
+
+
+void UMeshVertexPaintTool::OnTick(float DeltaTime)
+{
+	TRACE_CPUPROFILER_EVENT_SCOPE(MeshVPaint_OnTick);
+
+	UMeshSculptToolBase::OnTick(DeltaTime);
+	MeshElementsDisplay->OnTick(DeltaTime);
+
+	bool bIsLasso = (BasicProperties->SubToolType == EMeshVertexPaintInteractionType::PolyLasso);
+	PolyLassoMechanic->SetIsEnabled(bIsLasso);
+
+	ConfigureIndicator(FilterProperties->BrushAreaMode == EMeshVertexPaintBrushAreaType::Volumetric);
+	SetIndicatorVisibility(bIsLasso == false);
+
+	if (bHavePendingAction)
+	{
+		ApplyAction(PendingAction);
+		bHavePendingAction = false;
+		PendingAction = EMeshVertexPaintToolActions::NoAction;
+	}
+
+	// process the undo update
+	if (bUndoUpdatePending)
+	{
+		// wait for updates
+		WaitForPendingUndoRedo();
+
+		// post rendering update
+		DynamicMeshComponent->FastNotifyTriangleVerticesUpdated(AccumulatedTriangleROI, EMeshRenderAttributeFlags::VertexColors);
+		GetToolManager()->PostInvalidation();
+
+		// ignore stamp and wait for next tick to do anything else
+		bUndoUpdatePending = false;
+		return;
+	}
+
+	if (bPendingPickColor || bPendingPickEraseColor )
+	{
+		int32 HitTriangleID = GetBrushTriangleID();
+		if (HitTriangleID >= 0 && IsStampPending() == false )
+		{
+			if (GetSculptMesh()->IsTriangle(HitTriangleID))
+			{
+				if (bPendingPickColor || bPendingPickEraseColor)
+				{
+					if (const FDynamicMeshAttributeSet* AttributeSet = GetSculptMesh()->Attributes())
+					{
+						if (const FDynamicMeshColorOverlay* ColorAttrib = AttributeSet->PrimaryColors())
+						{
+							if (ColorAttrib->IsSetTriangle(HitTriangleID))
+							{
+								FVector4f HitColor;
+								// TODO: use actual hit point here...
+								FVector3f BaryCoords = (1.0f / 3.0f) * FVector3f::One();
+								ColorAttrib->GetTriBaryInterpolate<float>(HitTriangleID, &BaryCoords.X, &HitColor.X);
+								if (bPendingPickEraseColor)
+								{
+									BasicProperties->EraseColor = HitColor;
+								}
+								else
+								{
+									BasicProperties->PaintColor = HitColor;
+								}
+								NotifyOfPropertyChangeByTool(BasicProperties);
+							}
+						}
+					}
+				}
+			}
+		}
+		bPendingPickColor = bPendingPickEraseColor = false;
+	}
+
+
+	if (IsInBrushSubMode())
+	{
+		if (InStroke())
+		{
+			TRACE_CPUPROFILER_EVENT_SCOPE(MeshVPaint_OnTick_InStroke);
+
+			// update brush position
+			if (UpdateStampPosition(GetPendingStampRayWorld()) == false)
+			{
+				return;
+			}
+			UpdateStampPendingState();
+			if (IsStampPending() == false)
+			{
+				return;
+			}
+
+			// update sculpt ROI
+			UpdateROI(CurrentStamp);
+
+			// append updated ROI to modified region (async)
+			TFuture<void> AccumulateROI = Async(VertexPaintToolAsyncExecTarget, [&]()
+			{
+				AccumulatedTriangleROI.Append(TriangleROI);
+			});
+
+			// apply the stamp
+			bool bColorModified = ApplyStamp();
+
+			if (bColorModified)
+			{
+				TRACE_CPUPROFILER_EVENT_SCOPE(MeshVPaint_OnTick_UpdateComponent);
+				DynamicMeshComponent->FastNotifyTriangleVerticesUpdated(TriangleROI, EMeshRenderAttributeFlags::VertexColors);
+				GetToolManager()->PostInvalidation();
+			}
+
+			// we don't really need to wait for these to happen to end Tick()...
+			AccumulateROI.Wait();
+		}
+	}
+
+}
+
+
+void UMeshVertexPaintTool::ApplyChannelFilter(const FVector4f& CurColor, FVector4f& NewColor)
+{
+	const FModelingToolsColorChannelFilter& ChannelFilter = BasicProperties->ChannelFilter;
+	if (ChannelFilter.bRed == false)
+	{
+		NewColor.X = CurColor.X;
+	}
+	if (ChannelFilter.bGreen == false)
+	{
+		NewColor.Y = CurColor.Y;
+	}
+	if (ChannelFilter.bBlue == false)
+	{
+		NewColor.Z = CurColor.Z;
+	}
+	if (ChannelFilter.bAlpha == false)
+	{
+		NewColor.W = CurColor.W;
+	}
+}
+
+void UMeshVertexPaintTool::OnChannelFilterModified()
+{
+	FModelingToolsColorChannelFilter NewFilter = BasicProperties->ChannelFilter;
+
+	// if we are showing original material then we do not want to filter the colors by default (could be an option)
+	if (FilterProperties->MaterialMode == EMeshVertexPaintMaterialMode::OriginalMaterial)
+	{
+		NewFilter = FModelingToolsColorChannelFilter();
+	}
+
+	if (NewFilter.bRed && NewFilter.bGreen && NewFilter.bBlue)
+	{
+		DynamicMeshComponent->ClearVertexColorRemappingFunction();
+	}
+	else if (NewFilter.bRed || NewFilter.bBlue || NewFilter.bGreen)
+	{
+		DynamicMeshComponent->SetVertexColorRemappingFunction([Filter = NewFilter](FVector4f& Color)
+		{
+			Color.X = (Filter.bRed) ? Color.X : 0;
+			Color.Y = (Filter.bGreen) ? Color.Y : 0;
+			Color.Z = (Filter.bBlue) ? Color.Z : 0;
+		});
+	}
+	else if (NewFilter.bAlpha)
+	{
+		DynamicMeshComponent->SetVertexColorRemappingFunction([Filter = NewFilter](FVector4f& Color)
+		{
+			Color.X = Color.W;
+			Color.Y = Color.W;
+			Color.Z = Color.W;
+		});
+	}
+	else
+	{
+		DynamicMeshComponent->ClearVertexColorRemappingFunction();
+	}
+}
+
+
+FColor UMeshVertexPaintTool::GetColorForGroup(int32 GroupID)
+{
+	return LinearColors::SelectFColor(GroupID);
+}
+
+
+
+void UMeshVertexPaintTool::UpdateVertexPaintMaterialMode()
+{
+	if (FilterProperties->MaterialMode == EMeshVertexPaintMaterialMode::OriginalMaterial)
+	{
+		UpdateMaterialMode(EMeshEditingMaterialModes::ExistingMaterial);
+	}
+	else
+	{
+		// TODO: add two-sided vertex color support when it becomes available
+		UpdateMaterialMode(EMeshEditingMaterialModes::VertexColor);
+		if (FilterProperties->MaterialMode == EMeshVertexPaintMaterialMode::UnlitVertexColor
+			&& GEngine && GEngine->VertexColorMaterial)
+		{
+			GetSculptMeshComponent()->SetOverrideRenderMaterial(GEngine->VertexColorMaterial);
+		}
+	}
+
+	OnChannelFilterModified();		// channel visibility filtering may be disabled for existing material
+}
+
+
+
+void UMeshVertexPaintTool::FloodFillColorAction(FLinearColor Color)
+{
+	const FDynamicMesh3* Mesh = DynamicMeshComponent->GetMesh();
+	TempROIBuffer.Reset();
+	TempROIBuffer.Reserve(Mesh->TriangleCount());
+	for (int32 tid : Mesh->TriangleIndicesItr())
+	{
+		TempROIBuffer.Add(tid);
+	}
+
+	SetTrianglesToVertexColor(TempROIBuffer, Color);
+}
+
+
+void UMeshVertexPaintTool::ApplyCurrentUtilityAction()
+{
+	switch (UtilityActions->Operation)
+	{
+	case EMeshVertexPaintToolUtilityOperations::BlendAllSeams:
+		BlendAllSeams();
+		return;
+	case EMeshVertexPaintToolUtilityOperations::FillChannels:
+		FillChannels();
+		return;
+	case EMeshVertexPaintToolUtilityOperations::InvertChannels:
+		InvertChannels();
+		return;
+	case EMeshVertexPaintToolUtilityOperations::CopyChannelToChannel:
+		CopyChannelToChannel();
+		return;
+	case EMeshVertexPaintToolUtilityOperations::SwapChannels:
+		SwapChannels();
+		return;
+	case EMeshVertexPaintToolUtilityOperations::CopyFromWeightMap:
+		CopyFromWeightMap();
+		return;
+	case EMeshVertexPaintToolUtilityOperations::CopyToOtherLODs:
+		CopyToOtherLODs();
+		return;
+	case EMeshVertexPaintToolUtilityOperations::CopyToSingleLOD:
+		CopyToSpecificLOD();
+		return;
+	}
+}
+
+
+
+void UMeshVertexPaintTool::BlendAllSeams()
+{
+	BeginChange();
+
+	TArray<int32> TempElementIDs;
+	const FDynamicMesh3* Mesh = DynamicMeshComponent->GetMesh();
+	FDynamicMeshColorOverlay* VertexColors = GetActiveColorOverlay();
+	for (int32 VertexID : Mesh->VertexIndicesItr())
+	{
+		TempElementIDs.Reset();
+		VertexColors->GetVertexElements(VertexID, TempElementIDs);
+		if (TempElementIDs.Num() > 1)
+		{
+			FVector4f SumColor = FVector4f::Zero();
+			for (int32 ElementID : TempElementIDs)
+			{
+				SumColor += VertexColors->GetElement(ElementID);
+			}
+			SumColor /= (double)TempElementIDs.Num();
+			for (int32 ElementID : TempElementIDs)
+			{
+				ActiveChangeBuilder->UpdateValue(ElementID, VertexColors->GetElement(ElementID), SumColor);
+				VertexColors->SetElement(ElementID, SumColor);
+			}
+		}
+	}
+
+	DynamicMeshComponent->FastNotifyVertexAttributesUpdated(EMeshRenderAttributeFlags::VertexColors);
+	GetToolManager()->PostInvalidation();
+
+	EndChange();
+}
+
+void UMeshVertexPaintTool::FillChannels()
+{
+	float FillValue = UtilityActions->SourceValue;
+	FIndex4i Targets = UtilityActions->TargetChannels.AsFlags();
+	if (Targets == FIndex4i::Zero())
+	{
+		GetToolManager()->DisplayMessage(
+			LOCTEXT("NoChannelsSelectedMessage", "At least one Target Channel must be selected"), EToolMessageLevel::UserError);
+		return;
+	}
+
+	BeginChange();
+	FDynamicMeshColorOverlay* VertexColors = GetActiveColorOverlay();
+	for (int32 ElementID : VertexColors->ElementIndicesItr())
+	{
+		FVector4f CurColor = VertexColors->GetElement(ElementID);
+		FVector4f NewColor(
+			Targets[0] ? FillValue : CurColor.X, Targets[1] ? FillValue : CurColor.Y,
+			Targets[2] ? FillValue : CurColor.Z, Targets[3] ? FillValue : CurColor.W);
+		ActiveChangeBuilder->UpdateValue(ElementID, CurColor, NewColor);
+		VertexColors->SetElement(ElementID, NewColor);
+	}
+
+	DynamicMeshComponent->FastNotifyVertexAttributesUpdated(EMeshRenderAttributeFlags::VertexColors);
+	GetToolManager()->PostInvalidation();
+	EndChange();
+}
+
+void UMeshVertexPaintTool::InvertChannels()
+{
+	FIndex4i Targets = UtilityActions->TargetChannels.AsFlags();
+	if (Targets == FIndex4i::Zero())
+	{
+		GetToolManager()->DisplayMessage(
+			LOCTEXT("NoChannelsSelectedMessage", "At least one Target Channel must be selected"), EToolMessageLevel::UserError);
+		return;
+	}
+
+	BeginChange();
+	FDynamicMeshColorOverlay* VertexColors = GetActiveColorOverlay();
+	for (int32 ElementID : VertexColors->ElementIndicesItr())
+	{
+		FVector4f CurColor = VertexColors->GetElement(ElementID);
+		FVector4f NewColor(
+			Targets[0] ? FMathf::Clamp(1.0f-CurColor.X,0,1.0f) : CurColor.X, 
+			Targets[1] ? FMathf::Clamp(1.0f-CurColor.Y,0,1.0f) : CurColor.Y,
+			Targets[2] ? FMathf::Clamp(1.0f-CurColor.Z,0,1.0f) : CurColor.Z,
+			Targets[3] ? FMathf::Clamp(1.0f-CurColor.W,0,1.0f) : CurColor.W);
+		ActiveChangeBuilder->UpdateValue(ElementID, CurColor, NewColor);
+		VertexColors->SetElement(ElementID, NewColor);
+	}
+
+	DynamicMeshComponent->FastNotifyVertexAttributesUpdated(EMeshRenderAttributeFlags::VertexColors);
+	GetToolManager()->PostInvalidation();
+	EndChange();
+}
+
+
+void UMeshVertexPaintTool::CopyChannelToChannel()
+{
+	int32 SourceChannelIndex = VectorUtil::Clamp(static_cast<int32>(UtilityActions->SourceChannel), 0, 3);
+	FIndex4i Targets = UtilityActions->TargetChannels.AsFlags();
+	Targets[SourceChannelIndex] = 0;
+	if (Targets == FIndex4i::Zero())
+	{
+		GetToolManager()->DisplayMessage(
+			LOCTEXT("NoChannelsSelectedMessage", "At least one Target Channel must be selected"), EToolMessageLevel::UserError);
+		return;
+	}
+
+	BeginChange();
+
+	FDynamicMeshColorOverlay* VertexColors = GetActiveColorOverlay();
+	for (int32 ElementID : VertexColors->ElementIndicesItr())
+	{
+		FVector4f CurColor = VertexColors->GetElement(ElementID);
+		FVector4f NewColor = CurColor;
+		float SourceValue = CurColor[SourceChannelIndex];
+		for (int32 j = 0; j < 4; ++j)
+		{
+			if (Targets[j] != 0)
+			{
+				NewColor[j] = SourceValue;
+			}
+		}
+		ActiveChangeBuilder->UpdateValue(ElementID, CurColor, NewColor);
+		VertexColors->SetElement(ElementID, NewColor);
+	}
+
+	DynamicMeshComponent->FastNotifyVertexAttributesUpdated(EMeshRenderAttributeFlags::VertexColors);
+	GetToolManager()->PostInvalidation();
+
+	EndChange();
+}
+
+void UMeshVertexPaintTool::SwapChannels()
+{
+	int32 SourceChannelIndex = VectorUtil::Clamp(static_cast<int32>(UtilityActions->SourceChannel), 0, 3);
+	int32 DestChannelIndex = VectorUtil::Clamp(static_cast<int32>(UtilityActions->TargetChannel), 0, 3);
+	if (SourceChannelIndex == DestChannelIndex)
+	{
+		GetToolManager()->DisplayMessage(
+			LOCTEXT("SameSourceDestMessage", "Source and Destination must be different Channels"), EToolMessageLevel::UserError);
+		return;
+	}
+
+	BeginChange();
+
+	FDynamicMeshColorOverlay* VertexColors = GetActiveColorOverlay();
+	for (int32 ElementID : VertexColors->ElementIndicesItr())
+	{
+		FVector4f CurColor = VertexColors->GetElement(ElementID);
+		FVector4f NewColor = CurColor;
+		Swap(NewColor[SourceChannelIndex], NewColor[DestChannelIndex]);
+		ActiveChangeBuilder->UpdateValue(ElementID, CurColor, NewColor);
+		VertexColors->SetElement(ElementID, NewColor);
+	}
+
+	DynamicMeshComponent->FastNotifyVertexAttributesUpdated(EMeshRenderAttributeFlags::VertexColors);
+	GetToolManager()->PostInvalidation();
+
+	EndChange();
+}
+
+
+
+void UMeshVertexPaintTool::CopyFromWeightMap()
+{
+	FIndex4i Targets = UtilityActions->TargetChannels.AsFlags();
+	if (Targets == FIndex4i::Zero())
+	{
+		GetToolManager()->DisplayMessage(
+			LOCTEXT("NoChannelsSelectedMessage", "At least one Target Channel must be selected"), EToolMessageLevel::UserError);
+		return;
+	}
+
+	FIndexedWeightMap1f WeightMap;
+	bool bFound = UE::WeightMaps::GetVertexWeightMap(UE::ToolTarget::GetMeshDescription(Target),
+		UtilityActions->WeightMap, WeightMap, 1.0f);
+
+	if (!bFound)
+	{
+		GetToolManager()->DisplayMessage(
+			LOCTEXT("NoWeightMapMessage", "Selected Weight Map could not be found on the Target Mesh"), EToolMessageLevel::UserError);
+		return;
+	}
+
+	// todo: range mapping
+
+	BeginChange();
+	FDynamicMeshColorOverlay* VertexColors = GetActiveColorOverlay();
+	for (int32 ElementID : VertexColors->ElementIndicesItr())
+	{
+		int32 ParentVertex = VertexColors->GetParentVertex(ElementID);
+		float Value = WeightMap.GetValue(ParentVertex);
+		Value = FMathf::Clamp(Value, 0.0f, 1.0f);
+
+		FVector4f CurColor = VertexColors->GetElement(ElementID);
+		FVector4f NewColor(
+			Targets[0] ? Value : CurColor.X, Targets[1] ? Value : CurColor.Y,
+			Targets[2] ? Value : CurColor.Z, Targets[3] ? Value : CurColor.W);
+		ActiveChangeBuilder->UpdateValue(ElementID, CurColor, NewColor);
+		VertexColors->SetElement(ElementID, NewColor);
+	}
+
+	DynamicMeshComponent->FastNotifyVertexAttributesUpdated(EMeshRenderAttributeFlags::VertexColors);
+	GetToolManager()->PostInvalidation();
+	EndChange();
+}
+
+
+// TODO should be moved to shared location
+namespace UELocal
+{
+
+static void CopyVertexColorsToLOD(
+	const FDynamicMesh3& FromMesh, 
+	const FDynamicMeshColorOverlay& FromColors,
+	const FDynamicMeshAABBTree3& Spatial,
+	FDynamicMesh3& ToMesh,
+	const FVector4f& FallbackColor)
+{
+	if (ToMesh.HasAttributes() == false)
+	{
+		ToMesh.EnableAttributes();
+	}
+	if (ToMesh.Attributes()->HasPrimaryColors() == false)
+	{
+		ToMesh.Attributes()->EnablePrimaryColors();
+	}
+	FDynamicMeshColorOverlay& ToColors = *ToMesh.Attributes()->PrimaryColors();
+	ToColors.ClearElements();
+
+	// maybe should use baking here...?
+	// could be parallel...
+
+	TArray<FVector4f> VertexColors;
+	VertexColors.SetNum(ToMesh.MaxVertexID());
+	for (int32 vid : ToMesh.VertexIndicesItr())
+	{
+		FVector3d ToPos = ToMesh.GetVertex(vid);
+		double NearDistSqr;
+		int32 NearestTri = Spatial.FindNearestTriangle(ToPos, NearDistSqr);
+		if (FromColors.IsSetTriangle(NearestTri) == false)
+		{
+			VertexColors[vid] = FallbackColor;
+			break;
+		}
+
+		FDistPoint3Triangle3d Dist = TMeshQueries<FDynamicMesh3>::TriangleDistance(FromMesh, NearestTri, ToPos);
+		FVector3f BaryCoords = (FVector3f)Dist.TriangleBaryCoords;
+		FVector4f InterpColor;
+		FromColors.GetTriBaryInterpolate<float>(NearestTri, &BaryCoords.X, &InterpColor.X);
+		VertexColors[vid] = InterpColor;
+	}
+
+	TArray<int32> VertToElemMap;
+	VertToElemMap.SetNum(VertexColors.Num());
+	for (int32 vid : ToMesh.VertexIndicesItr())
+	{
+		VertToElemMap[vid] = ToColors.AppendElement(VertexColors[vid]);
+	}
+	for (int32 tid : ToMesh.TriangleIndicesItr())
+	{
+		FIndex3i TriVerts = ToMesh.GetTriangle(tid);
+		ToColors.SetTriangle(tid, FIndex3i(
+			VertToElemMap[TriVerts.A], VertToElemMap[TriVerts.B], VertToElemMap[TriVerts.C]));
+	}
+}
+
+static void GetMeshDescriptionLODWithCopiedVertexColors(
+	UToolTarget* Target, EMeshLODIdentifier CopyToLOD,
+	FDynamicMesh3& SourceMesh, FDynamicMeshColorOverlay& ColorOverlay, FDynamicMeshAABBTree3& SourceMeshSpatial,
+	FVector4f DefaultColor,
+	FMeshDescription& UpdatedMeshDescriptionOut)
+{
+	UpdatedMeshDescriptionOut = UE::ToolTarget::GetMeshDescriptionCopy(Target, FGetMeshParameters(true, CopyToLOD));
+
+	FDynamicMesh3 LODMesh;
+	FMeshDescriptionToDynamicMesh Converter1;
+	Converter1.bTransformVertexColorsLinearToSRGB = false;
+	Converter1.Convert(&UpdatedMeshDescriptionOut, LODMesh);
+
+	UELocal::CopyVertexColorsToLOD(SourceMesh, ColorOverlay, SourceMeshSpatial, LODMesh, DefaultColor);
+
+	FDynamicMeshToMeshDescription Converter2;
+	Converter2.ConversionOptions.bTransformVtxColorsSRGBToLinear = false;
+	Converter2.UpdateVertexColors(&LODMesh, UpdatedMeshDescriptionOut);
+}
+
+
+}
+
+
+
+void UMeshVertexPaintTool::CopyToOtherLODs()
+{
+	// have to do this via MeshDescriptions for now
+	if (bTargetSupportsLODs == false)
+	{
+		GetToolManager()->DisplayMessage(
+			LOCTEXT("NoLODsErrorMessage", "Target Mesh does not support LODs"), EToolMessageLevel::UserError);
+		return;
+	}
+	TArray<EMeshLODIdentifier> CopyToLODs = AvailableLODs;
+	if (UtilityActions->bCopyToHiRes == false)
+	{
+		CopyToLODs.Remove(EMeshLODIdentifier::HiResSource);
+	}
+	if (CopyToLODs.Num() == 0)
+	{
+		GetToolManager()->DisplayMessage(
+			LOCTEXT("NoLODsToCopyToMessage", "Target Mesh does not have any other editable LODs"), EToolMessageLevel::UserError);
+		return;
+	}
+
+	FDynamicMesh3* SourceMesh = GetSculptMesh();
+	FDynamicMeshColorOverlay* ColorOverlay = GetActiveColorOverlay();
+	FDynamicMeshAABBTree3 SourceMeshSpatial(SourceMesh, true);
+
+	TArray<FMeshDescription> LODMeshDescriptions;
+	LODMeshDescriptions.SetNum(CopyToLODs.Num());
+	for (int32 k = 0; k < CopyToLODs.Num(); ++k)
+	{
+		EMeshLODIdentifier LOD = CopyToLODs[k];
+		UELocal::GetMeshDescriptionLODWithCopiedVertexColors(Target, LOD,
+			*SourceMesh, *ColorOverlay, SourceMeshSpatial, (FVector4f)BasicProperties->PaintColor, LODMeshDescriptions[k]);
+	}
+
+	GetToolManager()->BeginUndoTransaction(LOCTEXT("CopyToOtherLODs", "Copy Vertex Colors To LODs"));
+	for (int32 k = 0; k < CopyToLODs.Num(); ++k)
+	{
+		EMeshLODIdentifier LOD = CopyToLODs[k];
+		UE::ToolTarget::CommitMeshDescriptionUpdate(Target, &LODMeshDescriptions[k], nullptr /*no material set changes*/, FCommitMeshParameters(true, LOD));
+	}
+	GetToolManager()->EndUndoTransaction();
+}
+
+
+void UMeshVertexPaintTool::CopyToSpecificLOD()
+{
+	// have to do this via MeshDescriptions for now
+	if (bTargetSupportsLODs == false)
+	{
+		GetToolManager()->DisplayMessage(
+			LOCTEXT("NoLODsErrorMessage", "Target Mesh does not support LODs"), EToolMessageLevel::UserError);
+		return;
+	}
+
+	int32 Index = UtilityActions->LODNamesList.IndexOfByKey(UtilityActions->CopyToLODName);
+	if (Index == INDEX_NONE)
+	{
+		GetToolManager()->DisplayMessage(
+			LOCTEXT("MissingLODMessage", "Could not find selected LOD on Target Mesh"), EToolMessageLevel::UserError);
+		return;
+	}
+
+	FDynamicMesh3* SourceMesh = GetSculptMesh();
+	FDynamicMeshColorOverlay* ColorOverlay = GetActiveColorOverlay();
+	FDynamicMeshAABBTree3 SourceMeshSpatial(SourceMesh, true);
+
+	FMeshDescription LODMeshDescription;
+	EMeshLODIdentifier LOD = AvailableLODs[Index];
+	UELocal::GetMeshDescriptionLODWithCopiedVertexColors(Target, LOD,
+		*SourceMesh, *ColorOverlay, SourceMeshSpatial, (FVector4f)BasicProperties->PaintColor, LODMeshDescription);
+
+	GetToolManager()->BeginUndoTransaction(LOCTEXT("CopyToSpecificLOD", "Copy Vertex Colors To LOD"));
+	UE::ToolTarget::CommitMeshDescriptionUpdate(Target, &LODMeshDescription, nullptr /*no material set changes*/, FCommitMeshParameters(true, LOD));
+	GetToolManager()->EndUndoTransaction();
+}
+
+
+//
+// Change Tracking
+//
+
+void UMeshVertexPaintTool::BeginChange()
+{
+	check(ActiveChangeBuilder == nullptr);
+
+	ActiveChangeBuilder = MakeUnique<TIndexedValuesChangeBuilder<FVector4f, FMeshVertexColorPaintChange>>();
+	ActiveChangeBuilder->BeginNewChange();
+}
+
+
+void UMeshVertexPaintTool::EndChange()
+{
+	check(ActiveChangeBuilder);
+
+	TUniquePtr<FMeshVertexColorPaintChange> EditResult = ActiveChangeBuilder->ExtractResult();
+	ActiveChangeBuilder = nullptr;
+	
+	EditResult->ApplyFunction = [](UObject* Object, const int32& AttribIndex, const TArray<int32>& Indices, const TArray<FVector4f>& Values)
+	{
+		UMeshVertexPaintTool* Tool = CastChecked<UMeshVertexPaintTool>(Object);
+		Tool->ExternalUpdateValues(Indices, Values);
+	};
+	EditResult->RevertFunction = [](UObject* Object, const int32& AttribIndex, const TArray<int32>& Indices, const TArray<FVector4f>& Values)
+	{
+		UMeshVertexPaintTool* Tool = CastChecked<UMeshVertexPaintTool>(Object);
+		Tool->ExternalUpdateValues(Indices, Values);
+	};
+
+	TUniquePtr<TWrappedToolCommandChange<FMeshVertexColorPaintChange>> NewChange = MakeUnique<TWrappedToolCommandChange<FMeshVertexColorPaintChange>>();
+	NewChange->WrappedChange = MoveTemp(EditResult);
+	NewChange->BeforeModify = [this](bool bRevert)
+	{
+		this->WaitForPendingUndoRedo();
+	};
+
+	GetToolManager()->EmitObjectChange(this, MoveTemp(NewChange), LOCTEXT("VertexPaintChange", "Paint Stroke"));
+}
+
+
+
+void UMeshVertexPaintTool::ExternalUpdateValues(const TArray<int32>& ElementIDs, const TArray<FVector4f>& NewValues)
+{
+	DynamicMeshComponent->EditMesh([&](FDynamicMesh3& Mesh)
+	{
+		FDynamicMeshColorOverlay* ColorOverlay = Mesh.Attributes()->PrimaryColors();
+
+		int32 N = ElementIDs.Num();
+		for ( int32 k = 0; k < N; ++k)
+		{
+			int32 ElementID = ElementIDs[k];
+			ColorOverlay->SetElement(ElementID, NewValues[k]);
+		}
+	});
+
+	GetToolManager()->PostInvalidation();
+}
+
+
+
+void UMeshVertexPaintTool::WaitForPendingUndoRedo()
+{
+	if (bUndoUpdatePending)
+	{
+		bUndoUpdatePending = false;
+	}
+}
+
+void UMeshVertexPaintTool::OnDynamicMeshComponentChanged()
+{
+	// update octree
+	FDynamicMesh3* Mesh = GetSculptMesh();
+
+	// make sure any previous async computations are done, and update the undo ROI
+	if (bUndoUpdatePending)
+	{
+		// we should never hit this anymore, because of pre-change calling WaitForPendingUndoRedo()
+		WaitForPendingUndoRedo();
+
+		// TODO: do we need to read from mesh change here??
+		//UE::Geometry::VertexToTriangleOneRing(Mesh, Change->Vertices, AccumulatedTriangleROI);
+	}
+	else
+	{
+		// TODO: do we need to read from mesh change here??
+		//UE::Geometry::VertexToTriangleOneRing(Mesh, Change->Vertices, AccumulatedTriangleROI);
+	}
+
+	// note that we have a pending update
+	bUndoUpdatePending = true;
+}
+
+
+void UMeshVertexPaintTool::PrecomputeFilterData()
+{
+	const FDynamicMesh3* Mesh = GetSculptMesh();
+	
+	TriNormals.SetNum(Mesh->MaxTriangleID());
+	ParallelFor(Mesh->MaxTriangleID(), [&](int32 tid)
+	{
+		if (Mesh->IsTriangle(tid))
+		{
+			TriNormals[tid] = Mesh->GetTriNormal(tid);
+		}
+	});
+
+	const FDynamicMeshNormalOverlay* Normals = Mesh->Attributes()->PrimaryNormals();
+	const FDynamicMeshUVOverlay* UVs = Mesh->Attributes()->PrimaryUV();
+	UVSeamEdges.SetNum(Mesh->MaxEdgeID());
+	NormalSeamEdges.SetNum(Mesh->MaxEdgeID());
+	ParallelFor(Mesh->MaxEdgeID(), [&](int32 eid)
+	{
+		if (Mesh->IsEdge(eid))
+		{
+			UVSeamEdges[eid] = UVs->IsSeamEdge(eid);
+			NormalSeamEdges[eid] = Normals->IsSeamEdge(eid);
+		}
+	});
+}
+
+
+void UMeshVertexPaintTool::OnSelectedGroupLayerChanged()
+{
+	GetToolManager()->BeginUndoTransaction(LOCTEXT("ChangeActiveGroupLayer", "Change Polygroup Layer"));
+
+	int32 ActiveLayerIndex = (ActiveGroupSet) ? ActiveGroupSet->GetPolygroupIndex() : -1;
+	UpdateActiveGroupLayer();
+	int32 NewLayerIndex = (ActiveGroupSet) ? ActiveGroupSet->GetPolygroupIndex() : -1;
+
+	if (ActiveLayerIndex != NewLayerIndex)
+	{
+		TUniquePtr<TSimpleValueLambdaChange<int32>> GroupLayerChange = MakeUnique<TSimpleValueLambdaChange<int32>>();
+		GroupLayerChange->FromValue = ActiveLayerIndex;
+		GroupLayerChange->ToValue = NewLayerIndex;
+		GroupLayerChange->ValueChangeFunc = [this](UObject*, int32 FromIndex, int32 ToIndex, bool)
+		{
+			this->PolygroupLayerProperties->SetSelectedFromPolygroupIndex(ToIndex);
+			this->PolygroupLayerProperties->SilentUpdateWatched();		// to prevent OnSelectedGroupLayerChanged() from being called immediately
+			this->UpdateActiveGroupLayer();
+		};
+		GetToolManager()->EmitObjectChange(this, MoveTemp(GroupLayerChange), LOCTEXT("ChangeActiveGroupLayer", "Change Polygroup Layer"));
+	}
+
+	GetToolManager()->EndUndoTransaction();
+}
+
+
+void UMeshVertexPaintTool::UpdateActiveGroupLayer()
+{
+	if (PolygroupLayerProperties->HasSelectedPolygroup() == false)
+	{
+		ActiveGroupSet = MakeUnique<UE::Geometry::FPolygroupSet>(GetSculptMesh());
+	}
+	else
+	{
+		FName SelectedName = PolygroupLayerProperties->ActiveGroupLayer;
+		FDynamicMeshPolygroupAttribute* FoundAttrib = UE::Geometry::FindPolygroupLayerByName(*GetSculptMesh(), SelectedName);
+		ensureMsgf(FoundAttrib, TEXT("Selected Attribute Not Found! Falling back to Default group layer."));
+		ActiveGroupSet = MakeUnique<UE::Geometry::FPolygroupSet>(GetSculptMesh(), FoundAttrib);
+	}
+
+	// update colors
+	DynamicMeshComponent->FastNotifyVertexAttributesUpdated(EMeshRenderAttributeFlags::VertexColors);
+	GetToolManager()->PostInvalidation();
+}
+
+
+
+void UMeshVertexPaintTool::UpdateSubToolType(EMeshVertexPaintInteractionType NewType)
+{
+	FilterProperties->CurrentSubToolType = BasicProperties->SubToolType;
+
+	bool bSculptPropsVisible = (NewType == EMeshVertexPaintInteractionType::Brush);
+	SetToolPropertySourceEnabled(UMeshSculptToolBase::BrushProperties, bSculptPropsVisible);
+
+	//SetToolPropertySourceEnabled(BasicProperties, true);
+	SetBrushOpPropsVisibility(false);
+}
+
+
+void UMeshVertexPaintTool::UpdateBrushType(EMeshVertexPaintBrushType BrushType)
+{
+	static const FText BaseMessage = LOCTEXT("OnStartTool", "Hold Shift to Erase. [/] and S/D change Size (+Shift to small-step). Shift+G to pick Paint Color, +Ctrl for Erase Color.");
+	FTextBuilder Builder;
+	Builder.AppendLine(BaseMessage);
+
+	SetActivePrimaryBrushType((int32)BrushType);
+
+	SetToolPropertySourceEnabled(GizmoProperties, false);
+
+	GetToolManager()->DisplayMessage(Builder.ToText(), EToolMessageLevel::UserNotification);
+}
+
+void UMeshVertexPaintTool::UpdateSecondaryBrushType(EMeshVertexPaintSecondaryActionType NewType)
+{
+	if (NewType == EMeshVertexPaintSecondaryActionType::Erase)
+	{
+		SetActiveSecondaryBrushType((int32)EMeshVertexPaintBrushType::Erase);
+	}
+	else if (NewType == EMeshVertexPaintSecondaryActionType::Soften)
+	{
+		SetActiveSecondaryBrushType((int32)EMeshVertexPaintBrushType::Soften);
+	}
+	else if (NewType == EMeshVertexPaintSecondaryActionType::Smooth)
+	{
+		SetActiveSecondaryBrushType((int32)EMeshVertexPaintBrushType::Smooth);
+	}
+}
+
+
+void UMeshVertexPaintTool::RequestAction(EMeshVertexPaintToolActions ActionType)
+{
+	if (!bHavePendingAction)
+	{
+		PendingAction = ActionType;
+		bHavePendingAction = true;
+	}
+}
+
+
+void UMeshVertexPaintTool::ApplyAction(EMeshVertexPaintToolActions ActionType)
+{
+	switch (ActionType)
+	{
+
+	case EMeshVertexPaintToolActions::PaintAll:
+		FloodFillColorAction(BasicProperties->PaintColor);
+		break;
+
+	case EMeshVertexPaintToolActions::EraseAll:
+		FloodFillColorAction(BasicProperties->EraseColor);
+		break;
+
+	case EMeshVertexPaintToolActions::FillBlack:
+		FloodFillColorAction(FLinearColor::Black);
+		break;
+
+	case EMeshVertexPaintToolActions::FillWhite:
+		FloodFillColorAction(FLinearColor::White);
+		break;
+
+	case EMeshVertexPaintToolActions::ApplyCurrentUtility:
+		ApplyCurrentUtilityAction();
+		break;
+	}
+}
+
+
+
+
+
+#undef LOCTEXT_NAMESPACE
+
diff --git a/Engine/Plugins/Experimental/MeshModelingToolsetExp/Source/MeshModelingToolsExp/Public/MeshVertexPaintTool.h b/Engine/Plugins/Experimental/MeshModelingToolsetExp/Source/MeshModelingToolsExp/Public/MeshVertexPaintTool.h
new file mode 100644
--- /dev/null
+++ b/Engine/Plugins/Experimental/MeshModelingToolsetExp/Source/MeshModelingToolsExp/Public/MeshVertexPaintTool.h
@@ -0,0 +1,681 @@
+// Copyright Epic Games, Inc. All Rights Reserved.
+
+#pragma once
+
+#include "CoreMinimal.h"
+#include "HAL/Platform.h"
+#include "UObject/NoExportTypes.h"
+#include "BaseTools/BaseBrushTool.h"
+#include "BaseTools/MeshSurfacePointMeshEditingTool.h"
+#include "Components/DynamicMeshComponent.h"
+#include "PropertySets/PolygroupLayersProperties.h"
+#include "PropertySets/ColorChannelFilterPropertyType.h"
+#include "Mechanics/PolyLassoMarqueeMechanic.h"
+#include "TargetInterfaces/MeshTargetInterfaceTypes.h"
+
+#include "Sculpting/MeshSculptToolBase.h"
+#include "Sculpting/MeshBrushOpBase.h"
+
+#include "DynamicMesh/DynamicMeshAABBTree3.h"
+#include "DynamicMesh/DynamicMeshOctree3.h"
+#include "DynamicMesh/MeshNormals.h"
+#include "TransformTypes.h"
+#include "Changes/IndexedAttributeChange.h"
+#include "Polygroups/PolygroupSet.h"
+#include "FaceGroupUtil.h"
+
+#include "MeshVertexPaintTool.generated.h"
+
+class UMeshElementsVisualizer;
+class UVertexColorPaintBrushOpProps;
+class UVertexColorSmoothBrushOpProps;
+struct FMeshDescription;
+
+
+/**
+ * Tool Builder
+ */
+UCLASS()
+class MESHMODELINGTOOLSEXP_API UMeshVertexPaintToolBuilder : public UMeshSurfacePointMeshEditingToolBuilder
+{
+	GENERATED_BODY()
+
+public:
+	virtual UMeshSurfacePointTool* CreateNewTool(const FToolBuilderState& SceneState) const override;
+	virtual bool CanBuildTool(const FToolBuilderState& SceneState) const override;
+};
+
+
+
+
+/** Mesh Vertex Paint Primary Interactions */
+UENUM()
+enum class EMeshVertexPaintInteractionType : uint8
+{
+	/** Paint with an Interactive Brush */
+	Brush UMETA(DisplayName = "Paint"),
+	/** Fill any triangles connected to the brushed triangles */
+	Fill,
+	/** Fill any polygroups connected to the brushed triangles */
+	GroupFill,
+	/** Paint any triangles inside polygonal or freehand Lassos drawn in the viewport */
+	PolyLasso,
+
+	LastValue UMETA(Hidden)
+};
+
+
+UENUM()
+enum class EMeshVertexPaintColorChannel : uint8
+{
+	Red = 0,
+	Green = 1,
+	Blue = 2,
+	Alpha = 3
+};
+
+
+// currently in-sync with EVertexColorPaintBrushOpBlendMode
+UENUM()
+enum class EMeshVertexPaintColorBlendMode : uint8
+{
+	/** Interpolate between Paint color and existing Color */
+	Lerp = 0,
+	/** Alpha-Blend the Paint accumulated during each stroke with the existing Colors */
+	Mix = 1,
+	/** Multiply the Paint color with the existing Color */
+	Multiply = 2
+};
+
+
+
+
+/** Mesh Vertex Painting Brush Types */
+UENUM()
+enum class EMeshVertexPaintBrushType : uint8
+{
+	/** Paint the Primary Color */
+	Paint UMETA(DisplayName = "Paint"),
+
+	/** Paint the Erase/Secondary Color */
+	Erase UMETA(DisplayName = "Erase"),
+
+	/** Average any seam colors at a vertex */
+	Soften UMETA(DisplayName = "Soften"),
+
+	/** Smooth the colors */
+	Smooth UMETA(DisplayName = "Smooth"),
+
+	LastValue UMETA(Hidden)
+};
+
+/** Secondary/Erase Vertex Color Painting Types */
+UENUM()
+enum class EMeshVertexPaintSecondaryActionType : uint8
+{
+	/** Paint the Erase/Secondary Color */
+	Erase, 
+	/** Blend any split color values at painted vertices  */
+	Soften,
+	/** Blend vertex colors with nearby vertex colors (ie blur) */
+	Smooth
+};
+
+
+
+/** Brush Area Types */
+UENUM()
+enum class EMeshVertexPaintBrushAreaType : uint8
+{
+	/** Brush affects any triangles inside a sphere around the cursor */
+	Connected,
+	/** Brush affects any triangles geometrically connected to the triangle under the cursor */
+	Volumetric
+};
+
+/** Visibility Types */
+UENUM()
+enum class EMeshVertexPaintVisibilityType : uint8
+{
+	None,
+	/** Only paint vertices that are front-facing relative to the current camera direction */
+	FrontFacing,
+	/** Only paint triangles that are visible. Only considers active mesh, visibility test is based on triangle centers */
+	Unoccluded
+};
+
+/** Visualization Materials */
+UENUM()
+enum class EMeshVertexPaintMaterialMode : uint8
+{
+	/** Display Vertex Colors using a Lit flat-shaded material */
+	LitVertexColor,
+	/** Display Vertex Colors using an Unlit smooth-shaded material */
+	UnlitVertexColor,
+	/** Display Materials assigned to target Mesh */
+	OriginalMaterial
+};
+
+
+
+UCLASS()
+class MESHMODELINGTOOLSEXP_API UVertexPaintBasicProperties : public UInteractiveToolPropertySet
+{
+	GENERATED_BODY()
+
+public:
+	/** Primary Brush Mode */
+	//UPROPERTY(EditAnywhere, Category = Brush2, meta = (DisplayName = "Brush Type"))
+	UPROPERTY()
+	EMeshVertexPaintBrushType PrimaryBrushType = EMeshVertexPaintBrushType::Paint;
+
+	/** Painting Operation to apply when left-clicking and dragging */
+	UPROPERTY(EditAnywhere, Category = ActionType, meta = (DisplayName = "Action"))
+	EMeshVertexPaintInteractionType SubToolType = EMeshVertexPaintInteractionType::Brush;
+
+	/** The Color that will be assigned to painted triangle vertices */
+	UPROPERTY(EditAnywhere, Category = ActionType, meta = ())
+	FLinearColor PaintColor = FLinearColor::Red;
+
+	/** Controls how painted Colors will be combined with the existing Colors */
+	UPROPERTY(EditAnywhere, Category = ActionType)
+	EMeshVertexPaintColorBlendMode BlendMode = EMeshVertexPaintColorBlendMode::Lerp;
+
+	/** The Brush Operation that will be applied when holding the Shift key when in Painting */
+	UPROPERTY(EditAnywhere, Category = ActionType, meta = (DisplayName = "Secondary Brush"))
+	EMeshVertexPaintSecondaryActionType SecondaryActionType = EMeshVertexPaintSecondaryActionType::Erase;
+
+	/** Color to set when using Erase brush */
+	UPROPERTY(EditAnywhere, Category = ActionType, meta = (EditConditionHides, EditCondition = "SecondaryActionType != EMeshVertexPaintSecondaryActionType::Smooth"))
+	FLinearColor EraseColor = FLinearColor::White;
+
+	/** Strength of Smooth Brush */
+	UPROPERTY(EditAnywhere, Category = ActionType, meta = (UIMin = 0, UIMax = 1, EditConditionHides, EditCondition = "SecondaryActionType == EMeshVertexPaintSecondaryActionType::Smooth"))
+	float SmoothStrength = 0.25;
+
+	/** Controls which Color Channels will be affected by Operations. Only enabled Channels are rendered. */
+	UPROPERTY(EditAnywhere, Category = ActionType)
+	FModelingToolsColorChannelFilter ChannelFilter;
+};
+
+
+
+
+UCLASS()
+class MESHMODELINGTOOLSEXP_API UVertexPaintBrushFilterProperties : public UInteractiveToolPropertySet
+{
+	GENERATED_BODY()
+
+public:
+	/** Area Mode specifies the shape of the brush and which triangles will be included relative to the cursor */
+	UPROPERTY(EditAnywhere, Category = Filters, meta = (DisplayName = "Brush Area Mode",
+		EditConditionHides, EditCondition = "CurrentSubToolType != EMeshVertexPaintInteractionType::PolyLasso"))
+		EMeshVertexPaintBrushAreaType BrushAreaMode = EMeshVertexPaintBrushAreaType::Connected;
+
+	/** The Region affected by the current operation will be bounded by edge angles larger than this threshold */
+	UPROPERTY(EditAnywhere, Category = Filters, meta = (UIMin = "0.0", UIMax = "180.0", EditCondition = "CurrentSubToolType != EMeshVertexPaintInteractionType::PolyLasso && BrushAreaMode == EMeshVertexPaintBrushAreaType::Connected"))
+	float AngleThreshold = 180.0f;
+
+	/** The Region affected by the current operation will be bounded by UV borders/seams */
+	UPROPERTY(EditAnywhere, Category = Filters, meta = (DisplayName = "UV Seams", EditCondition = "CurrentSubToolType != EMeshVertexPaintInteractionType::PolyLasso && BrushAreaMode == EMeshVertexPaintBrushAreaType::Connected"))
+	bool bUVSeams = false;
+
+	/** The Region affected by the current operation will be bounded by Hard Normal edges/seams */
+	UPROPERTY(EditAnywhere, Category = Filters, meta = (DisplayName = "Hard Normals", EditCondition = "CurrentSubToolType != EMeshVertexPaintInteractionType::PolyLasso && BrushAreaMode == EMeshVertexPaintBrushAreaType::Connected"))
+	bool bNormalSeams = false;
+
+	/** Create Split Colors / Hard Color Edges at the borders of the painted area. Use Soften operations to un-split. */
+	UPROPERTY(EditAnywhere, Category = Filters, meta = (DisplayName = "Hard Edges"))
+	bool bHardEdges = false;
+
+	/** Control which triangles can be affected by the current operation based on visibility. Applied after all other filters. */
+	UPROPERTY(EditAnywhere, Category = Filters)
+	EMeshVertexPaintVisibilityType VisibilityFilter = EMeshVertexPaintVisibilityType::None;
+
+	/** Number of vertices in a triangle the Lasso must hit to be counted as "inside" */
+	//UPROPERTY(EditAnywhere, Category = Filters, AdvancedDisplay, meta = (UIMin = 1, UIMax = 3, EditCondition = "CurrentSubToolType == EMeshVertexPaintInteractionType::PolyLasso"))
+	UPROPERTY()
+	int MinTriVertCount = 1;		// todo: convert to cvar
+
+	/** Specify which Materials should be used to render the Mesh */
+	UPROPERTY(EditAnywhere, Category = Visualization)
+	EMeshVertexPaintMaterialMode MaterialMode = EMeshVertexPaintMaterialMode::LitVertexColor;
+
+	/** Display the Color under the cursor */
+	UPROPERTY(EditAnywhere, Category = Visualization, meta=(DisplayName="Show Color Under Cursor"))
+	bool bShowHitColor = false;
+
+
+	// values below are for edit conditions and track the current BasicProperties setting
+
+	UPROPERTY(meta=(TransientToolProperty))
+	EMeshVertexPaintInteractionType CurrentSubToolType;
+};
+
+
+
+UENUM()
+enum class EMeshVertexPaintToolActions
+{
+	NoAction,
+
+	/** Fill all Vertex Colors with the current Paint Color */
+	PaintAll,
+	/** Fill all Vertex Colors with the current Erase Color */
+	EraseAll,
+	/** Fill all Vertex Colors with Black (0,0,0,1) */
+	FillBlack,
+	/** Fill all Vertex Colors with White (1,1,1,1) */
+	FillWhite,
+
+	ApplyCurrentUtility
+};
+
+
+UCLASS()
+class MESHMODELINGTOOLSEXP_API UMeshVertexPaintToolActionPropertySet : public UInteractiveToolPropertySet
+{
+	GENERATED_BODY()
+
+public:
+	TWeakObjectPtr<UMeshVertexPaintTool> ParentTool;
+
+	void Initialize(UMeshVertexPaintTool* ParentToolIn) { ParentTool = ParentToolIn; }
+
+	void PostAction(EMeshVertexPaintToolActions Action);
+};
+
+
+
+UCLASS()
+class MESHMODELINGTOOLSEXP_API UMeshVertexPaintToolQuickActions : public UMeshVertexPaintToolActionPropertySet
+{
+	GENERATED_BODY()
+
+public:
+
+	/**
+	 * Fill all Vertex Colors with the current Paint color. Current Channel Filter still applies.
+	 */
+	UFUNCTION(CallInEditor, Category = Operations, meta = (DisplayPriority = 12))
+	void PaintAll()
+	{
+		PostAction(EMeshVertexPaintToolActions::PaintAll);
+	}
+
+	/**
+	 * Fill all Vertex Colors with the current Erase color. Current Channel Filter still applies.
+	 */
+	UFUNCTION(CallInEditor, Category = Operations, meta = (DisplayPriority = 13))
+	void EraseAll()
+	{
+		PostAction(EMeshVertexPaintToolActions::EraseAll);
+	}
+
+	/**
+	 * Fill all Vertex Colors with the Color (0,0,0,1). Current Channel Filter still applies.
+	 */
+	UFUNCTION(CallInEditor, Category = Operations, meta = (DisplayPriority = 14))
+	void FillBlack()
+	{
+		PostAction(EMeshVertexPaintToolActions::FillBlack);
+	}
+
+	/**
+	 * Fill all Vertex Colors with the Color (1,1,1,1). Current Channel Filter still applies.
+	 */
+	UFUNCTION(CallInEditor, Category = Operations, meta = (DisplayPriority = 14))
+	void FillWhite()
+	{
+		PostAction(EMeshVertexPaintToolActions::FillWhite);
+	}
+};
+
+
+
+UENUM()
+enum class EMeshVertexPaintToolUtilityOperations
+{
+	/**
+	 * Average the current color values at each vertex with split colors, so that there are no split vertices or seams in the color values
+	 */
+	BlendAllSeams,
+	/**
+	 * Set selected channels to a fixed value
+	 */
+	FillChannels,
+	/**
+	 * Invert channel values
+	 */
+	InvertChannels,
+	/**
+	 * Copy the color value from a source channel to all the selected target channels
+	 */
+	CopyChannelToChannel,
+	/**
+	 * Swap values between two Channels
+	 */
+	 SwapChannels,
+	 /**
+	 * Copy values from WeightMap into Vertex Color channels
+	 */
+	CopyFromWeightMap,
+	/**
+	 * Copy current values to any other LODs defined on the target
+	 */
+	CopyToOtherLODs,
+	/**
+	 * Copy current values to a specific LOD defined on the target
+	 */
+	CopyToSingleLOD
+};
+
+
+
+
+UCLASS()
+class MESHMODELINGTOOLSEXP_API UMeshVertexPaintToolUtilityActions : public UMeshVertexPaintToolActionPropertySet
+{
+	GENERATED_BODY()
+
+public:
+	/**
+	 * Operation to apply to current Vertex Colors
+	 */
+	UPROPERTY(EditAnywhere, Category = Utility, meta = (NoResetToDefault, DisplayPriority = 1))
+	EMeshVertexPaintToolUtilityOperations Operation = EMeshVertexPaintToolUtilityOperations::BlendAllSeams;
+
+	UPROPERTY(EditAnywhere, Category = Utility, meta = (EditConditionHides, EditCondition = "Operation == EMeshVertexPaintToolUtilityOperations::CopyChannelToChannel || Operation == EMeshVertexPaintToolUtilityOperations::SwapChannels"))
+	EMeshVertexPaintColorChannel SourceChannel = EMeshVertexPaintColorChannel::Red;
+
+	UPROPERTY(EditAnywhere, Category = Utility, meta = (UIMin = 0, UIMax = 1, EditConditionHides, EditCondition = "Operation == EMeshVertexPaintToolUtilityOperations::FillChannels"))
+	float SourceValue = 0.0f;
+
+	/** Target Vertex Weight Map */
+	UPROPERTY(EditAnywhere, Category = Utility, meta = (TransientToolProperty, NoResetToDefault, GetOptions = GetWeightMapsFunc, EditConditionHides, EditCondition = "Operation == EMeshVertexPaintToolUtilityOperations::CopyFromWeightMap"))
+	FName WeightMap;
+
+	// this function is called provide set of available weight maps
+	UFUNCTION()
+	TArray<FString> GetWeightMapsFunc() { return WeightMapsList; }
+
+	// internal list used to implement above
+	UPROPERTY(meta = (TransientToolProperty))
+	TArray<FString> WeightMapsList;
+
+	UPROPERTY(EditAnywhere, Category = Utility, meta = (EditConditionHides, EditCondition = "Operation == EMeshVertexPaintToolUtilityOperations::FillChannels || Operation == EMeshVertexPaintToolUtilityOperations::InvertChannels || Operation == EMeshVertexPaintToolUtilityOperations::CopyChannelToChannel || Operation == EMeshVertexPaintToolUtilityOperations::CopyFromWeightMap"))
+	FModelingToolsColorChannelFilter TargetChannels;
+
+	UPROPERTY(EditAnywhere, Category = Utility, meta = (EditConditionHides, EditCondition = "Operation == EMeshVertexPaintToolUtilityOperations::SwapChannels"))
+	EMeshVertexPaintColorChannel TargetChannel = EMeshVertexPaintColorChannel::Green;
+
+	/** Copy colors to HiRes Source Mesh, if it exists */
+	UPROPERTY(EditAnywhere, Category = Utility, meta = (EditConditionHides, EditCondition = "Operation == EMeshVertexPaintToolUtilityOperations::CopyToOtherLODs"))
+	bool bCopyToHiRes = false;
+
+	/** Target LOD to copy Colors to */
+	UPROPERTY(EditAnywhere, Category = Utility, meta = (TransientToolProperty, DisplayName = "Copy To LOD", NoResetToDefault, GetOptions = GetLODNamesFunc, EditConditionHides, EditCondition = "Operation == EMeshVertexPaintToolUtilityOperations::CopyToSingleLOD"))
+	FString CopyToLODName;
+
+	UFUNCTION()
+	const TArray<FString>& GetLODNamesFunc() const { return LODNamesList; }
+
+	UPROPERTY(meta = (TransientToolProperty))
+	TArray<FString> LODNamesList;
+
+
+	/**
+	 * Apply the Operation currently selected below
+	 */
+	UFUNCTION(CallInEditor, Category = Utility, meta = (DisplayPriority = 10))
+	void ApplySelectedOperation()
+	{
+		PostAction(EMeshVertexPaintToolActions::ApplyCurrentUtility);
+	}
+
+};
+
+
+
+/**
+ * FCommandChange for Vertex Color changes
+ */
+class MESHMODELINGTOOLSEXP_API FMeshVertexColorPaintChange : public TCustomIndexedValuesChange<FVector4f, int32>
+{
+public:
+	virtual FString ToString() const override
+	{
+		return FString(TEXT("Paint Vertices"));
+	}
+};
+
+
+
+
+/**
+ * Mesh Vertex Color Painting TOol
+ */
+UCLASS()
+class MESHMODELINGTOOLSEXP_API UMeshVertexPaintTool : public UMeshSculptToolBase
+{
+	GENERATED_BODY()
+
+public:
+	virtual void RegisterActions(FInteractiveToolActionSet& ActionSet) override;
+
+	virtual void Setup() override;
+	virtual void Shutdown(EToolShutdownType ShutdownType) override;
+
+	virtual void OnTick(float DeltaTime) override;
+	virtual void DrawHUD(FCanvas* Canvas, IToolsContextRenderAPI* RenderAPI) override;
+
+	virtual bool HasCancel() const override { return true; }
+	virtual bool HasAccept() const override { return true; }
+	virtual bool CanAccept() const override { return true; }
+
+	virtual bool OnUpdateHover(const FInputDeviceRay& DevicePos) override;
+
+	virtual void OnPropertyModified(UObject* PropertySet, FProperty* Property) override;
+
+	bool IsInBrushSubMode() const;
+
+	virtual void CommitResult(UBaseDynamicMeshComponent* Component, bool bModifiedTopology) override;
+
+
+public:
+
+	UPROPERTY()
+	TObjectPtr<UPolygroupLayersProperties> PolygroupLayerProperties;
+
+	UPROPERTY()
+	TObjectPtr<UVertexPaintBasicProperties> BasicProperties;
+	
+	/** Filters on paint brush */
+	UPROPERTY()
+	TObjectPtr<UVertexPaintBrushFilterProperties> FilterProperties;
+
+
+private:
+	// This will be of type UVertexPaintBrushOpProps, we keep a ref so we can change active color on pick
+	UPROPERTY()
+	TObjectPtr<UVertexColorPaintBrushOpProps> PaintBrushOpProperties;
+
+	// This will be of type UVertexPaintBrushOpProps, we keep a ref so we can change active color on pick
+	UPROPERTY()
+	TObjectPtr<UVertexColorPaintBrushOpProps> EraseBrushOpProperties;
+
+public:
+	void FloodFillColorAction(FLinearColor Color);
+
+
+	
+	void SetTrianglesToVertexColor(const TArray<int32>& Triangles, const FLinearColor& ToColor);
+	void SetTrianglesToVertexColor(const TSet<int32>& Triangles, const FLinearColor& ToColor);
+
+	bool HaveVisibilityFilter() const;
+	void ApplyVisibilityFilter(const TArray<int32>& Triangles, TArray<int32>& VisibleTriangles);
+	void ApplyVisibilityFilter(TSet<int32>& Triangles, TArray<int32>& ROIBuffer, TArray<int32>& OutputBuffer);
+
+protected:
+	// UMeshSculptToolBase API
+	virtual UBaseDynamicMeshComponent* GetSculptMeshComponent() { return DynamicMeshComponent; }
+	virtual FDynamicMesh3* GetBaseMesh() { check(false); return nullptr; }
+	virtual const FDynamicMesh3* GetBaseMesh() const { check(false); return nullptr; }
+
+	virtual int32 FindHitSculptMeshTriangle(const FRay3d& LocalRay) override;
+	virtual int32 FindHitTargetMeshTriangle(const FRay3d& LocalRay) override;
+
+	virtual void OnBeginStroke(const FRay& WorldRay) override;
+	virtual void OnEndStroke() override;
+
+	virtual TUniquePtr<FMeshSculptBrushOp>& GetActiveBrushOp();
+	// end UMeshSculptToolBase API
+
+
+
+	//
+	// Action support
+	//
+
+public:
+	virtual void RequestAction(EMeshVertexPaintToolActions ActionType);
+
+	UPROPERTY()
+	TObjectPtr<UMeshVertexPaintToolQuickActions> QuickActions;
+
+
+public:
+	UPROPERTY()
+	TObjectPtr<UMeshVertexPaintToolUtilityActions> UtilityActions;
+	
+	void ApplyCurrentUtilityAction();
+	void BlendAllSeams();
+	void FillChannels();
+	void InvertChannels();
+	void CopyChannelToChannel();
+	void SwapChannels();
+	void CopyFromWeightMap();
+	void CopyToOtherLODs();
+	void CopyToSpecificLOD();
+
+protected:
+	bool bHavePendingAction = false;
+	EMeshVertexPaintToolActions PendingAction;
+	virtual void ApplyAction(EMeshVertexPaintToolActions ActionType);
+
+
+
+	//
+	// Marquee Support
+	//
+public:
+	UPROPERTY()
+	TObjectPtr<UPolyLassoMarqueeMechanic> PolyLassoMechanic;
+
+protected:
+	void OnPolyLassoFinished(const FCameraPolyLasso& Lasso, bool bCanceled);
+
+
+	//
+	// Internals
+	//
+
+protected:
+
+	UPROPERTY()
+	TObjectPtr<AInternalToolFrameworkActor> PreviewMeshActor = nullptr;
+
+	UPROPERTY()
+	TObjectPtr<UDynamicMeshComponent> DynamicMeshComponent;
+
+	UPROPERTY()
+	TObjectPtr<UMeshElementsVisualizer> MeshElementsDisplay;
+
+	// realtime visualization
+	void OnDynamicMeshComponentChanged();
+	FDelegateHandle OnDynamicMeshComponentChangedHandle;
+
+	EMeshLODIdentifier SourceLOD;
+	TArray<EMeshLODIdentifier> AvailableLODs;
+	bool bTargetSupportsLODs = false;
+
+	UE::Geometry::FDynamicMeshColorOverlay* ActiveColorOverlay = nullptr;
+	UE::Geometry::FDynamicMeshColorOverlay* GetActiveColorOverlay() const { return ActiveColorOverlay; }
+
+	TUniquePtr<UE::Geometry::FPolygroupSet> ActiveGroupSet;
+	void OnSelectedGroupLayerChanged();
+	void UpdateActiveGroupLayer();
+
+	void UpdateSubToolType(EMeshVertexPaintInteractionType NewType);
+	void UpdateBrushType(EMeshVertexPaintBrushType BrushType);
+	void UpdateSecondaryBrushType(EMeshVertexPaintSecondaryActionType NewType);
+	void UpdateVertexPaintMaterialMode();
+
+	TSet<int32> AccumulatedTriangleROI;
+	bool bUndoUpdatePending = false;
+	TArray<int> NormalsBuffer;
+	void WaitForPendingUndoRedo();
+
+	TArray<int> TempROIBuffer;
+	TArray<int> VertexROI;
+	TArray<bool> VisibilityFilterBuffer;
+	TSet<int> TempVertexSet;
+	TSet<int> TriangleROI;
+	void UpdateROI(const FSculptBrushStamp& CurrentStamp);
+
+	// in Mix blending mode we need to accumulate each stroke in a fully separate buffer and blend it
+	// with the background colors. So we need that buffer and also save initial colors
+	TArray<FVector4f> StrokeInitialColorBuffer;
+	TArray<FVector4f> StrokeAccumColorBuffer;
+
+	EMeshVertexPaintBrushType PendingStampType = EMeshVertexPaintBrushType::Paint;
+
+	bool UpdateStampPosition(const FRay& WorldRay);
+	bool ApplyStamp();
+
+	UE::Geometry::FDynamicMeshOctree3 Octree;
+
+	bool UpdateBrushPosition(const FRay& WorldRay);
+
+	bool GetInEraseStroke()
+	{
+		// Re-use the smoothing stroke key (shift) for erase stroke in the group paint tool
+		return GetInSmoothingStroke();
+	}
+
+
+	bool bPendingPickColor = false;
+	bool bPendingPickEraseColor = false;
+
+
+	TArray<int32> ROITriangleBuffer;
+	TSet<int32> ROIElementSet;
+	TArray<int32> ROIElementBuffer;
+	TArray<FVector4f> ROIColorBuffer;
+	void InitializeElementROIFromTriangleROI(const TArray<int32>& TriangleROI, bool bInitializeFlatBuffers);
+	bool SyncMeshWithColorBuffer(FDynamicMesh3* Mesh);
+
+	TUniquePtr<TIndexedValuesChangeBuilder<FVector4f, FMeshVertexColorPaintChange>> ActiveChangeBuilder;
+	void BeginChange();
+	void EndChange();
+	void ExternalUpdateValues(const TArray<int32>& ElementIDs, const TArray<FVector4f>& NewValues);
+
+	FColor GetColorForGroup(int32 GroupID);
+	void ApplyChannelFilter(const FVector4f& CurColor, FVector4f& NewColor);
+	void OnChannelFilterModified();
+
+	TArray<FVector3d> TriNormals;
+	TArray<int32> UVSeamEdges;
+	TArray<int32> NormalSeamEdges;
+	void PrecomputeFilterData();
+
+
+protected:
+	virtual bool ShowWorkPlane() const override { return false; }
+
+	// Currently using flow rate as 'brush strength', so disable temporal stamp spacing
+	virtual float GetStampTemporalFlowRate() const override { return 1.0f; }
+};
+
+
+
diff --git a/Engine/Plugins/Experimental/MeshModelingToolsetExp/Source/MeshModelingToolsExp/Public/Sculpting/MeshBrushOpBase.h b/Engine/Plugins/Experimental/MeshModelingToolsetExp/Source/MeshModelingToolsExp/Public/Sculpting/MeshBrushOpBase.h
--- a/Engine/Plugins/Experimental/MeshModelingToolsetExp/Source/MeshModelingToolsExp/Public/Sculpting/MeshBrushOpBase.h
+++ b/Engine/Plugins/Experimental/MeshModelingToolsetExp/Source/MeshModelingToolsExp/Public/Sculpting/MeshBrushOpBase.h
@@ -92,6 +92,7 @@ public:
 
 	// support for this is optional, used by UI level to edit brush props via hotkeys/etc
 	virtual void SetStrength(float NewStrength) { }
+	virtual void SetFalloff(float NewFalloff) { }
 };
 
 
diff --git a/Engine/Plugins/Experimental/MeshModelingToolsetExp/Source/MeshModelingToolsExp/Public/Sculpting/MeshVertexPaintBrushOps.h b/Engine/Plugins/Experimental/MeshModelingToolsetExp/Source/MeshModelingToolsExp/Public/Sculpting/MeshVertexPaintBrushOps.h
new file mode 100644
--- /dev/null
+++ b/Engine/Plugins/Experimental/MeshModelingToolsetExp/Source/MeshModelingToolsExp/Public/Sculpting/MeshVertexPaintBrushOps.h
@@ -0,0 +1,429 @@
+// Copyright Epic Games, Inc. All Rights Reserved.
+
+#pragma once
+
+#include "Sculpting/MeshBrushOpBase.h"
+#include "DynamicMesh/DynamicMesh3.h"
+#include "DynamicMesh/DynamicMeshAttributeSet.h"
+#include "MeshWeights.h"
+#include "MathUtil.h"
+#include "VectorTypes.h"
+#include "Async/ParallelFor.h"
+#include "MeshVertexPaintBrushOps.generated.h"
+
+
+UENUM()
+enum class EVertexColorPaintBrushOpBlendMode
+{
+	Lerp = 0,
+	Mix = 1,
+	Multiply = 2
+};
+
+
+
+namespace UELocal
+{
+	template <typename RealType>
+	inline UE::Math::TVector4<RealType> BlendColors_Mix(
+		const UE::Math::TVector4<RealType>& BackColor,
+		const UE::Math::TVector4<RealType>& ForeColor,
+		RealType Alpha)
+	{
+		UE::Math::TVector4<RealType> Result;
+		RealType OneMinusAlpha = (RealType)1 - Alpha;
+
+		//for (int32 k = 0; k < 3; ++k)
+		//{
+		//	Result[k] = TMathUtil<RealType>::Sqrt( OneMinusAlpha*(A[k]*A[k])  + Alpha*(B[k]*B[k]) );
+		//	Result[k] = TMathUtil<RealType>::Clamp(Result[k], (RealType)0, (RealType)1);
+		//}
+		//Result.W = TMathUtil<RealType>::Clamp(OneMinusAlpha*A.W + Alpha*B.W, (RealType)0, (RealType)1);
+
+		// pretty sure this is assuming linear color space...
+
+		RealType ForeAlpha = ForeColor.W * Alpha;
+		Result.W = ForeAlpha + BackColor.W * (1 - ForeAlpha);
+		for (int32 k = 0; k < 3; ++k)
+		{
+			Result[k] = ForeColor[k]*ForeAlpha + BackColor[k]*BackColor.W*(1-ForeAlpha);
+			if (Result.W > TMathUtil<RealType>::ZeroTolerance)
+			{
+				Result[k] /= Result.W;
+			}
+		}
+
+		Result.X = FMathf::Clamp(Result.X, (RealType)0, (RealType)1);
+		Result.Y = FMathf::Clamp(Result.Y, (RealType)0, (RealType)1);
+		Result.Z = FMathf::Clamp(Result.Z, (RealType)0, (RealType)1);
+		Result.W = FMathf::Clamp(Result.W, (RealType)0, (RealType)1);
+
+		return Result;
+	}
+
+
+	template <typename RealType>
+	inline UE::Math::TVector4<RealType> BlendColors_Lerp(
+		const UE::Math::TVector4<RealType>& A,
+		const UE::Math::TVector4<RealType>& B,
+		RealType Alpha)
+	{
+		UE::Math::TVector4<RealType> Result;
+		RealType OneMinusAlpha = (RealType)1 - Alpha;
+		for (int32 k = 0; k < 4; ++k)
+		{
+			Result[k] = OneMinusAlpha * A[k] + Alpha * B[k];
+			Result[k] = TMathUtil<RealType>::Clamp(OneMinusAlpha*A[k] + Alpha*B[k], (RealType)0, (RealType)1);
+		}
+		return Result;
+	}
+
+	template <typename RealType>
+	inline UE::Math::TVector4<RealType> BlendColors_Multiply(
+		const UE::Math::TVector4<RealType>& A,
+		const UE::Math::TVector4<RealType>& B,
+		RealType Alpha)
+	{
+		UE::Math::TVector4<RealType> Result;
+		Result.X = A.X * B.X;
+		Result.Y = A.Y * B.Y;
+		Result.Z = A.Z * B.Z;
+		Result.W = A.W * B.W;
+		return BlendColors_Lerp(A, Result, Alpha);
+	}
+
+
+	template <typename RealType>
+	inline UE::Math::TVector4<RealType> BlendColors(
+		const UE::Math::TVector4<RealType>& Background,
+		const UE::Math::TVector4<RealType>& Foreground,
+		RealType Alpha,
+		EVertexColorPaintBrushOpBlendMode BlendMode)
+	{
+		switch (BlendMode)
+		{
+		case EVertexColorPaintBrushOpBlendMode::Multiply:
+			return BlendColors_Multiply(Background, Foreground, Alpha);
+			break;
+		case EVertexColorPaintBrushOpBlendMode::Mix:
+			return BlendColors_Mix(Background, Foreground, Alpha);
+			break;
+		case EVertexColorPaintBrushOpBlendMode::Lerp:
+		default:
+			return BlendColors_Lerp(Background, Foreground, Alpha);
+		}
+	}
+}
+
+
+
+UCLASS()
+class MESHMODELINGTOOLSEXP_API UVertexColorBaseBrushOpProps : public UMeshSculptBrushOpProps
+{
+	GENERATED_BODY()
+public:
+
+	/** Strength of the Brush */
+	UPROPERTY(EditAnywhere, Category = PaintBrush, meta = (DisplayName = "Strength", UIMin = "0.0", UIMax = "1.0", ClampMin = "0.0", ClampMax = "1.0"))
+	float Strength = 0.5;
+
+	/** Amount of falloff to apply */
+	UPROPERTY(EditAnywhere, Category = PaintBrush, meta = (DisplayName = "Falloff", UIMin = "0.0", UIMax = "1.0", ClampMin = "0.0", ClampMax = "1.0"))
+	float Falloff = 1.0;
+
+	UPROPERTY(EditAnywhere, Category = PaintBrush)
+	EVertexColorPaintBrushOpBlendMode BlendMode = EVertexColorPaintBrushOpBlendMode::Mix;
+
+	virtual float GetStrength() override { return Strength; }
+	virtual void SetStrength(float NewStrength) override { Strength = FMathf::Clamp(NewStrength, 0.0f, 1.0f); }
+	virtual float GetFalloff() override { return Falloff; }
+	virtual void SetFalloff(float NewFalloff) override { Falloff = FMathf::Clamp(NewFalloff, 0.0f, 1.0f); }
+};
+
+
+
+
+
+UCLASS()
+class MESHMODELINGTOOLSEXP_API UVertexColorPaintBrushOpProps : public UVertexColorBaseBrushOpProps
+{
+	GENERATED_BODY()
+public:
+	/** The paint color */
+	UPROPERTY(EditAnywhere, Category = Paint, meta = (DisplayName = "Color"))
+	FLinearColor Color;
+
+	virtual FLinearColor GetColor() const { return Color; }
+};
+
+
+UCLASS()
+class MESHMODELINGTOOLSEXP_API UVertexColorSoftenBrushOpProps : public UVertexColorBaseBrushOpProps
+{
+	GENERATED_BODY()
+public:
+
+	virtual float GetFalloff() override { return 1.0f; }
+};
+
+UCLASS()
+class MESHMODELINGTOOLSEXP_API UVertexColorSmoothBrushOpProps : public UVertexColorBaseBrushOpProps
+{
+	GENERATED_BODY()
+public:
+};
+
+
+
+namespace UE::Geometry
+{
+
+
+class MESHMODELINGTOOLSEXP_API FMeshVertexColorBrushOp : public FMeshSculptBrushOp
+{
+public:
+
+	// not supported for this kind of brush op
+	virtual void ApplyStamp(
+		const FDynamicMesh3* Mesh,
+		const FSculptBrushStamp& Stamp,
+		const TArray<int32>& Vertices,
+		TArray<FVector3d>& NewColorsOut) override
+	{
+		check(false);
+	}
+
+
+	virtual void ApplyStampToVertexColors(
+		const FDynamicMesh3* Mesh,
+		const FDynamicMeshColorOverlay* VertexColors,
+		TArray<FVector4f>& InitialVertexColorBuffer,
+		TArray<FVector4f>& StrokeVertexColorBuffer,
+		const FSculptBrushStamp& Stamp,
+		const TArray<int32>& ElementIDs,
+		TArray<FVector4f>& NewColorsOut) = 0;
+
+};
+
+
+
+
+class FVertexPaintBrushOp : public FMeshVertexColorBrushOp
+{
+public:
+	double BrushSpeedTuning = 6.0;
+
+	virtual void ApplyStampToVertexColors(
+		const FDynamicMesh3* Mesh,
+		const FDynamicMeshColorOverlay* VertexColors,
+		TArray<FVector4f>& InitialVertexColorBuffer,
+		TArray<FVector4f>& StrokeVertexColorBuffer,
+		const FSculptBrushStamp& Stamp,
+		const TArray<int32>& ElementIDs,
+		TArray<FVector4f>& NewColorsOut) override
+	{
+		const FVector3d& StampPos = Stamp.LocalFrame.Origin;
+
+		//double UsePower = Stamp.Power * Stamp.Radius * Stamp.DeltaTime * BrushSpeedTuning;
+		double UsePower = Stamp.Power;
+
+		UVertexColorPaintBrushOpProps* Props = GetPropertySetAs<UVertexColorPaintBrushOpProps>();
+		FVector4f SetColor = (FVector4f)Props->GetColor();
+		float BlendStrokeAlpha = SetColor.W; 
+		if (Props->BlendMode == EVertexColorPaintBrushOpBlendMode::Mix)
+		{
+			SetColor.W = 1.0f;
+		}
+
+		int32 NumElements = ElementIDs.Num();
+		ensure(NewColorsOut.Num() == NumElements);
+		ParallelFor(NumElements, [&](int32 k)
+		{
+			int32 ElementID = ElementIDs[k];
+			int32 VertexID = VertexColors->GetParentVertex(ElementID);
+			FVector3d VertexPos = Mesh->GetVertex(VertexID);
+			double PositionFalloff = GetFalloff().Evaluate(Stamp, VertexPos);
+			double T = PositionFalloff * UsePower;
+
+			if (Props->BlendMode == EVertexColorPaintBrushOpBlendMode::Mix)
+			{
+				FVector4f CurStrokeColor = StrokeVertexColorBuffer[ElementID];
+				StrokeVertexColorBuffer[ElementID] = UELocal::BlendColors_Mix(CurStrokeColor, SetColor, (float)T);
+				FVector4f OverColor = StrokeVertexColorBuffer[ElementID];
+
+				FVector4f BaseColor = InitialVertexColorBuffer[ElementID];
+				NewColorsOut[k] = UELocal::BlendColors(BaseColor, OverColor, BlendStrokeAlpha, Props->BlendMode);
+			}
+			else
+			{
+				FVector4f CurColor = VertexColors->GetElement(ElementIDs[k]);
+				NewColorsOut[k] = UELocal::BlendColors(CurColor, SetColor, (float)T, Props->BlendMode);
+			}
+		});
+
+	}
+};
+
+
+
+/**
+ * Soften brush averages colors at any vertices that have split elements, effectively un-splitting them
+ * Currently this averaging is done in a single step, ie falloff/etc is ignored, but it could be done more gradually...
+ */
+class FVertexColorSoftenBrushOp : public FMeshVertexColorBrushOp
+{
+public:
+	TArray<int32> TempElementIDs;
+
+	virtual void ApplyStampToVertexColors(
+		const FDynamicMesh3* Mesh,
+		const FDynamicMeshColorOverlay* VertexColors,
+		TArray<FVector4f>& InitialVertexColorBuffer,
+		TArray<FVector4f>& StrokeVertexColorBuffer,
+		const FSculptBrushStamp& Stamp,
+		const TArray<int32>& ElementIDs,
+		TArray<FVector4f>& NewColorsOut) override
+	{
+		const FVector3d& StampPos = Stamp.LocalFrame.Origin;
+
+		int32 NumElements = ElementIDs.Num();
+		ensure(NewColorsOut.Num() == NumElements);
+		for (int32 k = 0; k < NumElements; ++k)
+		{
+			int32 VertexID = VertexColors->GetParentVertex(ElementIDs[k]);
+			TempElementIDs.Reset();
+			VertexColors->GetVertexElements(VertexID, TempElementIDs);
+			if (TempElementIDs.Num() > 1)
+			{
+				FVector4f SumColor = FVector4f::Zero();
+				for (int32 ElementID : TempElementIDs)
+				{
+					SumColor += VertexColors->GetElement(ElementID);
+				}
+				SumColor /= (double)TempElementIDs.Num();
+				NewColorsOut[k] = SumColor;
+			}
+			else
+			{
+				NewColorsOut[k] = VertexColors->GetElement(ElementIDs[k]);
+			}
+		}
+
+	}
+};
+
+
+
+
+class FVertexColorSmoothBrushOp : public FMeshVertexColorBrushOp
+{
+public:
+	double BrushSpeedTuning = 6.0;
+	TArray<int32> TempElementIDs;
+	TArray<int32> ElementIDVertices;
+	TArray<int32> ElementIDVertexIndices;
+	TArray<int32> UniqueVertices;
+	TArray<FVector4f> SoftenedVertexColors;
+	TArray<FVector4f> SmoothedVertexColors;
+
+
+	virtual void ApplyStampToVertexColors(
+		const FDynamicMesh3* Mesh,
+		const FDynamicMeshColorOverlay* VertexColors,
+		TArray<FVector4f>& InitialVertexColorBuffer,
+		TArray<FVector4f>& StrokeVertexColorBuffer,
+		const FSculptBrushStamp& Stamp,
+		const TArray<int32>& ElementIDs,
+		TArray<FVector4f>& NewColorsOut) override
+	{
+		const FVector3d& StampPos = Stamp.LocalFrame.Origin;
+
+		//double UsePower = Stamp.Power * Stamp.Radius * Stamp.DeltaTime * BrushSpeedTuning;
+		double UsePower = Stamp.Power;
+
+		UVertexColorSmoothBrushOpProps* Props = GetPropertySetAs<UVertexColorSmoothBrushOpProps>();
+
+		int32 NumElements = ElementIDs.Num();
+		ElementIDVertices.SetNum(NumElements, false);
+		ElementIDVertexIndices.SetNum(NumElements, false);
+		ensure(NewColorsOut.Num() == NumElements);
+		UniqueVertices.Reset();
+		for (int32 k = 0; k < NumElements; ++k)
+		{
+			int32 VertexID = VertexColors->GetParentVertex(ElementIDs[k]);
+			ElementIDVertices[k] = VertexID;
+			int32 VertexIndex = (int32)UniqueVertices.AddUnique(VertexID);
+			ElementIDVertexIndices[k] = VertexIndex;
+		}
+
+		// compute color at each vertex that blends any split elements
+		int32 NumVertices = UniqueVertices.Num();
+		SoftenedVertexColors.SetNum(NumVertices, false);
+		ParallelFor(NumVertices, [&](int32 Index)
+		{
+			TArray<int32> Temp;
+			Temp.Reserve(16);
+			VertexColors->GetVertexElements(UniqueVertices[Index], Temp);
+			SoftenedVertexColors[Index] = FVector4f::Zero();
+			int32 Count = 0;
+			for (int32 ElementID : Temp)
+			{
+				if (ElementIDs.Contains(ElementID))
+				{
+					SoftenedVertexColors[Index] += VertexColors->GetElement(ElementID);
+					Count++;
+				}
+			}
+			if (Count > 1)
+			{
+				SoftenedVertexColors[Index] /= (double)Count;
+			}
+		});
+		
+
+		// now compute one-ring-smoothed vertex colors
+		SmoothedVertexColors.SetNum(NumVertices, false);
+		ParallelFor(NumVertices, [&](int32 Index)
+		{
+			FVector4f SmoothedColor = FVector4f::Zero();
+			int32 NbrCount = 0;
+			int32 VertexID = UniqueVertices[Index];
+			for (int32 NbrVertexID : Mesh->VtxVerticesItr(VertexID))
+			{
+				int32 NbrIndex = UniqueVertices.IndexOfByKey(NbrVertexID);
+				if (NbrIndex != INDEX_NONE)
+				{
+					const float GeometryWeight = 1.0f;
+					SmoothedColor += GeometryWeight * SoftenedVertexColors[NbrIndex];
+					NbrCount++;
+				}
+			}
+			if (NbrCount > 1)
+			{
+				SmoothedColor /= (double)NbrCount;
+			}
+			SmoothedVertexColors[Index] = SmoothedColor;
+		} );
+
+		// now blend each element color w/ it's existing smoothed color
+		ParallelFor(NumElements, [&](int32 k)
+		{
+			int32 VertexID = ElementIDVertices[k];
+
+			FVector3d VertexPos = Mesh->GetVertex(VertexID);
+			double PositionFalloff = GetFalloff().Evaluate(Stamp, VertexPos);
+			double T = PositionFalloff * UsePower;
+
+			FVector4f CurColor = VertexColors->GetElement(ElementIDs[k]);
+			int32 VertIndex = ElementIDVertexIndices[k];
+			FVector4f SetColor = SmoothedVertexColors[VertIndex];
+
+			//NewColorsOut[k] = UELocal::BlendColors_Mix(CurColor, SetColor, (float)T);
+			NewColorsOut[k] = UELocal::BlendColors_Lerp(CurColor, SetColor, (float)T);
+		});
+
+	}
+};
+
+
+}
\ No newline at end of file
diff --git a/Engine/Plugins/Runtime/MeshModelingToolset/Source/ModelingComponents/Public/PropertySets/ColorChannelFilterPropertyType.h b/Engine/Plugins/Runtime/MeshModelingToolset/Source/ModelingComponents/Public/PropertySets/ColorChannelFilterPropertyType.h
new file mode 100644
--- /dev/null
+++ b/Engine/Plugins/Runtime/MeshModelingToolset/Source/ModelingComponents/Public/PropertySets/ColorChannelFilterPropertyType.h
@@ -0,0 +1,44 @@
+// Copyright Epic Games, Inc. All Rights Reserved.
+
+#pragma once
+
+#include "CoreMinimal.h"
+#include "IndexTypes.h"
+#include "ColorChannelFilterPropertyType.generated.h"
+
+
+USTRUCT(BlueprintType)
+struct MODELINGCOMPONENTS_API FModelingToolsColorChannelFilter
+{
+	GENERATED_BODY()
+
+	/** Red Channel */
+	UPROPERTY(EditAnywhere, Category = ChannelFilters, meta = (DisplayName = "R"))
+	bool bRed = true;
+
+	/** Green Channel */
+	UPROPERTY(EditAnywhere, Category = ChannelFilters, meta = (DisplayName = "G"))
+	bool bGreen = true;
+
+	/** Blue Channel */
+	UPROPERTY(EditAnywhere, Category = ChannelFilters, meta = (DisplayName = "B"))
+	bool bBlue = true;
+
+	/** Alpha Channel */
+	UPROPERTY(EditAnywhere, Category = ChannelFilters, meta = (DisplayName = "A"))
+	bool bAlpha = false;
+
+	// helper functions
+
+	bool operator!=(const FModelingToolsColorChannelFilter& Other) const
+	{
+		return (bRed != Other.bRed) || (bGreen != Other.bGreen) || (bBlue != Other.bBlue) || (bAlpha != Other.bAlpha);
+	}
+
+	bool AnyChannelFiltered() const { return bRed == false || bGreen == false || bBlue == false || bAlpha == false; }
+
+	UE::Geometry::FIndex4i AsFlags() const
+	{
+		return UE::Geometry::FIndex4i((bRed) ? 1 : 0, (bGreen) ? 1 : 0, (bBlue) ? 1 : 0, (bAlpha) ? 1 : 0);
+	}
+};
