ModelingTools: pvs fix and minor const-correctness fix
#rnx

diff --git a/Engine/Plugins/Experimental/MeshModelingToolsetExp/Source/MeshModelingToolsExp/Private/MeshVertexPaintTool.cpp b/Engine/Plugins/Experimental/MeshModelingToolsetExp/Source/MeshModelingToolsExp/Private/MeshVertexPaintTool.cpp
--- a/Engine/Plugins/Experimental/MeshModelingToolsetExp/Source/MeshModelingToolsExp/Private/MeshVertexPaintTool.cpp
+++ b/Engine/Plugins/Experimental/MeshModelingToolsetExp/Source/MeshModelingToolsExp/Private/MeshVertexPaintTool.cpp
@@ -140,10 +140,16 @@ void UMeshVertexPaintTool::Setup()
 		if (Mesh->Attributes()->PrimaryColors() == nullptr)
 		{
 			Mesh->Attributes()->EnablePrimaryColors();
-			ActiveColorOverlay = Mesh->Attributes()->PrimaryColors();
-			ActiveColorOverlay->CreateFromPredicate(
-				[](int, int, int) { return false; }, 1.0f);
 		}
+		ActiveColorOverlay = Mesh->Attributes()->PrimaryColors();
+		if (ActiveColorOverlay == nullptr)
+		{
+			GetToolManager()->PostActiveToolShutdownRequest(this, EToolShutdownType::Cancel, true,
+				LOCTEXT("InvalidColorsMessage", "Target Mesh does not support Vertex Colors"));
+			return;
+		}
+		ActiveColorOverlay->CreateFromPredicate(
+			[](int, int, int) { return false; }, 1.0f);
 	}
 	else
 	{
@@ -156,6 +162,7 @@ void UMeshVertexPaintTool::Setup()
 			});
 	}
 
+
 	StrokeInitialColorBuffer.SetNum(ActiveColorOverlay->MaxElementID());
 	StrokeAccumColorBuffer.SetNum(ActiveColorOverlay->MaxElementID());
 
diff --git a/Engine/Source/Runtime/GeometryCore/Private/DynamicMesh/DynamicMeshOverlay.cpp b/Engine/Source/Runtime/GeometryCore/Private/DynamicMesh/DynamicMeshOverlay.cpp
--- a/Engine/Source/Runtime/GeometryCore/Private/DynamicMesh/DynamicMeshOverlay.cpp
+++ b/Engine/Source/Runtime/GeometryCore/Private/DynamicMesh/DynamicMeshOverlay.cpp
@@ -1480,7 +1480,7 @@ template<typename RealType, int ElementSize, typename VectorType>
 bool TDynamicMeshVectorOverlay<RealType, ElementSize, VectorType>::EnumerateVertexElements(
 	int VertexID,
 	TFunctionRef<bool(int TriangleID, int ElementID, const VectorType& Value)> ProcessFunc,
-	bool bFindUniqueElements)
+	bool bFindUniqueElements) const
 {
 	if (this->ParentMesh->IsVertex(VertexID) == false) return false;
 
diff --git a/Engine/Source/Runtime/GeometryCore/Public/DynamicMesh/DynamicMeshOverlay.h b/Engine/Source/Runtime/GeometryCore/Public/DynamicMesh/DynamicMeshOverlay.h
--- a/Engine/Source/Runtime/GeometryCore/Public/DynamicMesh/DynamicMeshOverlay.h
+++ b/Engine/Source/Runtime/GeometryCore/Public/DynamicMesh/DynamicMeshOverlay.h
@@ -763,7 +763,7 @@ public:
 	bool EnumerateVertexElements(
 		int VertexID,
 		TFunctionRef<bool(int TriangleID, int ElementID, const VectorType& Value)> ProcessFunc,
-		bool bFindUniqueElements = true);
+		bool bFindUniqueElements = true) const;
 };
 
 
