CombineMeshInstances: clamp min-sample-points-per-triangle in hidden removal to 8 instead of 2. This may need to be exposed as a parameter in future. Improves a lot of current failure cases for hidden removal. Minor renaming. Tweaked sample radius in various calling code.
#rnx

diff --git a/Engine/Plugins/Experimental/MeshModelingToolsetExp/Source/GeometryProcessingAdapters/Private/CombineMeshInstancesImpl.cpp b/Engine/Plugins/Experimental/MeshModelingToolsetExp/Source/GeometryProcessingAdapters/Private/CombineMeshInstancesImpl.cpp
--- a/Engine/Plugins/Experimental/MeshModelingToolsetExp/Source/GeometryProcessingAdapters/Private/CombineMeshInstancesImpl.cpp
+++ b/Engine/Plugins/Experimental/MeshModelingToolsetExp/Source/GeometryProcessingAdapters/Private/CombineMeshInstancesImpl.cpp
@@ -1424,7 +1424,6 @@ static void RemoveHiddenFaces_ExteriorVisibility(
 	const double TriScalingAlpha = 0.999;
 	const double BaryCoordsThreshold = 0.001;
 
-
 	auto FindHitTriangleTest = [&](FVector3d TargetPosition, FVector3d TargetNormal, FVector3d FarPosition) -> int
 	{
 		FVector3d RayDir(TargetPosition - FarPosition);
@@ -1455,7 +1454,12 @@ static void RemoveHiddenFaces_ExteriorVisibility(
 		CardinalDirections.Add(-Direction);
 	}
 
-	//
+	// TODO: it seems like a quite common failure case is triangles that are somewhat deeply
+	// nested inside cavities/etc. A possible way to help here would be to essentially raytrace
+	// orthographic images from top/bottom/etc. This could be done async and combined w/ the
+	// visibilty determination below...
+
+
 	// First pass. For each triangle, cast a ray at it's centroid from 
 	// outside the model, along the X/Y/Z directions and tri normal.
 	// If tri is hit we mark it as having 'known' status, allowing it
@@ -1500,24 +1504,32 @@ static void RemoveHiddenFaces_ExteriorVisibility(
 	});
 
 	//
-	// Construct set of exterior sample directions, for each triangle sample point
-	// below we will check if it is visible from any of these sample directions.
-	// Order is shuffled in hopes that for visible tris we don't waste a bunch
+	// Construct set of exterior visibility test directions, below we will check if sample 
+	// points on the mesh triangles are visible from the exterior along these directions.
+	// Order is modulo-shuffled in hopes that for visible tris we don't waste a bunch
 	// of time on the 'far' side
-	//
-	int32 NumExteriorSamplePoints = 128;
-	TSphericalFibonacci<double> SphereSampler(NumExteriorSamplePoints);
-	TArray<FVector3d> ExteriorSamplePoints;
-	FModuloIteration ModuloIter(NumExteriorSamplePoints);
-	uint32 SampleIndex = 0;
-	while (ModuloIter.GetNextIndex(SampleIndex))
-	{
-		ExteriorSamplePoints.Add( Normalized(SphereSampler[SampleIndex]) );
-	}
+	int32 NumVisibilityTestDirections = 128;
+	TSphericalFibonacci<double> SphereSampler(NumVisibilityTestDirections);
+	TArray<FVector3d> VisibilityDirections;
+	FModuloIteration ModuloIter(NumVisibilityTestDirections);
+	uint32 DirectionIndex = 0;
+	while (ModuloIter.GetNextIndex(DirectionIndex))
+	{
+		VisibilityDirections.Add( Normalized(SphereSampler[DirectionIndex]) );
+	}
+	// Fibonacci set generally does not include the cardinal directions, but they are highly useful to check
+	VisibilityDirections.Append(CardinalDirections);
+	NumVisibilityTestDirections = VisibilityDirections.Num();
+
+	// For each triangle we will generate a set of sample points on the triangle surface,
+	// and then check if that point is visible along any of the sample directions.
+	// The number of sample points allocated to a triangle is based on it's area and the SampleRadius.
+	// However for small triangles this may be < 1, so we will clamp to at least this many samples.
+	// (this value should perhaps be relative to the mesh density, or exposed as a parameter...)
+	const int32 MinTriSamplesPerSamplePoint = 8;
 
-	//
 	// For each triangle, generate a set of sample points on the triangle surface,
-	// and then check if that point is visible from any of the exterior sample points.
+	
 	// This is the expensive part!
 	ParallelFor(TargetMesh.MaxTriangleID(), [&](int32 tid)
 	{
@@ -1542,7 +1554,7 @@ static void RemoveHiddenFaces_ExteriorVisibility(
 		}
 
 		double DiscArea = (FMathd::Pi * SampleRadius * SampleRadius);
-		int NumSamples = FMath::Max( (int)(TriArea / DiscArea), 2 );  // a bit arbitrary...
+		int NumSamples = FMath::Max( (int)(TriArea / DiscArea), MinTriSamplesPerSamplePoint);
 		FVector2d V1 = UVTriangle.V[1] - UVTriangle.V[0];
 		FVector2d V2 = UVTriangle.V[2] - UVTriangle.V[0];
 
@@ -1562,15 +1574,15 @@ static void RemoveHiddenFaces_ExteriorVisibility(
 
 				// cast ray from all exterior sample locations for this triangle sample point
 				HitTris.Reset();
-				for (int32 k = 0; k < NumExteriorSamplePoints; ++k)
+				for (int32 k = 0; k < NumVisibilityTestDirections; ++k)
 				{
-					FVector3d Direction = ExteriorSamplePoints[k];
+					FVector3d Direction = VisibilityDirections[k];
 					if (FMathd::Abs(Direction.Dot(TriNormal)) < GlancingAngleDotTolerance)
 					{
 						continue;
 					}
 
-					FVector3d RayFrom = Position + 2.0 * Radius * ExteriorSamplePoints[k];
+					FVector3d RayFrom = Position + 2.0 * Radius * VisibilityDirections[k];
 					int32 HitTriID = FindHitTriangleTest(Position, TriNormal, RayFrom);
 					if ( HitTriID != IndexConstants::InvalidID && TriStatusKnown[HitTriID] == false )
 					{
