MeshModelingTools: Make remesh and simplify tools ask for dynamic mesh directly to avoid conversions back and forth for non-mesh description backed targets.

diff --git a/Engine/Plugins/Experimental/MeshModelingToolsetExp/Source/MeshModelingToolsEditorOnlyExp/Private/SimplifyMeshTool.cpp b/Engine/Plugins/Experimental/MeshModelingToolsetExp/Source/MeshModelingToolsEditorOnlyExp/Private/SimplifyMeshTool.cpp
--- a/Engine/Plugins/Experimental/MeshModelingToolsetExp/Source/MeshModelingToolsEditorOnlyExp/Private/SimplifyMeshTool.cpp
+++ b/Engine/Plugins/Experimental/MeshModelingToolsetExp/Source/MeshModelingToolsEditorOnlyExp/Private/SimplifyMeshTool.cpp
@@ -79,7 +79,6 @@ void USimplifyMeshTool::Setup()
 	// some of this could be done async...
 	{
 		// if in editor, create progress indicator dialog because building mesh copies can be slow (for very large meshes)
-		// this is especially needed because of the copy we make of the meshdescription; for Reasons, copying meshdescription is pretty slow
 #if WITH_EDITOR
 		static const FText SlowTaskText = LOCTEXT("SimplifyMeshInit", "Building mesh simplification data...");
 
@@ -88,23 +87,20 @@ void USimplifyMeshTool::Setup()
 
 		// Declare progress shortcut lambdas
 		auto EnterProgressFrame = [&SlowTask](int Progress)
-	{
+		{
 			SlowTask.EnterProgressFrame((float)Progress);
 		};
 #else
 		auto EnterProgressFrame = [](int Progress) {};
 #endif
-		FGetMeshParameters GetMeshParams;
-		GetMeshParams.bWantMeshTangents = true;
-		OriginalMeshDescription = MakeShared<FMeshDescription, ESPMode::ThreadSafe>(UE::ToolTarget::GetMeshDescriptionCopy(Target, GetMeshParams));
+		EnterProgressFrame(2);
 
-		EnterProgressFrame(1);
-		// UE::ToolTarget::GetDynamicMeshCopy() would recompute the tangents a second time here
-		OriginalMesh = MakeShared<FDynamicMesh3, ESPMode::ThreadSafe>();
-		FMeshDescriptionToDynamicMesh Converter; 
-		Converter.Convert(OriginalMeshDescription.Get(), *OriginalMesh, true);   // convert with tangent overlay
+		// Note that we only copy over a mesh description if we use the UEStandard path (conditionally done in MakeNewOperator)
+		// to avoid doing conversions back and forth for non mesh-description-backed targets (where the conversions can 
+		// occasionally do odd things to the attributes), and to avoid the slow copy unless the user needs it.
+		OriginalMesh = MakeShared<FDynamicMesh3, ESPMode::ThreadSafe>(UE::ToolTarget::GetDynamicMeshCopy(Target, true));
 
-		EnterProgressFrame(2);
+		EnterProgressFrame(1);
 		OriginalMeshSpatial = MakeShared<FDynamicMeshAABBTree3, ESPMode::ThreadSafe>(OriginalMesh.Get(), true);
 	}
 
@@ -192,6 +188,23 @@ TUniquePtr<FDynamicMeshOperator> USimplifyMeshTool::MakeNewOperator()
 {
 	TUniquePtr<FSimplifyMeshOp> Op = MakeUnique<FSimplifyMeshOp>();
 
+	if (SimplifyProperties->SimplifierType == ESimplifyType::UEStandard && OriginalMeshDescription == nullptr)
+	{
+		// if in editor, create progress indicator dialog because building mesh copies can be slow (for very large meshes)
+		// This is especially needed here because for Reasons, copying meshdescription is pretty slow
+#if WITH_EDITOR
+		static const FText SlowTaskText = LOCTEXT("SimplifyMeshInit", "Building mesh simplification data...");
+
+		FScopedSlowTask SlowTask(1.0f, SlowTaskText);
+		SlowTask.MakeDialog();
+
+#endif
+
+		FGetMeshParameters GetMeshParams;
+		GetMeshParams.bWantMeshTangents = true;
+		OriginalMeshDescription = MakeShared<FMeshDescription, ESPMode::ThreadSafe>(UE::ToolTarget::GetMeshDescriptionCopy(Target, GetMeshParams));
+	}
+
 	Op->bDiscardAttributes = SimplifyProperties->bDiscardAttributes;
 	// We always want attributes enabled on result even if we discard them initially
 	Op->bResultMustHaveAttributesEnabled = true;
diff --git a/Engine/Plugins/Experimental/MeshModelingToolsetExp/Source/MeshModelingToolsExp/Private/RemeshMeshTool.cpp b/Engine/Plugins/Experimental/MeshModelingToolsetExp/Source/MeshModelingToolsExp/Private/RemeshMeshTool.cpp
--- a/Engine/Plugins/Experimental/MeshModelingToolsetExp/Source/MeshModelingToolsExp/Private/RemeshMeshTool.cpp
+++ b/Engine/Plugins/Experimental/MeshModelingToolsetExp/Source/MeshModelingToolsExp/Private/RemeshMeshTool.cpp
@@ -16,7 +16,6 @@
 #include "DynamicMesh/DynamicMesh3.h"
 #include "DynamicMesh/DynamicMeshAABBTree3.h"
 
-#include "MeshDescriptionToDynamicMesh.h"
 #include "DynamicMeshToMeshDescription.h"
 
 #include "TargetInterfaces/MaterialProvider.h"
@@ -105,9 +104,7 @@ void URemeshMeshTool::Setup()
 	BasicProperties->WatchProperty(BasicProperties->bShowGroupColors,
 								   [this](bool bNewValue) { UpdateVisualization();});
 
-	OriginalMesh = MakeShared<FDynamicMesh3, ESPMode::ThreadSafe>();
-	FMeshDescriptionToDynamicMesh Converter;
-	Converter.Convert(UE::ToolTarget::GetMeshDescription(Targets[0]), *OriginalMesh);
+	OriginalMesh = MakeShared<FDynamicMesh3, ESPMode::ThreadSafe>(UE::ToolTarget::GetDynamicMeshCopy(Targets[0]));
 
 	Preview->PreviewMesh->SetTransform((FTransform) UE::ToolTarget::GetLocalToWorldTransform(Targets[0]));
 	Preview->PreviewMesh->SetTangentsMode(EDynamicMeshComponentTangentsMode::AutoCalculated);
