ModelingTools: Fixed PVS warning in MeshVertexPaintBrushOps.

diff --git a/Engine/Plugins/Experimental/MeshModelingToolsetExp/Source/MeshModelingToolsExp/Public/Sculpting/MeshVertexPaintBrushOps.h b/Engine/Plugins/Experimental/MeshModelingToolsetExp/Source/MeshModelingToolsExp/Public/Sculpting/MeshVertexPaintBrushOps.h
--- a/Engine/Plugins/Experimental/MeshModelingToolsetExp/Source/MeshModelingToolsExp/Public/Sculpting/MeshVertexPaintBrushOps.h
+++ b/Engine/Plugins/Experimental/MeshModelingToolsetExp/Source/MeshModelingToolsExp/Public/Sculpting/MeshVertexPaintBrushOps.h
@@ -72,7 +72,6 @@ namespace UELocal
 		RealType OneMinusAlpha = (RealType)1 - Alpha;
 		for (int32 k = 0; k < 4; ++k)
 		{
-			Result[k] = OneMinusAlpha * A[k] + Alpha * B[k];
 			Result[k] = TMathUtil<RealType>::Clamp(OneMinusAlpha*A[k] + Alpha*B[k], (RealType)0, (RealType)1);
 		}
 		return Result;
