Render Capture Baking: Move BakeRC tool from MeshModelingToolsExp to MeshModelingToolsEditorOnlyExp module

diff --git a/Engine/Plugins/Experimental/MeshModelingToolsetExp/Source/MeshModelingToolsEditorOnlyExp/Private/BakeRenderCaptureTool.cpp b/Engine/Plugins/Experimental/MeshModelingToolsetExp/Source/MeshModelingToolsEditorOnlyExp/Private/BakeRenderCaptureTool.cpp
new file mode 100644
--- /dev/null
+++ b/Engine/Plugins/Experimental/MeshModelingToolsetExp/Source/MeshModelingToolsEditorOnlyExp/Private/BakeRenderCaptureTool.cpp
@@ -0,0 +1,1319 @@
+﻿// Copyright Epic Games, Inc. All Rights Reserved.
+
+#include "BakeRenderCaptureTool.h"
+#include "Components/StaticMeshComponent.h"
+#include "Materials/Material.h"
+#include "TargetInterfaces/MaterialProvider.h"
+#include "TargetInterfaces/MeshDescriptionProvider.h"
+#include "TargetInterfaces/PrimitiveComponentBackedTarget.h"
+#include "TargetInterfaces/StaticMeshBackedTarget.h"
+#include "ToolTargetManager.h"
+#include "InteractiveToolManager.h"
+
+#include "DynamicMesh/MeshTransforms.h"
+
+#include "BakeToolUtils.h"
+#include "ToolSetupUtil.h"
+#include "ModelingToolTargetUtil.h"
+
+#include "ModelingObjectsCreationAPI.h"
+
+#include "EngineAnalytics.h"
+
+#include "Baking/BakingTypes.h"
+#include "Sampling/MeshImageBakingCache.h"
+#include "Sampling/MeshMapBaker.h"
+#include "AssetUtils/Texture2DBuilder.h"
+
+#include UE_INLINE_GENERATED_CPP_BY_NAME(BakeRenderCaptureTool)
+
+
+using namespace UE::Geometry;
+
+#define LOCTEXT_NAMESPACE "UBakeRenderCaptureTool"
+
+
+static FString BaseColorTexParamName = TEXT("BaseColor");
+static FString RoughnessTexParamName = TEXT("Roughness");
+static FString MetallicTexParamName = TEXT("Metallic");
+static FString SpecularTexParamName = TEXT("Specular");
+static FString EmissiveTexParamName = TEXT("Emissive");
+static FString OpacityTexParamName  = TEXT("Opacity");
+static FString SubsurfaceColorTexParamName = TEXT("SubsurfaceColor");
+static FString NormalTexParamName = TEXT("NormalMap");
+static FString PackedMRSTexParamName = TEXT("PackedMRS");
+
+
+FRenderCaptureOptions
+MakeRenderCaptureOptions(const URenderCaptureProperties& RenderCaptureProperties)
+{
+	FRenderCaptureOptions Options;
+
+	Options.bBakeBaseColor = RenderCaptureProperties.bBaseColorMap;
+	Options.bBakeNormalMap = RenderCaptureProperties.bNormalMap;
+	Options.bBakeEmissive =  RenderCaptureProperties.bEmissiveMap;
+	Options.bBakeOpacity = RenderCaptureProperties.bOpacityMap;
+	Options.bBakeSubsurfaceColor = RenderCaptureProperties.bSubsurfaceColorMap;
+	Options.bBakeDeviceDepth = RenderCaptureProperties.bDeviceDepthMap;
+	
+	// Enforce the PackedMRS precondition here so we don't have to check it at each usage site.  Note: We don't
+	// apply this precondition on the RenderCaptureProperties because we don't want the user to have to re-enable
+	// options which enabling PackedMRS disabled.
+	Options.bUsePackedMRS =  RenderCaptureProperties.bPackedMRSMap;
+	Options.bBakeMetallic =  RenderCaptureProperties.bPackedMRSMap ? false : RenderCaptureProperties.bMetallicMap;
+	Options.bBakeRoughness = RenderCaptureProperties.bPackedMRSMap ? false : RenderCaptureProperties.bRoughnessMap;
+	Options.bBakeSpecular =  RenderCaptureProperties.bPackedMRSMap ? false : RenderCaptureProperties.bSpecularMap;
+
+	Options.RenderCaptureImageSize = static_cast<int32>(RenderCaptureProperties.Resolution);
+	Options.bAntiAliasing = RenderCaptureProperties.bAntiAliasing;
+	Options.FieldOfViewDegrees = RenderCaptureProperties.CaptureFieldOfView;
+	Options.NearPlaneDist = RenderCaptureProperties.NearPlaneDist;
+
+	return Options;
+}
+
+
+//
+// Tool Operator
+//
+
+class FRenderCaptureMapBakerOp : public TGenericDataOperator<FMeshMapBaker>
+{
+public:
+	UE::Geometry::FDynamicMesh3* BaseMesh = nullptr;
+	UE::Geometry::FDynamicMeshAABBTree3* BaseMeshSpatial = nullptr;
+	TSharedPtr<UE::Geometry::FMeshTangentsd, ESPMode::ThreadSafe> BaseMeshTangents;
+	TSharedPtr<TArray<int32>, ESPMode::ThreadSafe> BaseMeshUVCharts;
+	int32 TargetUVLayer;
+	double ValidSampleDepthThreshold;
+	EBakeTextureResolution TextureImageSize;
+	EBakeTextureSamplesPerPixel SamplesPerPixel;
+	FSceneCapturePhotoSet* SceneCapture = nullptr;
+
+	// Used to pass the channels which need baking via the bBakeXXX and bUsePackedMRS members
+	FRenderCaptureOptions PendingBake;
+
+	// Begin TGenericDataOperator interface
+	virtual void CalculateResult(FProgressCancel* Progress) override;
+	// End TGenericDataOperator interface
+};
+
+// Bake textures onto the base/target mesh by projecting/sampling the set of captured photos
+void FRenderCaptureMapBakerOp::CalculateResult(FProgressCancel*)
+{
+	FSceneCapturePhotoSetSampler Sampler(
+		SceneCapture,
+		ValidSampleDepthThreshold,
+		BaseMesh,
+		BaseMeshSpatial,
+		BaseMeshTangents.Get());
+
+	const FImageDimensions TextureDimensions(
+		static_cast<int32>(TextureImageSize),
+		static_cast<int32>(TextureImageSize));
+
+	FRenderCaptureOcclusionHandler OcclusionHandler(TextureDimensions);
+
+	Result = MakeRenderCaptureBaker(
+		BaseMesh,
+		BaseMeshTangents,
+		BaseMeshUVCharts,
+		SceneCapture,
+		&Sampler,
+		PendingBake,
+		TargetUVLayer,
+		TextureImageSize,
+		SamplesPerPixel,
+		&OcclusionHandler);
+
+	{
+		TRACE_CPUPROFILER_EVENT_SCOPE(FRenderCaptureMapBakerOp_CalculateResult_Bake);
+		Result->Bake();
+	}
+}
+
+
+//
+// Tool Builder
+//
+
+
+const FToolTargetTypeRequirements& UBakeRenderCaptureToolBuilder::GetTargetRequirements() const
+{
+	static FToolTargetTypeRequirements TypeRequirements({
+		UMeshDescriptionProvider::StaticClass(),
+		UPrimitiveComponentBackedTarget::StaticClass(),
+		UStaticMeshBackedTarget::StaticClass(),			// FMeshSceneAdapter currently only supports StaticMesh targets
+		UMaterialProvider::StaticClass()
+		});
+	return TypeRequirements;
+}
+
+bool UBakeRenderCaptureToolBuilder::CanBuildTool(const FToolBuilderState& SceneState) const
+{
+	const int32 NumTargets = SceneState.TargetManager->CountSelectedAndTargetable(SceneState, GetTargetRequirements());
+	return (NumTargets > 1);
+}
+
+UMultiSelectionMeshEditingTool* UBakeRenderCaptureToolBuilder::CreateNewTool(const FToolBuilderState& SceneState) const
+{
+	return NewObject<UBakeRenderCaptureTool>(SceneState.ToolManager);
+}
+
+
+
+//
+// Tool
+//
+
+
+
+
+void UBakeRenderCaptureTool::Setup()
+{
+	TRACE_CPUPROFILER_EVENT_SCOPE(UBakeRenderCaptureTool::Setup);
+
+	Super::Setup();
+
+	UE::ToolTarget::HideSourceObject(Targets[0]);
+
+	// Initialize materials and textures for background compute, error feedback and previewing the results
+	InitializePreviewMaterials();
+
+	// Initialize the PreviewMesh, which displays intermediate results
+	PreviewMesh = CreateBakePreviewMesh(this, Targets[0], GetTargetWorld());
+
+	// Initialize the datastructures used by the bake background compute/tool operator
+	PreviewMesh->ProcessMesh([this](const FDynamicMesh3& Mesh)
+	{
+		TargetMesh.Copy(Mesh);
+		const FTransformSRT3d BaseToWorld = UE::ToolTarget::GetLocalToWorldTransform(Targets[0]);
+		MeshTransforms::ApplyTransform(TargetMesh, BaseToWorld, true);
+
+		// Initialize UV charts
+		TargetMeshUVCharts = MakeShared<TArray<int32>, ESPMode::ThreadSafe>();
+		FMeshMapBaker::ComputeUVCharts(TargetMesh, *TargetMeshUVCharts);
+
+		// Initialize tangents
+		TargetMeshTangents = MakeShared<FMeshTangentsd, ESPMode::ThreadSafe>(&TargetMesh);
+		TargetMeshTangents->CopyTriVertexTangents(TargetMesh);
+
+		// Initialize spatial index
+		TargetMeshSpatial.SetMesh(&TargetMesh, true);
+	});
+
+	// Initialize actors
+	const int NumTargets = Targets.Num();
+	Actors.Empty(NumTargets - 1);
+	for (int Idx = 1; Idx < NumTargets; ++Idx)
+	{
+		if (AActor* Actor = UE::ToolTarget::GetTargetActor(Targets[Idx]))
+		{
+			Actors.Add(Actor);
+		}
+	}
+
+	UToolTarget* Target = Targets[0];
+
+	// Setup tool property sets
+
+	Settings = NewObject<UBakeRenderCaptureToolProperties>(this);
+	Settings->RestoreProperties(this);
+	AddToolPropertySource(Settings);
+
+	Settings->MapPreview = BaseColorTexParamName;
+	Settings->WatchProperty(Settings->MapPreview, [this](FString) { UpdateVisualization(); });
+	Settings->WatchProperty(Settings->SamplesPerPixel, [this](EBakeTextureSamplesPerPixel) { OpState |= EBakeOpState::Evaluate; });
+	Settings->WatchProperty(Settings->TextureSize, [this](EBakeTextureResolution) { OpState |= EBakeOpState::Evaluate; });
+	Settings->WatchProperty(Settings->ValidSampleDepthThreshold, [this](float ValidSampleDepthThreshold)
+	{
+		// Only compute the device depth if we compute at least one other channel, the DeviceDepth is used to eliminate
+		// occlusion artefacts from the other channels
+		RenderCaptureProperties->bDeviceDepthMap = (ValidSampleDepthThreshold > 0) &&
+			(
+			RenderCaptureProperties->bBaseColorMap ||
+			RenderCaptureProperties->bNormalMap    ||
+			RenderCaptureProperties->bEmissiveMap  ||
+			RenderCaptureProperties->bOpacityMap   ||
+			RenderCaptureProperties->bSubsurfaceColorMap ||
+			RenderCaptureProperties->bPackedMRSMap ||
+			RenderCaptureProperties->bMetallicMap  ||
+			RenderCaptureProperties->bRoughnessMap ||
+			RenderCaptureProperties->bSpecularMap
+			);
+
+		OpState |= EBakeOpState::Evaluate;
+	});
+
+	// Put these properties before the list of preview textures so its easier to find
+	VisualizationProps = NewObject<UBakeRenderCaptureVisualizationProperties>(this);
+	VisualizationProps->RestoreProperties(this);
+	AddToolPropertySource(VisualizationProps);
+	VisualizationProps->WatchProperty(VisualizationProps->bPreviewAsMaterial, [this](bool) { UpdateVisualization(); });
+
+	RenderCaptureProperties = NewObject<URenderCaptureProperties>(this);
+	RenderCaptureProperties->RestoreProperties(this);
+	AddToolPropertySource(RenderCaptureProperties);
+
+	RenderCaptureProperties->WatchProperty(RenderCaptureProperties->Resolution, [this](EBakeTextureResolution) { OpState |= EBakeOpState::Evaluate; });
+	RenderCaptureProperties->WatchProperty(RenderCaptureProperties->bBaseColorMap, [this](bool) { OpState |= EBakeOpState::Evaluate; });
+	RenderCaptureProperties->WatchProperty(RenderCaptureProperties->bNormalMap, [this](bool) { OpState |= EBakeOpState::Evaluate; });
+	RenderCaptureProperties->WatchProperty(RenderCaptureProperties->bMetallicMap, [this](bool) { OpState |= EBakeOpState::Evaluate; });
+	RenderCaptureProperties->WatchProperty(RenderCaptureProperties->bRoughnessMap, [this](bool) { OpState |= EBakeOpState::Evaluate; });
+	RenderCaptureProperties->WatchProperty(RenderCaptureProperties->bSpecularMap, [this](bool) { OpState |= EBakeOpState::Evaluate; });
+	RenderCaptureProperties->WatchProperty(RenderCaptureProperties->bPackedMRSMap, [this](bool) { OpState |= EBakeOpState::Evaluate; });
+	RenderCaptureProperties->WatchProperty(RenderCaptureProperties->bEmissiveMap, [this](bool) { OpState |= EBakeOpState::Evaluate; });
+	RenderCaptureProperties->WatchProperty(RenderCaptureProperties->bOpacityMap, [this](bool) { OpState |= EBakeOpState::Evaluate; });
+	RenderCaptureProperties->WatchProperty(RenderCaptureProperties->bSubsurfaceColorMap, [this](bool) { OpState |= EBakeOpState::Evaluate; });
+	RenderCaptureProperties->WatchProperty(RenderCaptureProperties->bAntiAliasing, [this](bool) { OpState |= EBakeOpState::Evaluate; });
+	// These are not exposed to the UI, but we watch them anyway because we might change that later
+	RenderCaptureProperties->WatchProperty(RenderCaptureProperties->CaptureFieldOfView, [this](float) { OpState |= EBakeOpState::Evaluate; });
+	RenderCaptureProperties->WatchProperty(RenderCaptureProperties->NearPlaneDist, [this](float) { OpState |= EBakeOpState::Evaluate; });
+	
+	InputMeshSettings = NewObject<UBakeRenderCaptureInputToolProperties>(this);
+	InputMeshSettings->RestoreProperties(this);
+	AddToolPropertySource(InputMeshSettings);
+	InputMeshSettings->TargetStaticMesh = UE::ToolTarget::GetStaticMeshFromTargetIfAvailable(Target);
+	UpdateUVLayerNames(InputMeshSettings->TargetUVLayer, InputMeshSettings->TargetUVLayerNamesList, TargetMesh);
+	InputMeshSettings->WatchProperty(InputMeshSettings->TargetUVLayer, [this](FString) { OpState |= EBakeOpState::Evaluate; });
+	
+	{
+		Settings->MapPreviewNamesList.Add(BaseColorTexParamName);
+		Settings->MapPreviewNamesList.Add(NormalTexParamName);
+		Settings->MapPreviewNamesList.Add(PackedMRSTexParamName);
+		Settings->MapPreviewNamesList.Add(MetallicTexParamName);
+		Settings->MapPreviewNamesList.Add(RoughnessTexParamName);
+		Settings->MapPreviewNamesList.Add(SpecularTexParamName);
+		Settings->MapPreviewNamesList.Add(EmissiveTexParamName);
+		Settings->MapPreviewNamesList.Add(OpacityTexParamName);
+		Settings->MapPreviewNamesList.Add(SubsurfaceColorTexParamName);
+	}
+
+	ResultSettings = NewObject<UBakeRenderCaptureResults>(this);
+	ResultSettings->RestoreProperties(this);
+	AddToolPropertySource(ResultSettings);
+	SetToolPropertySourceEnabled(ResultSettings, true);
+
+	TargetUVLayerToError.Reset();
+
+	// Hide the render capture meshes since this baker operates solely in world space which will occlude the preview of
+	// the target mesh.
+	for (int Idx = 1; Idx < NumTargets; ++Idx)
+	{
+		UE::ToolTarget::HideSourceObject(Targets[Idx]);
+	}
+
+	// Make sure we trigger SceneCapture computation in UpdateResult
+	OpState |= EBakeOpState::Evaluate;
+
+	SetToolDisplayName(LOCTEXT("ToolName", "Bake Render Capture"));
+	GetToolManager()->DisplayMessage(
+		LOCTEXT("OnStartTool", "Bake Render Capture. Select Bake Mesh (LowPoly) first, then select Detail Meshes (HiPoly) to bake. Assets will be created on Accept."),
+		EToolMessageLevel::UserNotification);
+
+	GatherAnalytics(BakeAnalytics.MeshSettings);
+}
+
+
+
+
+
+void UBakeRenderCaptureTool::Render(IToolsContextRenderAPI* RenderAPI)
+{
+	// This block contains things we want to update every frame to respond to sliders in the UI
+	{
+		const float Brightness = VisualizationProps->Brightness;
+		const FVector BaseColorBrightness(Brightness, Brightness, Brightness);
+
+		PreviewMaterialRC->SetVectorParameterValue(TEXT("Brightness"), BaseColorBrightness);
+		PreviewMaterialPackedRC->SetVectorParameterValue(TEXT("Brightness"), BaseColorBrightness);
+		PreviewMaterialRC_Subsurface->SetVectorParameterValue(TEXT("Brightness"), BaseColorBrightness);
+		PreviewMaterialPackedRC_Subsurface->SetVectorParameterValue(TEXT("Brightness"), BaseColorBrightness);
+
+		const float EmissiveScale = VisualizationProps->EmissiveScale;
+
+		PreviewMaterialRC->SetScalarParameterValue(TEXT("EmissiveScale"), EmissiveScale);
+		PreviewMaterialPackedRC->SetScalarParameterValue(TEXT("EmissiveScale"), EmissiveScale);
+		PreviewMaterialRC_Subsurface->SetScalarParameterValue(TEXT("EmissiveScale"), EmissiveScale);
+		PreviewMaterialPackedRC_Subsurface->SetScalarParameterValue(TEXT("EmissiveScale"), EmissiveScale);
+
+		const float SSBrightness = VisualizationProps->SSBrightness;
+		const FVector SubsurfaceColorBrightness(SSBrightness, SSBrightness, SSBrightness);
+
+		PreviewMaterialRC_Subsurface->SetVectorParameterValue(TEXT("BrightnessSubsurface"), SubsurfaceColorBrightness);
+		PreviewMaterialPackedRC_Subsurface->SetVectorParameterValue(TEXT("BrightnessSubsurface"), SubsurfaceColorBrightness);
+	}
+
+	// This is a no-op if OpState == EBakeOpState::Clean
+	UpdateResult();
+}
+
+
+void UBakeRenderCaptureTool::OnTick(float DeltaTime)
+{
+	if (Compute)
+	{
+		Compute->Tick(DeltaTime);
+
+		if (static_cast<bool>(OpState & EBakeOpState::Invalid))
+		{
+			PreviewMesh->SetOverrideRenderMaterial(ErrorPreviewMaterial);
+		}
+		else
+		{
+			const float ElapsedComputeTime = Compute->GetElapsedComputeTime();
+			if (!CanAccept() && ElapsedComputeTime > SecondsBeforeWorkingMaterial)
+			{
+				PreviewMesh->SetOverrideRenderMaterial(WorkingPreviewMaterial);
+			}
+		}
+	}
+	else if (static_cast<bool>(OpState & EBakeOpState::Invalid))
+	{
+		PreviewMesh->SetOverrideRenderMaterial(ErrorPreviewMaterial);
+	} 
+}
+
+
+void UBakeRenderCaptureTool::OnShutdown(EToolShutdownType ShutdownType)
+{
+	TRACE_CPUPROFILER_EVENT_SCOPE(UBakeRenderCaptureTool::Shutdown);
+
+	Settings->SaveProperties(this);
+	RenderCaptureProperties->SaveProperties(this);
+	InputMeshSettings->SaveProperties(this);
+	VisualizationProps->SaveProperties(this);
+
+	if (PreviewMesh != nullptr)
+	{
+		PreviewMesh->SetVisible(false);
+		PreviewMesh->Disconnect();
+		PreviewMesh = nullptr;
+	}
+
+	if (Compute)
+	{
+		Compute->Shutdown();
+	}
+
+	// Restore visibility of source meshes
+	const int NumTargets = Targets.Num();
+	for (int Idx = 1; Idx < NumTargets; ++Idx)
+	{
+		UE::ToolTarget::ShowSourceObject(Targets[Idx]);
+	}
+	
+	if (ShutdownType == EToolShutdownType::Accept)
+	{
+		// TODO Support skeletal meshes here---see BakeMeshAttributeMapsTool::OnShutdown
+		IStaticMeshBackedTarget* StaticMeshTarget = Cast<IStaticMeshBackedTarget>(Targets[0]);
+		UObject* SourceAsset = StaticMeshTarget ? StaticMeshTarget->GetStaticMesh() : nullptr;
+		const UPrimitiveComponent* SourceComponent = UE::ToolTarget::GetTargetComponent(Targets[0]);
+		CreateAssets(SourceComponent->GetWorld(), SourceAsset);
+	}
+
+	// Clear actors on shutdown so that their lifetime is not tied to the lifetime of the tool
+	Actors.Empty();
+
+	UE::ToolTarget::ShowSourceObject(Targets[0]);
+}
+
+void UBakeRenderCaptureTool::CreateAssets(UWorld* SourceWorld, UObject* SourceAsset)
+{
+	const FString BaseName = UE::ToolTarget::GetTargetActor(Targets[0])->GetActorNameOrLabel();
+	const bool bPackedMRS = ResultSettings->PackedMRSMap != nullptr;
+	const bool bSubsurfaceMaterial = ResultSettings->SubsurfaceColorMap || ResultSettings->OpacityMap;
+
+	{
+		// Choose the material
+		TObjectPtr<UMaterialInstanceDynamic> Material = ResultSettings->PackedMRSMap ? PreviewMaterialPackedRC : PreviewMaterialRC;
+		if (bSubsurfaceMaterial)
+		{
+			Material = ResultSettings->PackedMRSMap ? PreviewMaterialPackedRC_Subsurface : PreviewMaterialRC_Subsurface;
+			ensure(Material->GetShadingModels().HasShadingModel(EMaterialShadingModel::MSM_Subsurface));
+			ensure(Material->GetBlendMode() == EBlendMode::BLEND_Masked);
+		}
+
+		FCreateMaterialObjectParams MaterialParams;
+		MaterialParams.TargetWorld = SourceWorld;
+		MaterialParams.StoreRelativeToObject = SourceAsset;
+		MaterialParams.BaseName = FString::Printf(TEXT("%s_Material"), *BaseName);
+		MaterialParams.MaterialToDuplicate = Material;
+		const FCreateMaterialObjectResult MaterialResult = UE::Modeling::CreateMaterialObject(GetToolManager(), MoveTemp(MaterialParams));
+
+		UMaterial* NewMaterial = nullptr;
+		if (ensure(MaterialResult.ResultCode == ECreateModelingObjectResult::Ok))
+		{
+			NewMaterial = CastChecked<UMaterial>(MaterialResult.NewAsset);
+			ensure(NewMaterial);
+		}
+
+		const auto TrySetTextureEditorOnly = [NewMaterial](
+			const FString& TextureName,
+			TObjectPtr<UTexture2D> Texture,
+			TObjectPtr<UTexture2D> Fallback,
+			bool bMaterialHasTexture)
+		{
+			if (NewMaterial && bMaterialHasTexture)
+			{
+				NewMaterial->SetTextureParameterValueEditorOnly(FName(TextureName), Texture ? Texture : Fallback);
+			}
+		};
+
+		// Set all computed textures or fallback to the empty texture map
+		TrySetTextureEditorOnly(BaseColorTexParamName, ResultSettings->BaseColorMap, EmptyColorMapWhite, true);
+		TrySetTextureEditorOnly(EmissiveTexParamName,  ResultSettings->EmissiveMap,  EmptyEmissiveMap,   true);
+		TrySetTextureEditorOnly(NormalTexParamName,    ResultSettings->NormalMap,    EmptyNormalMap,     true);
+		TrySetTextureEditorOnly(PackedMRSTexParamName, ResultSettings->PackedMRSMap, EmptyPackedMRSMap,  bPackedMRS);
+		TrySetTextureEditorOnly(MetallicTexParamName,  ResultSettings->MetallicMap,  EmptyMetallicMap,  !bPackedMRS);
+		TrySetTextureEditorOnly(RoughnessTexParamName, ResultSettings->RoughnessMap, EmptyRoughnessMap, !bPackedMRS);
+		TrySetTextureEditorOnly(SpecularTexParamName,  ResultSettings->SpecularMap,  EmptySpecularMap,  !bPackedMRS);
+		TrySetTextureEditorOnly(OpacityTexParamName,   ResultSettings->OpacityMap,   EmptyOpacityMap,    bSubsurfaceMaterial);
+		TrySetTextureEditorOnly(SubsurfaceColorTexParamName, ResultSettings->SubsurfaceColorMap, EmptySubsurfaceColorMap, bSubsurfaceMaterial);
+
+		if (NewMaterial)
+		{
+			// Force material update now that we have updated texture parameters
+			NewMaterial->PostEditChange();
+		}
+	}
+
+	auto CreateTextureAsset = [this, &BaseName, &SourceWorld, &SourceAsset] (
+		const FString& TexParamName,
+		FTexture2DBuilder::ETextureType Type,
+		TObjectPtr<UTexture2D> Texture,
+		TObjectPtr<UTexture2D> Fallback,
+		bool bForceFallback,
+		bool bMaterialHasTexture)
+	{
+		if (!bMaterialHasTexture)
+		{
+			return true;
+		}
+
+		const bool bUseTexture = Texture && !bForceFallback;
+
+		// See :DeferredPopulateSourceData
+		if (bUseTexture)
+		{
+			FTexture2DBuilder::CopyPlatformDataToSourceData(Texture, Type);
+		}
+		else
+		{
+			FTexture2DBuilder::CopyPlatformDataToSourceData(Fallback, Type);
+		}
+
+		// TODO The original implementation in ApproximateActors also did the following, see WriteTextureLambda in ApproximateActorsImpl.cpp
+		//if (Type == FTexture2DBuilder::ETextureType::Roughness
+		//	|| Type == FTexture2DBuilder::ETextureType::Metallic
+		//	|| Type == FTexture2DBuilder::ETextureType::Specular)
+		//{
+		//	UE::AssetUtils::ConvertToSingleChannel(Texture);
+		//}
+
+		// We need to save the Fallback textures as well so that the generated Materials can reference them
+		FCreateTextureObjectParams TexParams;
+		TexParams.TargetWorld = SourceWorld;
+		TexParams.StoreRelativeToObject = SourceAsset;
+		TexParams.BaseName = FString::Printf(TEXT("%s_%s"), *BaseName, *TexParamName);
+		TexParams.GeneratedTransientTexture = bUseTexture ? Texture : Fallback; 
+
+		FCreateTextureObjectResult TexResult = UE::Modeling::CreateTextureObject(GetToolManager(), MoveTemp(TexParams));
+		return TexResult.IsOK();
+	};
+
+	ensure(CreateTextureAsset(BaseColorTexParamName,
+		FTexture2DBuilder::ETextureType::Color,
+		ResultSettings->BaseColorMap,
+		EmptyColorMapWhite,
+		RenderCaptureProperties->bBaseColorMap == false,
+		true));
+
+	ensure(CreateTextureAsset(NormalTexParamName,
+		FTexture2DBuilder::ETextureType::NormalMap,
+		ResultSettings->NormalMap,
+		EmptyNormalMap,
+		RenderCaptureProperties->bNormalMap == false,
+		true));
+
+	ensure(CreateTextureAsset(EmissiveTexParamName,
+		FTexture2DBuilder::ETextureType::EmissiveHDR,
+		ResultSettings->EmissiveMap,
+		EmptyEmissiveMap,
+		RenderCaptureProperties->bEmissiveMap == false,
+		true));
+
+	ensure(CreateTextureAsset(PackedMRSTexParamName,
+		FTexture2DBuilder::ETextureType::ColorLinear,
+		ResultSettings->PackedMRSMap,
+		EmptyPackedMRSMap,
+		RenderCaptureProperties->bPackedMRSMap == false,
+		bPackedMRS));
+
+	ensure(CreateTextureAsset(RoughnessTexParamName,
+		FTexture2DBuilder::ETextureType::Roughness,
+		ResultSettings->RoughnessMap,
+		EmptyRoughnessMap,
+		RenderCaptureProperties->bRoughnessMap == false,
+		!bPackedMRS));
+
+	ensure(CreateTextureAsset(MetallicTexParamName,
+		FTexture2DBuilder::ETextureType::Metallic,
+		ResultSettings->MetallicMap,
+		EmptyMetallicMap,
+		RenderCaptureProperties->bMetallicMap == false,
+		!bPackedMRS));
+
+	ensure(CreateTextureAsset(SpecularTexParamName,
+		FTexture2DBuilder::ETextureType::Specular,
+		ResultSettings->SpecularMap,
+		EmptySpecularMap,
+		RenderCaptureProperties->bSpecularMap == false,
+		!bPackedMRS));
+
+	ensure(CreateTextureAsset(OpacityTexParamName,
+		FTexture2DBuilder::ETextureType::ColorLinear,
+		ResultSettings->OpacityMap,
+		EmptyOpacityMap,
+		RenderCaptureProperties->bOpacityMap == false,
+		bSubsurfaceMaterial));
+
+	ensure(CreateTextureAsset(SubsurfaceColorTexParamName,
+		FTexture2DBuilder::ETextureType::Color,
+		ResultSettings->SubsurfaceColorMap,
+		EmptySubsurfaceColorMap,
+		RenderCaptureProperties->bSubsurfaceColorMap == false,
+		bSubsurfaceMaterial));
+
+	RecordAnalytics();
+}
+
+
+
+bool UBakeRenderCaptureTool::CanAccept() const
+{
+	if ((OpState & EBakeOpState::Invalid) == EBakeOpState::Invalid)
+	{
+		return false;
+	}
+
+	if (RenderCaptureProperties->bBaseColorMap && ResultSettings->BaseColorMap == nullptr)
+	{
+		return false;
+	}
+	if (RenderCaptureProperties->bNormalMap && ResultSettings->NormalMap == nullptr)
+	{
+		return false;
+	}
+	if (RenderCaptureProperties->bEmissiveMap && ResultSettings->EmissiveMap == nullptr)
+	{
+		return false;
+	}
+	if (RenderCaptureProperties->bOpacityMap && ResultSettings->OpacityMap == nullptr)
+	{
+		return false;
+	}
+	if (RenderCaptureProperties->bSubsurfaceColorMap && ResultSettings->SubsurfaceColorMap == nullptr)
+	{
+		return false;
+	}
+	
+	// We need different code paths based on PackedMRS here because we don't want to uncheck the separate channels
+	// when PackedMRS is enabled to give the user a better UX (they don't have to re-check them after disabling
+	// PackedMRS). In other place we can test the PackedMRS and separate channel booleans in series and avoid the
+	// complexity of nested if statements.
+	if (RenderCaptureProperties->bPackedMRSMap)
+	{
+		if (ResultSettings->PackedMRSMap == nullptr)
+		{
+			return false;
+		}
+	}
+	else
+	{
+		if (RenderCaptureProperties->bMetallicMap && ResultSettings->MetallicMap == nullptr)
+		{
+			return false;
+		}
+		if (RenderCaptureProperties->bRoughnessMap && ResultSettings->RoughnessMap == nullptr)
+		{
+			return false;
+		}
+		if (RenderCaptureProperties->bSpecularMap && ResultSettings->SpecularMap == nullptr)
+		{
+			return false;
+		}
+	}
+
+	return true;
+}
+
+
+
+TUniquePtr<TGenericDataOperator<FMeshMapBaker>> UBakeRenderCaptureTool::MakeNewOperator()
+{
+	// We should not have requested a bake if we don't have a SceneCapture
+	check(SceneCapture.IsValid());
+
+	TUniquePtr<FRenderCaptureMapBakerOp> Op = MakeUnique<FRenderCaptureMapBakerOp>();
+	Op->BaseMesh = &TargetMesh;
+	Op->BaseMeshSpatial = &TargetMeshSpatial;
+	Op->BaseMeshTangents = TargetMeshTangents;
+	Op->BaseMeshUVCharts = TargetMeshUVCharts;
+	Op->TargetUVLayer = InputMeshSettings->GetTargetUVLayerIndex();
+	Op->ValidSampleDepthThreshold = Settings->ValidSampleDepthThreshold;
+	Op->TextureImageSize = Settings->TextureSize;
+	Op->SamplesPerPixel = Settings->SamplesPerPixel;
+	Op->SceneCapture = SceneCapture.Get();
+
+	auto IsPendingBake = [this](UTexture *CurrentResult, ERenderCaptureType CaptureType) -> bool
+	{
+		const bool bRequested = SceneCapture->GetCaptureTypeEnabled(CaptureType);
+		const bool bBaked = (CurrentResult != nullptr);
+		return (bRequested && !bBaked);
+	};
+
+	Op->PendingBake.bBakeBaseColor = IsPendingBake(ResultSettings->BaseColorMap, ERenderCaptureType::BaseColor);
+	Op->PendingBake.bBakeRoughness = IsPendingBake(ResultSettings->RoughnessMap, ERenderCaptureType::Roughness);
+	Op->PendingBake.bBakeMetallic =  IsPendingBake(ResultSettings->MetallicMap,  ERenderCaptureType::Metallic);
+	Op->PendingBake.bBakeSpecular =  IsPendingBake(ResultSettings->SpecularMap,  ERenderCaptureType::Specular);
+	Op->PendingBake.bUsePackedMRS =  IsPendingBake(ResultSettings->PackedMRSMap, ERenderCaptureType::CombinedMRS);
+	Op->PendingBake.bBakeEmissive =  IsPendingBake(ResultSettings->EmissiveMap,  ERenderCaptureType::Emissive);
+	Op->PendingBake.bBakeNormalMap = IsPendingBake(ResultSettings->NormalMap,    ERenderCaptureType::WorldNormal);
+	Op->PendingBake.bBakeOpacity =   IsPendingBake(ResultSettings->OpacityMap,   ERenderCaptureType::Opacity);
+	Op->PendingBake.bBakeSubsurfaceColor = IsPendingBake(ResultSettings->SubsurfaceColorMap, ERenderCaptureType::SubsurfaceColor);
+
+	return Op;
+}
+
+
+
+void UBakeRenderCaptureTool::OnMapsUpdated(const TUniquePtr<FMeshMapBaker>& NewResult)
+{
+	TRACE_CPUPROFILER_EVENT_SCOPE(BakeRenderCaptureTool_Textures_BuildTextures);
+
+	FRenderCaptureTextures TexturesOut;
+	GetTexturesFromRenderCaptureBaker(NewResult, TexturesOut);
+
+	// The NewResult will contain the newly baked textures so we only update those and not overwrite any already baked
+	// valid results. Note that if results are invalidated by some parameter change they are cleared by the
+	// InvalidateResults function, by comparing/syncing with the enabled capture types in SceneCapture
+	if (TexturesOut.BaseColorMap)
+	{
+		ResultSettings->BaseColorMap = TexturesOut.BaseColorMap;
+	}
+	if (TexturesOut.NormalMap)
+	{
+		ResultSettings->NormalMap = TexturesOut.NormalMap;
+	}
+	if (TexturesOut.PackedMRSMap)
+	{
+		ResultSettings->PackedMRSMap = TexturesOut.PackedMRSMap;
+	}
+	if (TexturesOut.MetallicMap)
+	{
+		ResultSettings->MetallicMap = TexturesOut.MetallicMap;
+	}
+	if (TexturesOut.RoughnessMap)
+	{
+		ResultSettings->RoughnessMap = TexturesOut.RoughnessMap;
+	}
+	if (TexturesOut.SpecularMap)
+	{
+		ResultSettings->SpecularMap = TexturesOut.SpecularMap;
+	}
+	if (TexturesOut.EmissiveMap)
+	{
+		ResultSettings->EmissiveMap = TexturesOut.EmissiveMap;
+	}
+	if (TexturesOut.OpacityMap)
+	{
+		ResultSettings->OpacityMap = TexturesOut.OpacityMap;
+	}
+	if (TexturesOut.SubsurfaceColorMap)
+	{
+		ResultSettings->SubsurfaceColorMap = TexturesOut.SubsurfaceColorMap;
+	}
+
+	GatherAnalytics(*NewResult);
+	UpdateVisualization();
+}
+
+bool UBakeRenderCaptureTool::ValidTargetMeshTangents()
+{
+	if (bCheckTargetMeshTangents)
+	{
+		bValidTargetMeshTangents = TargetMeshTangents ? FDynamicMeshTangents(&TargetMesh).HasValidTangents(true) : false;
+		bCheckTargetMeshTangents = false;
+	}
+	return bValidTargetMeshTangents;
+}
+
+void UBakeRenderCaptureTool::InitializePreviewMaterials()
+{
+	// We will only need source data if we're saving these textures so that the material generated when the tool is
+	// Accepted can reference them :DeferredPopulateSourceData
+	constexpr bool bPopulateSourceData = false;
+
+	{
+		FTexture2DBuilder Builder;
+		Builder.Initialize(FTexture2DBuilder::ETextureType::NormalMap, FImageDimensions(16, 16));
+		Builder.Commit(bPopulateSourceData);
+		EmptyNormalMap = Builder.GetTexture2D();
+	}
+
+	{
+		FTexture2DBuilder Builder;
+		Builder.Initialize(FTexture2DBuilder::ETextureType::Color, FImageDimensions(16, 16));
+		Builder.Clear(FColor(0,0,0));
+		Builder.Commit(bPopulateSourceData);
+		EmptyColorMapBlack = Builder.GetTexture2D();
+	}
+
+	{
+		FTexture2DBuilder Builder;
+		Builder.Initialize(FTexture2DBuilder::ETextureType::Color, FImageDimensions(16, 16));
+		Builder.Clear(FColor::White);
+		Builder.Commit(bPopulateSourceData);
+		EmptyColorMapWhite = Builder.GetTexture2D();
+	}
+
+	{
+		FTexture2DBuilder Builder;
+		Builder.Initialize(FTexture2DBuilder::ETextureType::EmissiveHDR, FImageDimensions(16, 16));
+		Builder.Commit(bPopulateSourceData);
+		EmptyEmissiveMap = Builder.GetTexture2D();
+	}
+
+	{
+		FTexture2DBuilder Builder;
+		Builder.Initialize(FTexture2DBuilder::ETextureType::ColorLinear, FImageDimensions(16, 16));
+		// The Opacity texture is passed to the Material's Opacity pin as well as the Opacity Mask pin, so set white
+		// here so we see something when previewing the subsurface material when opacity is not baked
+		Builder.Clear(FColor::White);
+		Builder.Commit(bPopulateSourceData);
+		EmptyOpacityMap = Builder.GetTexture2D();
+	}
+
+	{
+		FTexture2DBuilder Builder;
+		Builder.Initialize(FTexture2DBuilder::ETextureType::Color, FImageDimensions(16, 16));
+		Builder.Clear(FColor::Black);
+		Builder.Commit(bPopulateSourceData);
+		EmptySubsurfaceColorMap = Builder.GetTexture2D();
+	}
+
+	{
+		FTexture2DBuilder Builder;
+		Builder.Initialize(FTexture2DBuilder::ETextureType::ColorLinear, FImageDimensions(16, 16));
+		Builder.Clear(FColor(0,0,0));
+		Builder.Commit(bPopulateSourceData);
+		EmptyPackedMRSMap = Builder.GetTexture2D();
+	}
+
+	{
+		FTexture2DBuilder Builder;
+		Builder.Initialize(FTexture2DBuilder::ETextureType::Roughness, FImageDimensions(16, 16));
+		Builder.Commit(bPopulateSourceData);
+		EmptyRoughnessMap = Builder.GetTexture2D();
+	}
+
+	{
+		FTexture2DBuilder Builder;
+		Builder.Initialize(FTexture2DBuilder::ETextureType::Metallic, FImageDimensions(16, 16));
+		Builder.Commit(bPopulateSourceData);
+		EmptyMetallicMap = Builder.GetTexture2D();
+	}
+
+	{
+		FTexture2DBuilder Builder;
+		Builder.Initialize(FTexture2DBuilder::ETextureType::Specular, FImageDimensions(16, 16));
+		Builder.Commit(bPopulateSourceData);
+		EmptySpecularMap = Builder.GetTexture2D();
+	}
+
+	WorkingPreviewMaterial = ToolSetupUtil::GetDefaultWorkingMaterialInstance(GetToolManager());
+
+	ErrorPreviewMaterial = ToolSetupUtil::GetDefaultErrorMaterial(GetToolManager());
+
+	{
+		UMaterial* Material = LoadObject<UMaterial>(nullptr,
+			TEXT("/MeshModelingToolsetExp/Materials/BakeRenderCapturePreviewMaterial"));
+		check(Material);
+		if (Material != nullptr)
+		{
+			PreviewMaterialRC = UMaterialInstanceDynamic::Create(Material, GetToolManager());
+			PreviewMaterialRC->SetTextureParameterValue(FName(*BaseColorTexParamName), EmptyColorMapWhite);
+			PreviewMaterialRC->SetTextureParameterValue(FName(*RoughnessTexParamName), EmptyRoughnessMap);
+			PreviewMaterialRC->SetTextureParameterValue(FName(*MetallicTexParamName), EmptyMetallicMap);
+			PreviewMaterialRC->SetTextureParameterValue(FName(*SpecularTexParamName), EmptySpecularMap);
+			PreviewMaterialRC->SetTextureParameterValue(FName(*EmissiveTexParamName), EmptyEmissiveMap);
+			PreviewMaterialRC->SetTextureParameterValue(FName(*NormalTexParamName), EmptyNormalMap);
+		}
+	}
+
+	{
+		UMaterial* Material = LoadObject<UMaterial>(nullptr,
+			TEXT("/MeshModelingToolsetExp/Materials/BakeRenderCapturePreviewSubsurfaceMaterial"));
+		check(Material);
+		if (Material != nullptr)
+		{
+			ensure(Material->GetShadingModels().HasShadingModel(EMaterialShadingModel::MSM_Subsurface));
+			PreviewMaterialRC_Subsurface = UMaterialInstanceDynamic::Create(Material, GetToolManager());
+			PreviewMaterialRC_Subsurface->SetTextureParameterValue(FName(*BaseColorTexParamName), EmptyColorMapWhite);
+			PreviewMaterialRC_Subsurface->SetTextureParameterValue(FName(*RoughnessTexParamName), EmptyRoughnessMap);
+			PreviewMaterialRC_Subsurface->SetTextureParameterValue(FName(*MetallicTexParamName), EmptyMetallicMap);
+			PreviewMaterialRC_Subsurface->SetTextureParameterValue(FName(*SpecularTexParamName), EmptySpecularMap);
+			PreviewMaterialRC_Subsurface->SetTextureParameterValue(FName(*EmissiveTexParamName), EmptyEmissiveMap);
+			PreviewMaterialRC_Subsurface->SetTextureParameterValue(FName(*NormalTexParamName), EmptyNormalMap);
+			PreviewMaterialRC_Subsurface->SetTextureParameterValue(FName(*OpacityTexParamName), EmptyOpacityMap);
+			PreviewMaterialRC_Subsurface->SetTextureParameterValue(FName(*SubsurfaceColorTexParamName), EmptySubsurfaceColorMap);
+		}
+	}
+
+	{
+		UMaterial* Material = LoadObject<UMaterial>(nullptr,
+			TEXT("/MeshModelingToolsetExp/Materials/FullMaterialBakePreviewMaterial_PackedMRS"));
+		check(Material);
+		if (Material != nullptr)
+		{
+			PreviewMaterialPackedRC = UMaterialInstanceDynamic::Create(Material, GetToolManager());
+			PreviewMaterialPackedRC->SetTextureParameterValue(FName(*BaseColorTexParamName), EmptyColorMapWhite);
+			PreviewMaterialPackedRC->SetTextureParameterValue(FName(*PackedMRSTexParamName), EmptyPackedMRSMap);
+			PreviewMaterialPackedRC->SetTextureParameterValue(FName(*EmissiveTexParamName), EmptyEmissiveMap);
+			PreviewMaterialPackedRC->SetTextureParameterValue(FName(*NormalTexParamName), EmptyNormalMap);
+		}
+	}
+
+	{
+		UMaterial* Material = LoadObject<UMaterial>(nullptr,
+			TEXT("/MeshModelingToolsetExp/Materials/FullMaterialBakePreviewSubsurfaceMaterial_PackedMRS"));
+		check(Material);
+		if (Material != nullptr)
+		{
+			ensure(Material->GetShadingModels().HasShadingModel(EMaterialShadingModel::MSM_Subsurface));
+			PreviewMaterialPackedRC_Subsurface = UMaterialInstanceDynamic::Create(Material, GetToolManager());
+			PreviewMaterialPackedRC_Subsurface->SetTextureParameterValue(FName(*BaseColorTexParamName), EmptyColorMapWhite);
+			PreviewMaterialPackedRC_Subsurface->SetTextureParameterValue(FName(*PackedMRSTexParamName), EmptyPackedMRSMap);
+			PreviewMaterialPackedRC_Subsurface->SetTextureParameterValue(FName(*EmissiveTexParamName), EmptyEmissiveMap);
+			PreviewMaterialPackedRC_Subsurface->SetTextureParameterValue(FName(*NormalTexParamName), EmptyNormalMap);
+			PreviewMaterialPackedRC_Subsurface->SetTextureParameterValue(FName(*OpacityTexParamName), EmptyOpacityMap);
+			PreviewMaterialPackedRC_Subsurface->SetTextureParameterValue(FName(*SubsurfaceColorTexParamName), EmptySubsurfaceColorMap);
+		}
+	}
+}
+
+
+void UBakeRenderCaptureTool::InvalidateCompute()
+{
+	if (!Compute)
+	{
+		// Initialize background compute
+		Compute = MakeUnique<TGenericDataBackgroundCompute<FMeshMapBaker>>();
+		Compute->Setup(this);
+		Compute->OnResultUpdated.AddLambda([this](const TUniquePtr<FMeshMapBaker>& NewResult) { OnMapsUpdated(NewResult); });
+	}
+	Compute->InvalidateResult();
+	OpState = EBakeOpState::Clean;
+}
+
+FRenderCaptureUpdate UBakeRenderCaptureTool::UpdateSceneCapture()
+{
+	for (int Idx = 1; Idx < Targets.Num(); ++Idx)
+	{
+		UE::ToolTarget::ShowSourceObject(Targets[Idx]);
+	}
+
+	FRenderCaptureUpdate Update;
+
+	if (!SceneCapture)
+	{
+		const FRenderCaptureOptions Options = MakeRenderCaptureOptions(*RenderCaptureProperties);
+
+		constexpr bool bAllowCancel = false;
+		SceneCapture = CapturePhotoSet(Actors, Options, Update, bAllowCancel);
+	}
+	else
+	{
+		// Get already computed options from the existing SceneCapture so we can restore them if the capture is cancelled
+		const FRenderCaptureOptions ComputedOptions = GetComputedPhotoSetOptions(SceneCapture);
+
+		const FRenderCaptureOptions Options = MakeRenderCaptureOptions(*RenderCaptureProperties);
+
+		constexpr bool bAllowCancel = true;
+		Update = UpdatePhotoSets(SceneCapture, Actors, Options, bAllowCancel);
+
+		if (SceneCapture->Cancelled())
+		{
+			// Restore the settings present before the change that invoked the scene capture recompute
+			// Note UpdatePhotoSets will have restored the SceneCapture settings to their values prior to the cancel
+			RenderCaptureProperties->bBaseColorMap      = ComputedOptions.bBakeBaseColor;
+			RenderCaptureProperties->bNormalMap         = ComputedOptions.bBakeNormalMap;
+			RenderCaptureProperties->bMetallicMap       = ComputedOptions.bBakeMetallic;
+			RenderCaptureProperties->bRoughnessMap      = ComputedOptions.bBakeRoughness;
+			RenderCaptureProperties->bSpecularMap       = ComputedOptions.bBakeSpecular;
+			RenderCaptureProperties->bPackedMRSMap      = ComputedOptions.bUsePackedMRS;
+			RenderCaptureProperties->bEmissiveMap       = ComputedOptions.bBakeEmissive;
+			RenderCaptureProperties->bOpacityMap        = ComputedOptions.bBakeOpacity;
+			RenderCaptureProperties->bSubsurfaceColorMap = ComputedOptions.bBakeSubsurfaceColor;
+			RenderCaptureProperties->bDeviceDepthMap    = ComputedOptions.bBakeDeviceDepth;
+			RenderCaptureProperties->bAntiAliasing      = ComputedOptions.bAntiAliasing;
+			RenderCaptureProperties->CaptureFieldOfView = ComputedOptions.FieldOfViewDegrees;
+			RenderCaptureProperties->NearPlaneDist      = ComputedOptions.NearPlaneDist;
+			RenderCaptureProperties->Resolution         = static_cast<EBakeTextureResolution>(ComputedOptions.RenderCaptureImageSize);
+
+			// Silently make the above updates so we don't overwrite the change to OpState below and call this function again
+			RenderCaptureProperties->SilentUpdateWatched();
+			Settings->SilentUpdateWatched();
+		}
+	}
+
+	for (int Idx = 1; Idx < Targets.Num(); ++Idx)
+	{
+		UE::ToolTarget::HideSourceObject(Targets[Idx]);
+	}
+
+	return Update;
+}
+
+// Process dirty props and update background compute. Called by Render
+void UBakeRenderCaptureTool::UpdateResult()
+{
+	// Return if the bake is already launched/complete.
+	if (OpState == EBakeOpState::Clean)
+	{
+		return;
+	}
+
+	// The bake operation, Compute, stores a pointer to the SceneCapture so that must not be modified while baking
+	const bool bComputeInProgress = (Compute && (Compute->GetElapsedComputeTime() > 0.f));
+	if (bComputeInProgress)
+	{
+		return;
+	}
+
+	FText ErrorMessage; // Empty message indicates no error
+
+	{
+		const int32 TargetUVLayer = InputMeshSettings->GetTargetUVLayerIndex();
+		if (FText* Message = TargetUVLayerToError.Find(TargetUVLayer); Message)
+		{
+			ErrorMessage = *Message;
+		}
+		else
+		{
+			const auto HasDegenerateUVs = [this]
+			{
+				FDynamicMeshUVOverlay* UVOverlay = TargetMesh.Attributes()->GetUVLayer(InputMeshSettings->GetTargetUVLayerIndex());
+				FAxisAlignedBox2f Bounds = FAxisAlignedBox2f::Empty();
+				for (const int Index : UVOverlay->ElementIndicesItr())
+				{
+					FVector2f UV;
+					UVOverlay->GetElement(Index, UV);
+					Bounds.Contain(UV);
+				}
+				return Bounds.Min == Bounds.Max;
+			};
+
+			if (TargetMesh.Attributes()->GetUVLayer(InputMeshSettings->GetTargetUVLayerIndex()) == nullptr)
+			{
+				ErrorMessage = LOCTEXT("TargetMeshMissingUVs", "The Target Mesh UV layer is missing");
+			}
+			else if (HasDegenerateUVs())
+			{
+				ErrorMessage = LOCTEXT("TargetMeshDegenerateUVs", "The Target Mesh UV layer is degenerate");
+			}
+			else
+			{
+				ErrorMessage = FText(); // No error
+			}
+			TargetUVLayerToError.Add(TargetUVLayer, ErrorMessage);
+		}
+
+		// If there are no UV layer errors check for missing tangent space error
+		if (ErrorMessage.IsEmpty() && RenderCaptureProperties->bNormalMap && ValidTargetMeshTangents() == false)
+		{
+			ErrorMessage = LOCTEXT("TargetMeshMissingTangentSpace", "The Target Mesh is missing a tangent space. Disable Normal Map capture to continue.");
+		}
+	}
+
+	// Calling DisplayMessage with an empty string will clear existing messages
+	GetToolManager()->DisplayMessage(ErrorMessage, EToolMessageLevel::UserWarning);
+
+	const bool bIsInvalid = (ErrorMessage.IsEmpty() == false);
+	if (bIsInvalid)
+	{
+		// Clear all results and wait for user to fix the invalid tool inputs
+		InvalidateResults(FRenderCaptureUpdate{});
+
+		// Only call UpdateVisualization when we first detect the invalid inputs
+		const bool bWasValid = static_cast<bool>(OpState & EBakeOpState::Invalid) == false;
+		if (bWasValid)
+		{
+			UpdateVisualization();
+		}
+
+		// Set an invalid op state so we re-enter this function until the inputs are valid
+		OpState = EBakeOpState::Invalid;
+	}
+	else
+	{
+		FRenderCaptureUpdate Update = UpdateSceneCapture();
+
+		const bool bInvalidateAll =
+			ComputedTextureSize != Settings->TextureSize ||
+			ComputedSamplesPerPixel != Settings->SamplesPerPixel ||
+			ComputedValidDepthThreshold != Settings->ValidSampleDepthThreshold;
+
+		// Invalidate any results which need re-baking because the SceneCapture was updated, or
+		// Invalidate all results if the any baking parameters were changed
+		InvalidateResults(bInvalidateAll ? FRenderCaptureUpdate{} : Update);
+
+		// Update the preview mesh material with the (possibly invalidated) bake results
+		UpdateVisualization();
+
+		// Start another bake operation
+		InvalidateCompute();
+
+		// Cache computed parameters which are used to determine if results need re-baking
+		ComputedTextureSize = Settings->TextureSize;
+		ComputedSamplesPerPixel = Settings->SamplesPerPixel;
+		ComputedValidDepthThreshold = Settings->ValidSampleDepthThreshold;
+	}
+}
+
+
+
+void UBakeRenderCaptureTool::UpdateVisualization()
+{
+	if (Settings->MapPreview.IsEmpty())
+	{
+		return;
+	}
+
+	const bool bSubsurfaceMaterial = ResultSettings->SubsurfaceColorMap || ResultSettings->OpacityMap;
+	const bool bPackedMRS = ResultSettings->PackedMRSMap != nullptr;
+
+	// Choose the material
+	TObjectPtr<UMaterialInstanceDynamic> Material = bPackedMRS ? PreviewMaterialPackedRC : PreviewMaterialRC;
+	if (bSubsurfaceMaterial)
+	{
+		Material = bPackedMRS ? PreviewMaterialPackedRC_Subsurface : PreviewMaterialRC_Subsurface;
+		ensure(Material->GetShadingModels().HasShadingModel(EMaterialShadingModel::MSM_Subsurface));
+		ensure(Material->GetBlendMode() == EBlendMode::BLEND_Masked);
+	}
+
+	if (VisualizationProps->bPreviewAsMaterial)
+	{
+		const auto TrySetTexture =
+			[Material](const FString& TextureName, TObjectPtr<UTexture2D> Texture, TObjectPtr<UTexture2D> Fallback, bool bMaterialHasTexture)
+		{
+			if (bMaterialHasTexture)
+			{
+				Material->SetTextureParameterValue(FName(TextureName), Texture ? Texture : Fallback);
+			}
+		};
+
+		// Set all computed textures or fallback to the empty texture map
+		TrySetTexture(BaseColorTexParamName, ResultSettings->BaseColorMap, EmptyColorMapWhite, true);
+		TrySetTexture(EmissiveTexParamName,  ResultSettings->EmissiveMap,  EmptyEmissiveMap,   true);
+		TrySetTexture(NormalTexParamName,    ResultSettings->NormalMap,    EmptyNormalMap,     true);
+		TrySetTexture(PackedMRSTexParamName, ResultSettings->PackedMRSMap, EmptyPackedMRSMap,  bPackedMRS);
+		TrySetTexture(RoughnessTexParamName, ResultSettings->RoughnessMap, EmptyRoughnessMap, !bPackedMRS);
+		TrySetTexture(MetallicTexParamName,  ResultSettings->MetallicMap,  EmptyMetallicMap,  !bPackedMRS);
+		TrySetTexture(SpecularTexParamName,  ResultSettings->SpecularMap,  EmptySpecularMap,  !bPackedMRS);
+		TrySetTexture(OpacityTexParamName,   ResultSettings->OpacityMap,   EmptyOpacityMap,    bSubsurfaceMaterial);
+		TrySetTexture(SubsurfaceColorTexParamName, ResultSettings->SubsurfaceColorMap, EmptySubsurfaceColorMap, bSubsurfaceMaterial);
+	}
+	else
+	{
+		const auto TrySetTexture =
+			[Material, this](const FString& TextureName, TObjectPtr<UTexture2D> Texture, TObjectPtr<UTexture2D> Fallback, bool bMaterialHasTexture)
+		{
+			// Set the BaseColor texture to the MapPreview texture if it exists and use white otherwise
+			if (TextureName == Settings->MapPreview)
+			{
+				if (bMaterialHasTexture && Texture)
+				{
+					Material->SetTextureParameterValue(FName(BaseColorTexParamName), Texture);
+				}
+				else
+				{
+					Material->SetTextureParameterValue(FName(BaseColorTexParamName), EmptyColorMapWhite);
+				}
+			}
+
+			// Set the non-BaseColor texture parameters to empty fallback textures
+			if (TextureName != BaseColorTexParamName)
+			{
+				if (bMaterialHasTexture)
+				{
+					Material->SetTextureParameterValue(FName(TextureName), Fallback);
+				}
+			}
+		};
+
+		TrySetTexture(BaseColorTexParamName, ResultSettings->BaseColorMap, EmptyColorMapWhite, true);
+		TrySetTexture(EmissiveTexParamName,  ResultSettings->EmissiveMap,  EmptyEmissiveMap,   true);
+		TrySetTexture(NormalTexParamName,    ResultSettings->NormalMap,    EmptyNormalMap,     true);
+		TrySetTexture(PackedMRSTexParamName, ResultSettings->PackedMRSMap, EmptyPackedMRSMap,  bPackedMRS);
+		TrySetTexture(RoughnessTexParamName, ResultSettings->RoughnessMap, EmptyRoughnessMap, !bPackedMRS);
+		TrySetTexture(MetallicTexParamName,  ResultSettings->MetallicMap,  EmptyMetallicMap,  !bPackedMRS);
+		TrySetTexture(SpecularTexParamName,  ResultSettings->SpecularMap,  EmptySpecularMap,  !bPackedMRS);
+		TrySetTexture(OpacityTexParamName,   ResultSettings->OpacityMap,   EmptyOpacityMap,    bSubsurfaceMaterial);
+		TrySetTexture(SubsurfaceColorTexParamName, ResultSettings->SubsurfaceColorMap, EmptySubsurfaceColorMap, bSubsurfaceMaterial);
+	}
+
+	Material->SetScalarParameterValue(TEXT("UVChannel"), InputMeshSettings->GetTargetUVLayerIndex());
+	PreviewMesh->SetOverrideRenderMaterial(Material);
+
+	GetToolManager()->PostInvalidation();
+}
+
+
+void UBakeRenderCaptureTool::InvalidateResults(FRenderCaptureUpdate Update)
+{
+	// Note that the bake operation, Compute, updates ResultSettings when results are available via the
+	// Compute->OnResultUpdated delegate.
+
+	if (Update.bUpdatedBaseColor)
+	{
+		ResultSettings->BaseColorMap = nullptr;
+	}
+	if (Update.bUpdatedRoughness)
+	{
+		ResultSettings->RoughnessMap = nullptr;
+	}
+	if (Update.bUpdatedMetallic)
+	{
+		ResultSettings->MetallicMap = nullptr;
+	}
+	if (Update.bUpdatedSpecular)
+	{
+		ResultSettings->SpecularMap = nullptr;
+	}
+	if (Update.bUpdatedEmissive)
+	{
+		ResultSettings->EmissiveMap = nullptr;
+	}
+	if (Update.bUpdatedNormalMap)
+	{
+		ResultSettings->NormalMap = nullptr;
+	}
+	if (Update.bUpdatedOpacity)
+	{
+		ResultSettings->OpacityMap = nullptr;
+	}
+	if (Update.bUpdatedSubsurfaceColor)
+	{
+		ResultSettings->SubsurfaceColorMap = nullptr;
+	}
+	if (Update.bUpdatedPackedMRS)
+	{
+		ResultSettings->PackedMRSMap = nullptr;
+	}
+}
+
+
+
+void UBakeRenderCaptureTool::RecordAnalytics() const
+{
+	if (FEngineAnalytics::IsAvailable() == false)
+	{
+		return;
+	}
+
+	TArray<FAnalyticsEventAttribute> Attributes;
+
+	// General
+	Attributes.Add(FAnalyticsEventAttribute(TEXT("Bake.Duration.Total.Seconds"), BakeAnalytics.TotalBakeDuration));
+	Attributes.Add(FAnalyticsEventAttribute(TEXT("Bake.Duration.WriteToImage.Seconds"), BakeAnalytics.WriteToImageDuration));
+	Attributes.Add(FAnalyticsEventAttribute(TEXT("Bake.Duration.WriteToGutter.Seconds"), BakeAnalytics.WriteToGutterDuration));
+	Attributes.Add(FAnalyticsEventAttribute(TEXT("Bake.Stats.NumSamplePixels"), BakeAnalytics.NumSamplePixels));
+	Attributes.Add(FAnalyticsEventAttribute(TEXT("Bake.Stats.NumGutterPixels"), BakeAnalytics.NumGutterPixels));
+
+	// Input mesh data
+	Attributes.Add(FAnalyticsEventAttribute(TEXT("Input.TargetMesh.NumTriangles"), BakeAnalytics.MeshSettings.NumTargetMeshTris));
+	Attributes.Add(FAnalyticsEventAttribute(TEXT("Input.RenderCapture.NumMeshes"), BakeAnalytics.MeshSettings.NumDetailMesh));
+	Attributes.Add(FAnalyticsEventAttribute(TEXT("Input.RenderCapture.NumTriangles"), BakeAnalytics.MeshSettings.NumDetailMeshTris));
+
+	// Bake settings
+	Attributes.Add(FAnalyticsEventAttribute(TEXT("Settings.Image.Width"), static_cast<int32>(Settings->TextureSize)));
+	Attributes.Add(FAnalyticsEventAttribute(TEXT("Settings.Image.Height"), static_cast<int32>(Settings->TextureSize)));
+	Attributes.Add(FAnalyticsEventAttribute(TEXT("Settings.SamplesPerPixel"), static_cast<int32>(Settings->SamplesPerPixel)));
+	Attributes.Add(FAnalyticsEventAttribute(TEXT("Settings.TargetUVLayer"), InputMeshSettings->GetTargetUVLayerIndex()));
+
+	// Render Capture settings
+	Attributes.Add(FAnalyticsEventAttribute(TEXT("Settings.RenderCapture.Image.Width"), static_cast<int32>(RenderCaptureProperties->Resolution)));
+	Attributes.Add(FAnalyticsEventAttribute(TEXT("Settings.RenderCapture.Image.Height"), static_cast<int32>(RenderCaptureProperties->Resolution)));
+	Attributes.Add(FAnalyticsEventAttribute(TEXT("Settings.RenderCapture.BaseColorMap.Enabled"), RenderCaptureProperties->bBaseColorMap));
+	Attributes.Add(FAnalyticsEventAttribute(TEXT("Settings.RenderCapture.NormalMap.Enabled"), RenderCaptureProperties->bNormalMap));
+	Attributes.Add(FAnalyticsEventAttribute(TEXT("Settings.RenderCapture.MetallicMap.Enabled"), RenderCaptureProperties->bMetallicMap));
+	Attributes.Add(FAnalyticsEventAttribute(TEXT("Settings.RenderCapture.RoughnessMap.Enabled"), RenderCaptureProperties->bRoughnessMap));
+	Attributes.Add(FAnalyticsEventAttribute(TEXT("Settings.RenderCapture.SpecularMap.Enabled"), RenderCaptureProperties->bSpecularMap));
+	Attributes.Add(FAnalyticsEventAttribute(TEXT("Settings.RenderCapture.PackedMRSMap.Enabled"), RenderCaptureProperties->bPackedMRSMap));
+	Attributes.Add(FAnalyticsEventAttribute(TEXT("Settings.RenderCapture.EmissiveMap.Enabled"), RenderCaptureProperties->bEmissiveMap));
+	Attributes.Add(FAnalyticsEventAttribute(TEXT("Settings.RenderCapture.OpacityMap.Enabled"), RenderCaptureProperties->bOpacityMap));
+	Attributes.Add(FAnalyticsEventAttribute(TEXT("Settings.RenderCapture.SubsurfaceColorMap.Enabled"), RenderCaptureProperties->bSubsurfaceColorMap));
+	Attributes.Add(FAnalyticsEventAttribute(TEXT("Settings.RenderCapture.CaptureFieldOfView"), RenderCaptureProperties->CaptureFieldOfView));
+	Attributes.Add(FAnalyticsEventAttribute(TEXT("Settings.RenderCapture.NearPlaneDistance"), RenderCaptureProperties->NearPlaneDist));
+
+	FEngineAnalytics::GetProvider().RecordEvent(FString(TEXT("Editor.Usage.MeshModelingMode.")) + GetAnalyticsEventName(), Attributes);
+
+	constexpr bool bDebugLogAnalytics = false;
+	if constexpr (bDebugLogAnalytics)
+	{
+		for (const FAnalyticsEventAttribute& Attr : Attributes)
+		{
+			UE_LOG(LogGeometry, Log, TEXT("[%s] %s = %s"), *GetAnalyticsEventName(), *Attr.GetName(), *Attr.GetValue());
+		}
+	}
+}
+
+
+void UBakeRenderCaptureTool::GatherAnalytics(const FMeshMapBaker& Result)
+{
+	if (!FEngineAnalytics::IsAvailable())
+	{
+		return;
+	}
+
+	BakeAnalytics.TotalBakeDuration = Result.BakeAnalytics.TotalBakeDuration;
+	BakeAnalytics.WriteToImageDuration = Result.BakeAnalytics.WriteToImageDuration;
+	BakeAnalytics.WriteToGutterDuration = Result.BakeAnalytics.WriteToGutterDuration;
+	BakeAnalytics.NumSamplePixels = Result.BakeAnalytics.NumSamplePixels;
+	BakeAnalytics.NumGutterPixels = Result.BakeAnalytics.NumGutterPixels;
+}
+
+
+void UBakeRenderCaptureTool::GatherAnalytics(FBakeAnalytics::FMeshSettings& Data)
+{
+	if (FEngineAnalytics::IsAvailable() == false)
+	{
+		return;
+	}
+
+	Data.NumTargetMeshTris = TargetMesh.TriangleCount();
+	Data.NumDetailMesh = Actors.Num();
+	Data.NumDetailMeshTris = 0;
+	for (AActor* Actor : Actors)
+	{
+		check(Actor != nullptr);
+		TInlineComponentArray<UPrimitiveComponent*> PrimitiveComponents;
+		Actor->GetComponents(PrimitiveComponents);
+		for (UPrimitiveComponent* PrimitiveComponent : PrimitiveComponents)
+		{
+			if (UStaticMeshComponent* StaticMeshComponent = Cast<UStaticMeshComponent>(PrimitiveComponent))
+			{
+				if (StaticMeshComponent->GetStaticMesh() != nullptr)
+				{
+					// TODO We could also check GetNumNaniteTriangles here and use the maximum
+					Data.NumDetailMeshTris += StaticMeshComponent->GetStaticMesh()->GetNumTriangles(0);
+				}
+			} 
+		}
+	}
+}
+
+#undef LOCTEXT_NAMESPACE
+
diff --git a/Engine/Plugins/Experimental/MeshModelingToolsetExp/Source/MeshModelingToolsEditorOnlyExp/Public/BakeRenderCaptureTool.h b/Engine/Plugins/Experimental/MeshModelingToolsetExp/Source/MeshModelingToolsEditorOnlyExp/Public/BakeRenderCaptureTool.h
new file mode 100644
--- /dev/null
+++ b/Engine/Plugins/Experimental/MeshModelingToolsetExp/Source/MeshModelingToolsEditorOnlyExp/Public/BakeRenderCaptureTool.h
@@ -0,0 +1,453 @@
+﻿// Copyright Epic Games, Inc. All Rights Reserved.
+
+#pragma once
+
+#include "CoreMinimal.h"
+#include "InteractiveToolBuilder.h"
+#include "ModelingOperators.h"
+#include "MeshOpPreviewHelpers.h"
+#include "BaseTools/MultiSelectionMeshEditingTool.h"
+#include "BakeMeshAttributeToolCommon.h"
+#include "Baking/BakingTypes.h"
+#include "Baking/RenderCaptureFunctions.h"
+
+// Render Capture algorithm includes
+#include "Scene/SceneCapturePhotoSet.h"
+
+#include "BakeRenderCaptureTool.generated.h"
+
+class UTexture2D;
+
+namespace UE
+{
+namespace Geometry
+{
+
+class FSceneCapturePhotoSet;
+
+}
+}
+
+//
+// Tool Result
+//
+
+
+UCLASS()
+class MESHMODELINGTOOLSEDITORONLYEXP_API UBakeRenderCaptureResults  : public UInteractiveToolPropertySet
+{
+	GENERATED_BODY()
+
+public:
+	UPROPERTY(VisibleAnywhere, Category = Results, meta = (TransientToolProperty))
+	TObjectPtr<UTexture2D> BaseColorMap;
+
+	/** World space normal map */
+	UPROPERTY(VisibleAnywhere, Category = Results, meta = (TransientToolProperty))
+	TObjectPtr<UTexture2D> NormalMap;
+
+	/** Packed Metallic/Roughness/Specular Map */
+	UPROPERTY(VisibleAnywhere, Category = Results, meta = (TransientToolProperty, DisplayName = "Packed MRS Map"))
+	TObjectPtr<UTexture2D> PackedMRSMap;
+
+	UPROPERTY(VisibleAnywhere, Category = Results, meta = (TransientToolProperty))
+	TObjectPtr<UTexture2D> MetallicMap;
+
+	UPROPERTY(VisibleAnywhere, Category = Results, meta = (TransientToolProperty))
+	TObjectPtr<UTexture2D> RoughnessMap;
+
+	UPROPERTY(VisibleAnywhere, Category = Results, meta = (TransientToolProperty))
+	TObjectPtr<UTexture2D> SpecularMap;
+
+	UPROPERTY(VisibleAnywhere, Category = Results, meta = (TransientToolProperty))
+	TObjectPtr<UTexture2D> EmissiveMap;
+
+	UPROPERTY(VisibleAnywhere, Category = Results, meta = (TransientToolProperty))
+	TObjectPtr<UTexture2D> OpacityMap;
+
+	UPROPERTY(VisibleAnywhere, Category = Results, meta = (TransientToolProperty))
+	TObjectPtr<UTexture2D> SubsurfaceColorMap;
+};
+
+
+//
+// Tool Builder
+//
+
+
+
+UCLASS()
+class MESHMODELINGTOOLSEDITORONLYEXP_API UBakeRenderCaptureToolBuilder : public UMultiSelectionMeshEditingToolBuilder
+{
+	GENERATED_BODY()
+
+public:
+	virtual bool CanBuildTool(const FToolBuilderState& SceneState) const override;
+	virtual UMultiSelectionMeshEditingTool* CreateNewTool(const FToolBuilderState& SceneState) const override;
+
+protected:
+	virtual const FToolTargetTypeRequirements& GetTargetRequirements() const override;
+};
+
+
+
+//
+// Tool Properties
+//
+
+
+UCLASS()
+class MESHMODELINGTOOLSEDITORONLYEXP_API URenderCaptureProperties : public UInteractiveToolPropertySet
+{
+	GENERATED_BODY()
+
+public:
+	UPROPERTY(Category = RenderCaptureOptions, EditAnywhere, meta = (NoResetToDefault, ClampMin = "1", UIMin= "1"), DisplayName="Render Capture Resolution")
+	EBakeTextureResolution Resolution = EBakeTextureResolution::Resolution512;
+
+	// Whether to generate a texture for the Base Color property
+	UPROPERTY(Category = RenderCaptureOptions, EditAnywhere, meta = (NoResetToDefault))
+	bool bBaseColorMap = false;
+
+	// Whether to generate a texture for the World Normal property
+	UPROPERTY(Category = RenderCaptureOptions, EditAnywhere, meta = (NoResetToDefault))
+	bool bNormalMap = false;
+
+	// Whether to generate a packed texture with Metallic, Roughness and Specular properties
+	UPROPERTY(Category = RenderCaptureOptions, EditAnywhere, meta = (NoResetToDefault), DisplayName="Packed MRS Map")
+	bool bPackedMRSMap = false;
+
+	// Whether to generate a texture for the Metallic property
+	UPROPERTY(Category = RenderCaptureOptions, EditAnywhere, meta = (NoResetToDefault, EditCondition="bPackedMRSMap == false"))
+	bool bMetallicMap = false;
+
+	// Whether to generate a texture for the Roughness property
+	UPROPERTY(Category = RenderCaptureOptions, EditAnywhere, meta = (NoResetToDefault, EditCondition="bPackedMRSMap == false"))
+	bool bRoughnessMap = false;
+
+	// Whether to generate a texture for the Specular property
+	UPROPERTY(Category = RenderCaptureOptions, EditAnywhere, meta = (NoResetToDefault, EditCondition="bPackedMRSMap == false"))
+	bool bSpecularMap = false;
+	
+	// Whether to generate a texture for the Emissive property
+	UPROPERTY(Category = RenderCaptureOptions, EditAnywhere, meta = (NoResetToDefault))
+	bool bEmissiveMap = false;
+
+	// Whether to generate a texture for the Opacity property
+	UPROPERTY(Category = RenderCaptureOptions, EditAnywhere, meta = (NoResetToDefault))
+	bool bOpacityMap = false;
+
+	// Whether to generate a texture for the SubsurfaceColor property
+	UPROPERTY(Category = RenderCaptureOptions, EditAnywhere, meta = (NoResetToDefault))
+	bool bSubsurfaceColorMap = false;
+
+	// Whether to use anti-aliasing in the render captures, this may introduce artefacts if pixels at different scene depths get blended
+	UPROPERTY(Category = RenderCaptureOptions, EditAnywhere, AdvancedDisplay)
+	bool bAntiAliasing = false;
+
+	// These are hidden in the UI right now, we might want to expose them if they turn out to be useful for very large
+	// or very small objects (not tested yet) TODO Figure out if we want to expose these options
+	
+	UPROPERTY(meta = (ClampMin = "5.0", ClampMax = "160.0"))
+	float CaptureFieldOfView = 30.0f;
+
+	UPROPERTY(meta = (ClampMin = "0.001", ClampMax = "1000.0"))
+	float NearPlaneDist = 1.0f;
+	
+	UPROPERTY()
+	bool bDeviceDepthMap = false;
+
+	bool operator==(const URenderCaptureProperties& Other) const
+	{
+		return Resolution == Other.Resolution
+			&& bBaseColorMap == Other.bBaseColorMap
+			&& bNormalMap == Other.bNormalMap
+			&& bMetallicMap == Other.bMetallicMap
+			&& bRoughnessMap == Other.bRoughnessMap
+			&& bSpecularMap == Other.bSpecularMap
+			&& bPackedMRSMap == Other.bPackedMRSMap
+			&& bEmissiveMap == Other.bEmissiveMap
+			&& bOpacityMap == Other.bOpacityMap
+			&& bSubsurfaceColorMap == Other.bSubsurfaceColorMap
+			&& bAntiAliasing == Other.bAntiAliasing
+			&& CaptureFieldOfView == Other.CaptureFieldOfView
+			&& NearPlaneDist == Other.NearPlaneDist
+			&& bDeviceDepthMap == Other.bDeviceDepthMap;
+	}
+
+	bool operator != (const URenderCaptureProperties& Other) const
+	{
+		return !(*this == Other);
+	}
+};
+
+UCLASS()
+class MESHMODELINGTOOLSEDITORONLYEXP_API UBakeRenderCaptureToolProperties : public UInteractiveToolPropertySet
+{
+	GENERATED_BODY()
+
+public:
+	
+	/** The map type to preview */
+	UPROPERTY(EditAnywhere, Category = BakeOutput, meta = (DisplayName="Preview Output Type", TransientToolProperty, GetOptions = GetMapPreviewNamesFunc))
+	FString MapPreview;
+	
+	UFUNCTION()
+	const TArray<FString>& GetMapPreviewNamesFunc()
+	{
+		return MapPreviewNamesList;
+	}
+	UPROPERTY(meta = (TransientToolProperty))
+	TArray<FString> MapPreviewNamesList;
+
+	/** Number of samples per pixel */
+	UPROPERTY(EditAnywhere, Category = BakeOutput)
+	EBakeTextureSamplesPerPixel SamplesPerPixel = EBakeTextureSamplesPerPixel::Sample1;
+
+	/* Size of generated textures */
+	UPROPERTY(EditAnywhere, Category = BakeOutput, meta = (ClampMin="1", UIMin="1"), DisplayName="Output Resolution")
+	EBakeTextureResolution TextureSize = EBakeTextureResolution::Resolution512;
+
+	/**
+	 * This threshold is used to detect occlusion artefacts (e.g., silhouettes/blotches in the base color) so that texels
+	 * exhibiting them can be cleaned up i.e., rejected and filled in using the values of nearby valid texels instead.
+	 * 
+	 * If the threshold is zero, the cleanup step is skipped.
+	 * If the threshold is too large, texels with artefacts won't be detected and the cleanup step is effectively skipped.
+	 * If the threshold is too small, texels without artefacts get detected and results will be bad because there won't
+	 * be enough nearby valid texels from which to infill values.
+	 * 
+	 * A good starting point is to choose a threshold around the size of the distance (in world space) between the
+	 * target and source meshes; if the target was generated by VoxWrap then the voxel size estimates this distance.
+	 */
+	UPROPERTY(EditAnywhere, Category = BakeOutput, meta = (ClampMin="0", UIMin="0"), DisplayName="Cleanup Threshold")
+	float ValidSampleDepthThreshold = 0.f;
+};
+
+
+UCLASS()
+class MESHMODELINGTOOLSEDITORONLYEXP_API UBakeRenderCaptureInputToolProperties : public UInteractiveToolPropertySet
+{
+	GENERATED_BODY()
+
+public:
+	/** Target mesh to sample to */
+	UPROPERTY(VisibleAnywhere, Category = BakeInput, DisplayName = "Target Mesh", meta = (TransientToolProperty, NoResetToDefault))
+	TObjectPtr<UStaticMesh> TargetStaticMesh = nullptr;
+
+	/** UV channel to use for the target mesh */
+	UPROPERTY(EditAnywhere, Category = BakeInput, meta = (DisplayName = "Target Mesh UV Channel", GetOptions = GetTargetUVLayerNamesFunc, NoResetToDefault))
+	FString TargetUVLayer;
+
+	UFUNCTION()
+	int32 GetTargetUVLayerIndex() const
+	{
+		return TargetUVLayerNamesList.IndexOfByKey(TargetUVLayer);
+	}
+
+	UFUNCTION()
+	const TArray<FString>& GetTargetUVLayerNamesFunc() const
+	{
+		return TargetUVLayerNamesList;
+	}
+
+	UPROPERTY(meta = (TransientToolProperty))
+	TArray<FString> TargetUVLayerNamesList;
+};
+
+UCLASS()
+class MESHMODELINGTOOLSEDITORONLYEXP_API UBakeRenderCaptureVisualizationProperties : public UInteractiveToolPropertySet
+{
+	GENERATED_BODY()
+public:
+	/**
+	 * If true  preview results by connecting them to corresponding material inputs
+	 * If false connect the selected preview output as the Base Color input and use empty maps for other material inputs
+	 */
+	UPROPERTY(EditAnywhere, Category = Preview)
+	bool bPreviewAsMaterial = false;
+
+	/** Adjust the brightness of the Base Color material input; does not affect results stored in textures */
+	UPROPERTY(EditAnywhere, Category = Preview, meta = (UIMin = "0.0", UIMax = "1.0"))
+	float Brightness = 1.0f;
+
+	/** Adjust the brightness of the Subsurface Color material input; does not affect results stored in textures */
+	UPROPERTY(EditAnywhere, Category = Preview, meta = (UIMin = "0.0", UIMax = "1.0", DisplayName = "SS Brightness"))
+	float SSBrightness = 1.0f;
+
+	/** Adjust the brightness of the Emissive Color material input; does not affect results stored in textures */
+	UPROPERTY(EditAnywhere, Category = Preview, meta = (UIMin = "0.0", UIMax = "1.0"))
+	float EmissiveScale = 1.0f;
+};
+
+
+//
+// Tool
+//
+
+
+
+UCLASS()
+class MESHMODELINGTOOLSEDITORONLYEXP_API UBakeRenderCaptureTool :
+	public UMultiSelectionMeshEditingTool,
+	public UE::Geometry::IGenericDataOperatorFactory<UE::Geometry::FMeshMapBaker>,
+	public IInteractiveToolExclusiveToolAPI
+{
+	GENERATED_BODY()
+
+public:
+	UBakeRenderCaptureTool() = default;
+
+	// Begin UMultiSelectionMeshEditing < UMultiSelectionTool < UInteractiveTool interface
+	virtual void Setup() override;
+	virtual void Render(IToolsContextRenderAPI* RenderAPI) override;
+	virtual void OnTick(float DeltaTime) override;
+	virtual void OnShutdown(EToolShutdownType ShutdownType) override;
+	virtual bool HasCancel() const override { return true; }
+	virtual bool HasAccept() const override { return true; }
+	virtual bool CanAccept() const override;
+	// End UMultiSelectionMeshEditing < UMultiSelectionTool < UInteractiveTool interface
+
+	// Begin IGenericDataOperatorFactory interface
+	virtual TUniquePtr<UE::Geometry::TGenericDataOperator<UE::Geometry::FMeshMapBaker>> MakeNewOperator() override;
+	// End IGenericDataOperatorFactory interface
+
+protected:
+
+	UPROPERTY()
+	TArray<TObjectPtr<AActor>> Actors;
+
+	UPROPERTY()
+	TObjectPtr<UBakeRenderCaptureToolProperties> Settings;
+
+	UPROPERTY()
+	TObjectPtr<URenderCaptureProperties> RenderCaptureProperties;
+
+	UPROPERTY()
+	TObjectPtr<UBakeRenderCaptureInputToolProperties> InputMeshSettings;
+
+	UPROPERTY()
+	TObjectPtr<UBakeRenderCaptureVisualizationProperties> VisualizationProps;
+
+	// The computed textures are displayed in the details panel and used in the preview material, they are written
+	// out to assest on shutdown.
+	UPROPERTY()
+	TObjectPtr<UBakeRenderCaptureResults> ResultSettings;
+	
+protected:
+
+	void UpdateResult();
+	void UpdateVisualization();
+	void InvalidateResults(UE::Geometry::FRenderCaptureUpdate Update);
+	void InvalidateCompute();
+	void OnMapsUpdated(const TUniquePtr<UE::Geometry::FMeshMapBaker>& NewResult);
+
+	/**
+	 * Compute validity of the Target Mesh tangents. Only checks validity
+	 * once and then caches the result for successive calls.
+	 * 
+	 * @return true if the TargetMesh tangents are valid
+	 */
+	bool ValidTargetMeshTangents();
+	bool bCheckTargetMeshTangents = true;
+	bool bValidTargetMeshTangents = false;
+
+	/**
+	 * Create texture and material assets from our result map of Texture2D
+	 * @param SourceWorld the source world to define where the assets will be stored.
+	 * @param SourceAsset if not null, result textures/material will be stored adjacent to this asset.
+	 */
+	void CreateAssets(UWorld* SourceWorld, UObject* SourceAsset);
+
+	// The baking background compute operation
+	TUniquePtr<TGenericDataBackgroundCompute<UE::Geometry::FMeshMapBaker>> Compute = nullptr;
+	EBakeOpState OpState = EBakeOpState::Evaluate;
+
+	UE::Geometry::FDynamicMesh3 TargetMesh;
+	UE::Geometry::FDynamicMeshAABBTree3 TargetMeshSpatial;
+	TSharedPtr<TArray<int32>, ESPMode::ThreadSafe> TargetMeshUVCharts;
+	TSharedPtr<UE::Geometry::FMeshTangentsd, ESPMode::ThreadSafe> TargetMeshTangents;
+
+	// Empty maps are shown when nothing is computed
+	UPROPERTY()
+	TObjectPtr<UTexture2D> EmptyNormalMap;
+	UPROPERTY()
+	TObjectPtr<UTexture2D> EmptyColorMapBlack;
+	UPROPERTY()
+	TObjectPtr<UTexture2D> EmptyColorMapWhite;
+	UPROPERTY()
+	TObjectPtr<UTexture2D> EmptyEmissiveMap;
+	UPROPERTY()
+	TObjectPtr<UTexture2D> EmptyOpacityMap;
+	UPROPERTY()
+	TObjectPtr<UTexture2D> EmptySubsurfaceColorMap;
+	UPROPERTY()
+	TObjectPtr<UTexture2D> EmptyPackedMRSMap;
+	UPROPERTY()
+	TObjectPtr<UTexture2D> EmptyRoughnessMap;
+	UPROPERTY()
+	TObjectPtr<UTexture2D> EmptyMetallicMap;
+	UPROPERTY()
+	TObjectPtr<UTexture2D> EmptySpecularMap;
+
+	float SecondsBeforeWorkingMaterial = 0.75;
+
+	UPROPERTY()
+	TObjectPtr<UMaterialInstanceDynamic> WorkingPreviewMaterial;
+	UPROPERTY()
+	TObjectPtr<UMaterialInstanceDynamic> ErrorPreviewMaterial;
+	UPROPERTY()
+	TObjectPtr<UMaterialInstanceDynamic> PreviewMaterialRC;
+	UPROPERTY()
+	TObjectPtr<UMaterialInstanceDynamic> PreviewMaterialPackedRC;
+	UPROPERTY()
+	TObjectPtr<UMaterialInstanceDynamic> PreviewMaterialRC_Subsurface;
+	UPROPERTY()
+	TObjectPtr<UMaterialInstanceDynamic> PreviewMaterialPackedRC_Subsurface;
+
+	void InitializePreviewMaterials();
+
+	UPROPERTY()
+	TObjectPtr<UPreviewMesh> PreviewMesh;
+
+	// Note: We need to compute this on the game thread because the implementation has checks for this
+	TUniquePtr<UE::Geometry::FSceneCapturePhotoSet> SceneCapture = nullptr;
+
+	// These are used to determine if we need to re-bake results
+	float ComputedValidDepthThreshold;
+	EBakeTextureSamplesPerPixel ComputedSamplesPerPixel = EBakeTextureSamplesPerPixel::Sample1;
+	EBakeTextureResolution ComputedTextureSize = EBakeTextureResolution::Resolution512;
+
+	TMap<int, FText> TargetUVLayerToError;
+
+	//
+	// Analytics
+	//
+	
+	struct FBakeAnalytics
+	{
+		double TotalBakeDuration = 0.0;
+		double WriteToImageDuration = 0.0;
+		double WriteToGutterDuration = 0.0;
+		int64 NumSamplePixels = 0;
+		int64 NumGutterPixels = 0;
+
+		struct FMeshSettings
+		{
+			int32 NumTargetMeshTris = 0;
+			int32 NumDetailMesh = 0;
+			int64 NumDetailMeshTris = 0;
+		};
+		FMeshSettings MeshSettings;
+	};
+
+	FBakeAnalytics BakeAnalytics;
+
+	FString GetAnalyticsEventName() const { return TEXT("BakeRC"); }
+	void GatherAnalytics(FBakeAnalytics::FMeshSettings& Data);
+	void GatherAnalytics(const UE::Geometry::FMeshMapBaker& Result);
+	void RecordAnalytics() const;
+
+private:
+	UE::Geometry::FRenderCaptureUpdate UpdateSceneCapture();
+};
\ No newline at end of file
diff --git a/Engine/Plugins/Experimental/MeshModelingToolsetExp/Source/MeshModelingToolsExp/Private/BakeRenderCaptureTool.cpp b/Engine/Plugins/Experimental/MeshModelingToolsetExp/Source/MeshModelingToolsExp/Private/BakeRenderCaptureTool.cpp
deleted file mode 100644
diff --git a/Engine/Plugins/Experimental/MeshModelingToolsetExp/Source/MeshModelingToolsExp/Public/BakeRenderCaptureTool.h b/Engine/Plugins/Experimental/MeshModelingToolsetExp/Source/MeshModelingToolsExp/Public/BakeRenderCaptureTool.h
deleted file mode 100644
