CombineMeshInstances: add support for alternative to Voxel LOD, based on planar sweep.

Add ECoarseApproximationStrategy to IGeometryProcessing_CombineMeshInstances, and FOptions.CoarseLODStrategy, as well as .CoarseApproximationDetailSize  (previously was hardcoded to 10 in impelmentation). Rename 'Voxel' to 'Coarse' in the Options struct.
Improve ComputeSweptSolidApproximation in CombineMeshInstancesImpl - 'Square' offset was a chamfer, but generally we want a Miter, with Miter cut threshold set to diagonal length.
Add ComputeBestFullProjectionMesh in CombineMeshInstancesImpl, computes swept-solid approximation along all 3 axes and selects the one w/ lowest geometric deviation
BuildCombinedMesh now supports using SweptPlanarProjection instead of Voxel approximation, or trying both and selecting the one w/ lower geometric deviation.
Disabled coplanar merging and planar retriangulation when doing hidden removal of voxel LODs (was unintentionally doing this)
Update internal usage

diff --git a/Engine/Plugins/Experimental/MeshModelingToolsetExp/Source/GeometryProcessingAdapters/Private/CombineMeshInstancesImpl.cpp b/Engine/Plugins/Experimental/MeshModelingToolsetExp/Source/GeometryProcessingAdapters/Private/CombineMeshInstancesImpl.cpp
--- a/Engine/Plugins/Experimental/MeshModelingToolsetExp/Source/GeometryProcessingAdapters/Private/CombineMeshInstancesImpl.cpp
+++ b/Engine/Plugins/Experimental/MeshModelingToolsetExp/Source/GeometryProcessingAdapters/Private/CombineMeshInstancesImpl.cpp
@@ -321,7 +321,48 @@ void InitializeAssemblySourceMeshesFromLOD(
 
 
 
+/**
+ * @return ( Sqrt(Sum-of-squared-distances) , Max(distance)  )
+ * 
+ */
+static FVector2d ComputeGeometricDeviation(
+	const FDynamicMesh3& MeasureMesh, 
+	const FDynamicMeshAABBTree3& SourceBVH)
+{
+	int PointCount = 0;
+	double SumDistanceSqr = 0;
+	double MaxDistanceSqr = 0;
+	auto TestPointFunc = [&SumDistanceSqr, &MaxDistanceSqr, &PointCount, &SourceBVH](FVector3d Point)
+	{
+		double NearDistSqr;
+		SourceBVH.FindNearestTriangle(Point, NearDistSqr);
+		if (NearDistSqr > MaxDistanceSqr)
+		{
+			MaxDistanceSqr = NearDistSqr;
+		}
+		SumDistanceSqr += NearDistSqr;
+		PointCount++;
+	};
+
+	for (int32 vid : MeasureMesh.VertexIndicesItr())
+	{
+		TestPointFunc(MeasureMesh.GetVertex(vid));
+	}
+
+	for (int32 tid : MeasureMesh.TriangleIndicesItr())
+	{
+		TestPointFunc(MeasureMesh.GetTriCentroid(tid));
+	}
+
+	for (int32 eid : MeasureMesh.EdgeIndicesItr())
+	{
+		TestPointFunc(MeasureMesh.GetEdgePoint(eid, 0.5));
+	}
 
+	return FVector2d(
+		FMathd::Sqrt(SumDistanceSqr),
+		FMathd::Sqrt(MaxDistanceSqr) );
+}
 
 
 
@@ -955,18 +996,27 @@ static void ComputeSweptSolidApproximation(
 		}
 	}
 
+	// above result is likely to be extremely noisy, so we want to clean it up a bit, particularly
+	// if we are going to do an offset/inset closure...
+	double CleanupTol = FMath::Max(SimplifyTolerance * 0.25, 0.1);
+	for (FGeneralPolygon2d& Polygon : Polygons)
+	{
+		Polygon.Simplify(CleanupTol, CleanupTol);
+	}
+
 	// can optionally try to reduce polygon complexity by topological closure (dilate/erode)
 	if (MergeOffset > 0)
 	{
 		TArray<FGeneralPolygon2d> TmpPolygons;
 		PolygonsOffsets(MergeOffset, -MergeOffset,
-			Polygons, TmpPolygons, true, 1.0,
-			EPolygonOffsetJoinType::Square,
+			Polygons, TmpPolygons, true, MergeOffset * FMathd::Sqrt2,
+			EPolygonOffsetJoinType::Miter,
 			EPolygonOffsetEndType::Polygon);
 
 		Polygons = MoveTemp(TmpPolygons);
 	}
 
+	// clean up polygons, remove small holes, and pass to triangulator
 	FConstrainedDelaunay2d Triangulator;
 	for (FGeneralPolygon2d& Polygon : Polygons)
 	{
@@ -2149,6 +2199,65 @@ static void ComputeSimplifiedVoxWrapMesh(
 }
 
 
+/**
+ * Computes best cardinal-axis swept-solid approximation to CombinedMesh and returns in ResultMesh.
+ * The swept-solid approximation is found by flattening ResultMesh along that axis, and then doing
+ * polygon booleans and topological closure. The various simplification parameters are derived from
+ * the ClosureDistance parameter.
+ */
+static void ComputeBestFullProjectionMesh(
+	const FDynamicMesh3& CombinedMesh,
+	FDynamicMeshAABBTree3& CombinedMeshSpatial,
+	FDynamicMesh3& ResultMesh,
+	double ClosureDistance)
+{
+	TArray<FVector3d, TInlineAllocator<3>> Directions;
+	Directions.Add(FVector3d::UnitZ());
+	Directions.Add(FVector3d::UnitX());
+	Directions.Add(FVector3d::UnitY());
+	int32 N = Directions.Num();
+
+	TArray<FDynamicMesh3, TInlineAllocator<3>> DirectionMeshes;
+	DirectionMeshes.SetNum(N);
+
+	TArray<FVector2d> DeviationMeasures;
+	DeviationMeasures.SetNum(N);
+
+	ParallelFor(Directions.Num(), [&](int32 k)
+	{
+		FDynamicMesh3& UseMesh = DirectionMeshes[k];
+		FVector3d UseDirection = Directions[k];
+		ComputeSweptSolidApproximation(CombinedMesh, UseMesh, UseDirection,
+			ClosureDistance, ClosureDistance / 4, 4.0f * ClosureDistance * ClosureDistance);
+		UseMesh.DiscardAttributes();
+
+		// simplify to planar
+		FQEMSimplification Simplifier(&UseMesh);
+
+		Simplifier.CollapseMode = FQEMSimplification::ESimplificationCollapseModes::MinimalExistingVertexError;
+		// no constraints as we discarded attributes
+		Simplifier.SimplifyToMinimalPlanar(ClosureDistance/2);
+
+		DeviationMeasures[k] = ComputeGeometricDeviation(UseMesh, CombinedMeshSpatial);
+	});
+
+	// select option w/ smallest max geometric deviation
+	double MinMaxDistance = TNumericLimits<double>::Max();
+	int32 UseIndex = 0;
+	for (int32 k = 0; k < N; ++k)
+	{
+		if (DeviationMeasures[k].Y < MinMaxDistance)
+		{
+			MinMaxDistance = DeviationMeasures[k].Y;
+			UseIndex = k;
+		}
+	}
+
+	ResultMesh = MoveTemp(DirectionMeshes[UseIndex]);
+	ResultMesh.CompactInPlace();
+}
+
+
 
 template<typename SimplificationType>
 void DoSimplifyMesh(
@@ -2477,7 +2586,8 @@ static void SortMesh(FDynamicMesh3& Mesh)
 void ComputeHiddenRemovalForLOD(
 	FDynamicMesh3& MeshLOD,
 	int32 LODIndex,
-	IGeometryProcessing_CombineMeshInstances::FOptions CombineOptions)
+	IGeometryProcessing_CombineMeshInstances::FOptions CombineOptions,
+	bool bSkipPostProcessing = false)
 {
 	bool bVerbose = CVarGeometryCombineMeshInstancesVerbose.GetValueOnAnyThread();
 
@@ -2501,7 +2611,7 @@ void ComputeHiddenRemovalForLOD(
 		UE_LOG(LogGeometry, Log, TEXT("    Remove Hidden Faces - [Tris %6d Verts %6d]"), MeshLOD.TriangleCount(), MeshLOD.VertexCount());
 	}
 
-	if ( bModified )
+	if ( bModified && bSkipPostProcessing == false)
 	{
 		TFunction<FIndex3i(const FDynamicMesh3& Mesh, int32 TriangleID)> GroupingIDFunc = CombineOptions.TriangleGroupingIDFunc;
 		if (!GroupingIDFunc)
@@ -2549,7 +2659,7 @@ void BuildCombinedMesh(
 		{
 			LODTypes[LODLevel] = ECombinedLODType::Simplified;
 		}
-		else if (LODLevel >= NumLODs - CombineOptions.NumVoxWrapLODs)
+		else if (LODLevel >= NumLODs - CombineOptions.NumCoarseLODs)
 		{
 			LODTypes[LODLevel] = ECombinedLODType::VoxWrapped;
 			FirstVoxWrappedIndex = FMath::Min(LODLevel, FirstVoxWrappedIndex);
@@ -2672,6 +2782,13 @@ void BuildCombinedMesh(
 		}
 	}
 
+	// Above we copied the last approximation LOD into the first VoxWrap LOD slot. However 
+	// approximate LOD may be very coarse, it's better to use a known-good starting mesh?
+	// Need to do this before starting hidden-removal task...
+	if (FirstVoxWrappedIndex < MeshLODs.Num())
+	{
+		MeshLODs[FirstVoxWrappedIndex].Mesh = MeshLODs[CombineOptions.NumCopiedLODs-1].Mesh;
+	}
 
 	//
 	// start hidden-removal passes on all meshes up to voxel LODs here, because we can compute voxel LOD at the same time
@@ -2706,70 +2823,101 @@ void BuildCombinedMesh(
 	//
 	// Process VoxWrapped LODs 
 	//
+	bool bUsingCoarseSweepApproximation = false;
 	if ( FirstVoxWrappedIndex < 9999 )
 	{
 		FDynamicMesh3 SourceVoxWrapMesh = MoveTemp(MeshLODs[FirstVoxWrappedIndex].Mesh);
-		FDynamicMeshAABBTree3 Spatial(&SourceVoxWrapMesh, true);
+		FDynamicMeshAABBTree3 SourceSpatial(&SourceVoxWrapMesh, true);
 
-		const double CLOSURE_DIST = 10.0;		// TODO: this needs to be exposed as an option, perhaps per-part
+		FDynamicMesh3 InitialCoarseApproximation;
 
-		FDynamicMesh3 TempBaseVoxWrapMesh;
 		double VoxelDimension = 2.0;	// may be modified by ComputeVoxWrapMesh call
 		{
 			TRACE_CPUPROFILER_EVENT_SCOPE(ComputeVoxWrap);
-			ComputeVoxWrapMesh(SourceVoxWrapMesh, Spatial, TempBaseVoxWrapMesh, CLOSURE_DIST, VoxelDimension);
+			
+			if (CombineOptions.CoarseLODStrategy == IGeometryProcessing_CombineMeshInstances::ECoarseApproximationStrategy::VoxelBasedSolidApproximation)
+			{
+				ComputeVoxWrapMesh(SourceVoxWrapMesh, SourceSpatial, InitialCoarseApproximation, CombineOptions.CoarseApproximationDetailSize, VoxelDimension);
+				bUsingCoarseSweepApproximation = false;
+			}
+			else if (CombineOptions.CoarseLODStrategy == IGeometryProcessing_CombineMeshInstances::ECoarseApproximationStrategy::SweptPlanarProjection)
+			{
+				ComputeBestFullProjectionMesh(SourceVoxWrapMesh, SourceSpatial, InitialCoarseApproximation, CombineOptions.CoarseApproximationDetailSize);
+				bUsingCoarseSweepApproximation = true;
+			}
+			else  // Automatic
+			{
+				// try swept-planar-projection as it is cheaper and generally better. If it deviates too much, fall back to voxel
+				FDynamicMesh3 SweptPlanarCoarseMesh;
+				ComputeBestFullProjectionMesh(SourceVoxWrapMesh, SourceSpatial, SweptPlanarCoarseMesh, CombineOptions.CoarseApproximationDetailSize);
+				FVector2d SweepDeviation = ComputeGeometricDeviation(SweptPlanarCoarseMesh, SourceSpatial);
+				bUsingCoarseSweepApproximation = (SweepDeviation.Y < 2.0 * CombineOptions.CoarseApproximationDetailSize);
+				if (bUsingCoarseSweepApproximation)
+				{
+					InitialCoarseApproximation = MoveTemp(SweptPlanarCoarseMesh);
+				}
+				else
+				{
+					FDynamicMesh3 VoxWrapCoarseMesh;
+					ComputeVoxWrapMesh(SourceVoxWrapMesh, SourceSpatial, VoxWrapCoarseMesh, CombineOptions.CoarseApproximationDetailSize, VoxelDimension);
+					InitialCoarseApproximation = MoveTemp(VoxWrapCoarseMesh);
+				}
+			}
+			
 			// currently need to re-sort output to remove non-determinism...
-			SortMesh(TempBaseVoxWrapMesh);
+			SortMesh(InitialCoarseApproximation);
 
-			//UE_LOG(LogGeometry, Warning, TEXT("VoxWrapMesh has %d triangles %d vertices"), TempBaseVoxWrapMesh.TriangleCount(), TempBaseVoxWrapMesh.VertexCount());
+			//UE_LOG(LogGeometry, Warning, TEXT("VoxWrapMesh has %d triangles %d vertices"), InitialCoarseApproximation.TriangleCount(), InitialCoarseApproximation.VertexCount());
 		}
 
 		if ( bVerbose )
 		{
-			UE_LOG(LogGeometry, Log, TEXT("  Generated Base VoxWrap Mesh - Tris %8d Verts %8d - CellSize is %4.3f"), TempBaseVoxWrapMesh.TriangleCount(), TempBaseVoxWrapMesh.VertexCount(), VoxelDimension);
+			UE_LOG(LogGeometry, Log, TEXT("  Generated Base Coarse Mesh - Tris %8d Verts %8d - CellSize is %4.3f"), InitialCoarseApproximation.TriangleCount(), InitialCoarseApproximation.VertexCount(), VoxelDimension);
 		}
 
+		InitialCoarseApproximation.DiscardAttributes();
+		const int FastCollapseToTriCount = 50000;
+		if (InitialCoarseApproximation.TriangleCount() > FastCollapseToTriCount+500)
 		{
 			TRACE_CPUPROFILER_EVENT_SCOPE(FastCollapsePrePass);
-			TempBaseVoxWrapMesh.DiscardAttributes();
-			FVolPresMeshSimplification Simplifier(&TempBaseVoxWrapMesh);
+			FVolPresMeshSimplification Simplifier(&InitialCoarseApproximation);
 			Simplifier.bAllowSeamCollapse = false;
-			Simplifier.FastCollapsePass(VoxelDimension*0.5, 10, false, 50000);
+			Simplifier.FastCollapsePass(VoxelDimension * 0.5, 10, false, 50000);
 		}
 
 		if ( bVerbose )
 		{
-			UE_LOG(LogGeometry, Log, TEXT("         FastCollapse         - Tris %8d Verts %8d"), TempBaseVoxWrapMesh.TriangleCount(), TempBaseVoxWrapMesh.VertexCount());
+			UE_LOG(LogGeometry, Log, TEXT("         FastCollapse         - Tris %8d Verts %8d"), InitialCoarseApproximation.TriangleCount(), InitialCoarseApproximation.VertexCount());
 		}
 
-		// need to ensure that the triangle count of the first voxel LOD ends up smaller
-		// than the triangle count of the last approximate/etc LOD. We don't know this until
-		// the hidden-removal tasks finish, so wait for them here. We cannot guarantee
-		// that those LODs necessarily reduce in triangle count, that gets sorted out later,
-		// so find the min count here
+		int32 MaxTriCount = CombineOptions.CoarseLODMaxTriCountBase;
+		double SimplifyTolerance = CombineOptions.CoarseLODBaseTolerance;
+
+		// for very simple parts it can be the case that the last approximate LOD is
+		// lower tri-count than the first coarse approximation. In that case just use it.
 		UE::Tasks::Wait(PendingRemoveHiddenTasks);
-		int32 MinNonVoxWrapLODTriCount = MeshLODs[0].Mesh.TriangleCount();
-		for (int32 k = 1; k < FirstVoxWrappedIndex; ++k)
+		int32 PrevLODTriCount = MeshLODs[FirstVoxWrappedIndex - 1].Mesh.TriangleCount();
+		if (PrevLODTriCount < InitialCoarseApproximation.TriangleCount() && PrevLODTriCount < MaxTriCount)
 		{
-			MinNonVoxWrapLODTriCount = FMath::Min(MinNonVoxWrapLODTriCount, MeshLODs[k].Mesh.TriangleCount());
+			InitialCoarseApproximation = MeshLODs[FirstVoxWrappedIndex-1].Mesh;
 		}
-		// half is maybe a bit aggressive...
-		int32 MaxTriCount = FMath::Min(CombineOptions.VoxWrapMaxTriCountBase, MinNonVoxWrapLODTriCount / 2);
-		double SimplifyTolerance = CombineOptions.VoxWrapBaseTolerance;
 
-		// Current state of TempBaseVoxWrapMesh is our initial voxel LOD. To ensure
+		// Current state of InitialCoarseApproximation is our initial voxel LOD. To ensure
 		// that voxel LODs have compatible UVs (to allow baking), we compute UVs on
 		// the first LOD and allow them to propagate (and currently normals as well)
-		TempBaseVoxWrapMesh.DiscardAttributes();
+		InitialCoarseApproximation.DiscardAttributes();
 		{
 			TRACE_CPUPROFILER_EVENT_SCOPE(SimplifyVoxWrap);
-			ComputeSimplifiedVoxWrapMesh(TempBaseVoxWrapMesh, &SourceVoxWrapMesh, &Spatial,
-				SimplifyTolerance, MaxTriCount);
+			if (InitialCoarseApproximation.TriangleCount() > MaxTriCount)
+			{
+				ComputeSimplifiedVoxWrapMesh(InitialCoarseApproximation, &SourceVoxWrapMesh, &SourceSpatial,
+					SimplifyTolerance, MaxTriCount);
+			}
 		}
-		TempBaseVoxWrapMesh.EnableAttributes();
-		InitializeNormalsFromAngleThreshold(TempBaseVoxWrapMesh, CombineOptions.HardNormalAngleDeg);
-		ComputeVoxWrapMeshAutoUV(TempBaseVoxWrapMesh);
-		MeshLODs[FirstVoxWrappedIndex].Mesh = MoveTemp(TempBaseVoxWrapMesh);
+		InitialCoarseApproximation.EnableAttributes();
+		InitializeNormalsFromAngleThreshold(InitialCoarseApproximation, CombineOptions.HardNormalAngleDeg);
+		ComputeVoxWrapMeshAutoUV(InitialCoarseApproximation);
+		MeshLODs[FirstVoxWrappedIndex].Mesh = MoveTemp(InitialCoarseApproximation);
 
 		// iterate simplification criteria to next level
 		SimplifyTolerance *= 1.5;
@@ -2780,7 +2928,10 @@ void BuildCombinedMesh(
 			// need to simplify from previous level to preserve UVs/etc
 			MeshLODs[LODIndex].Mesh = MeshLODs[LODIndex-1].Mesh;
 
-			DoSimplifyMesh<FAttrMeshSimplification>(MeshLODs[LODIndex].Mesh, MaxTriCount, nullptr, SimplifyTolerance);
+			if (MeshLODs[LODIndex].Mesh.TriangleCount() > MaxTriCount)
+			{
+				DoSimplifyMesh<FAttrMeshSimplification>(MeshLODs[LODIndex].Mesh, MaxTriCount, nullptr, SimplifyTolerance);
+			}
 
 			SimplifyTolerance *= 1.5;
 			MaxTriCount /= 2;
@@ -2788,10 +2939,10 @@ void BuildCombinedMesh(
 
 		// Project colors and materials after mesh simplification to avoid constraining it.
 		// If they /should/ constrain simplification, then they should be projected onto the first
-		// mesh (TempBaseVoxWrapMesh above) and they will automatically transfer
+		// mesh (InitialCoarseApproximation above) and they will automatically transfer
 		for (int32 LODIndex = FirstVoxWrappedIndex; LODIndex < NumLODs; ++LODIndex)
 		{
-			ProjectAttributes(MeshLODs[LODIndex].Mesh, &SourceVoxWrapMesh, &Spatial);
+			ProjectAttributes(MeshLODs[LODIndex].Mesh, &SourceVoxWrapMesh, &SourceSpatial);
 		}
 	}
 
@@ -2804,7 +2955,7 @@ void BuildCombinedMesh(
 	}
 
 	// remove hidden faces on voxel LODs (todo: can do this via shape sorting, much faster)
-	if (bRemoveHiddenFaces)
+	if (bRemoveHiddenFaces && bUsingCoarseSweepApproximation == false)
 	{
 		TRACE_CPUPROFILER_EVENT_SCOPE(RemoveHidden);
 		ParallelFor(NumLODs, [&](int32 LODIndex)
@@ -2816,7 +2967,8 @@ void BuildCombinedMesh(
 					UE_LOG(LogGeometry, Log, TEXT("  Optimizing LOD%d - Tris %6d Verts %6d"), LODIndex, MeshLODs[LODIndex].Mesh.TriangleCount(), MeshLODs[LODIndex].Mesh.VertexCount());
 				}
 
-				ComputeHiddenRemovalForLOD(MeshLODs[LODIndex].Mesh, LODIndex, CombineOptions);
+				// disable hidden-removal postprocessing for voxel LODs (this should not be part of ComputeHiddenRemovalForLOD...)
+				ComputeHiddenRemovalForLOD(MeshLODs[LODIndex].Mesh, LODIndex, CombineOptions, /*bSkipPostProcessing=*/ true);
 			}
 		}, (bVerbose) ? EParallelForFlags::ForceSingleThread :  EParallelForFlags::None );
 	}
@@ -3034,10 +3186,10 @@ void FCombineMeshInstancesImpl::CombineMeshInstances(
 	bool bVerbose = CVarGeometryCombineMeshInstancesVerbose.GetValueOnGameThread();
 	if (bVerbose)
 	{
-		int32 NumApproxLODs = FMath::Max(0, Options.NumLODs - Options.NumCopiedLODs - Options.NumSimplifiedLODs - Options.NumVoxWrapLODs);
-		UE_LOG(LogGeometry, Log, TEXT("CombineMeshInstances: processing %d Instances into %d LODs (%d Copied, %d Simplified, %d Approx, %d VoxWrapped)"),
+		int32 NumApproxLODs = FMath::Max(0, Options.NumLODs - Options.NumCopiedLODs - Options.NumSimplifiedLODs - Options.NumCoarseLODs);
+		UE_LOG(LogGeometry, Log, TEXT("CombineMeshInstances: processing %d Instances into %d LODs (%d Copied, %d Simplified, %d Approx, %d Coarse)"),
 			MeshInstances.StaticMeshInstances.Num(),
-			Options.NumLODs, Options.NumCopiedLODs, Options.NumSimplifiedLODs, NumApproxLODs, Options.NumVoxWrapLODs);
+			Options.NumLODs, Options.NumCopiedLODs, Options.NumSimplifiedLODs, NumApproxLODs, Options.NumCoarseLODs);
 	}
 
 	FMeshPartsAssembly PartAssembly;
diff --git a/Engine/Source/Developer/GeometryProcessingInterfaces/Public/GeometryProcessingInterfaces/CombineMeshInstances.h b/Engine/Source/Developer/GeometryProcessingInterfaces/Public/GeometryProcessingInterfaces/CombineMeshInstances.h
--- a/Engine/Source/Developer/GeometryProcessingInterfaces/Public/GeometryProcessingInterfaces/CombineMeshInstances.h
+++ b/Engine/Source/Developer/GeometryProcessingInterfaces/Public/GeometryProcessingInterfaces/CombineMeshInstances.h
@@ -89,6 +89,13 @@ public:
 	};
 
 
+	enum class ECoarseApproximationStrategy
+	{
+		Automatic = 0,
+		VoxelBasedSolidApproximation = 1,
+		SweptPlanarProjection = 2
+	};
+
 	enum class EVertexColorMappingMode
 	{
 		None = 0,
@@ -124,12 +131,25 @@ public:
 		double MaxAllowableApproximationDeviation = 5.0;
 
 		//
-		// settings for Voxel LODs
+		// Settings for Coarse LODs. Coarse LODs are the lowest (furthest) LODs and generally assumed to have 
+		// very low triangle counts. Two strategies are supported:
+		//   1) VoxelBasedSolidApproximation - does a solid approximation with topological closure (outset + inset), which will fill in any small holes/gaps
+		//      or small features. CoarseApproximationDetailSize is used as the closure radius/distance. 
+		//   2) SweptPlanarProjection - The mesh is effectively projected to 2D along each X/Y/Z axis, with 2D polygon boolean & topological closure using CoarseApproximationDetailSize,
+		//      as well as various other cleanup, and then that 2D polygon is triangulated and swept along the part bounding-box extent. The max distance
+		//      is measured from the approximation to the original mesh, and the axis-swept-mesh with the smallest max-distance is used.
+		// 
+		// In either case, the resulting mesh is simplified to CoarseLODMaxTriCountBase, and then further simplified for each additional coarse LOD,
+		// by halving the target triangle count. 
+		// 
+		// The 'Automatic' strategy uses the SweptPlanarProjection if it's max approximation deviation is within K*CoarseApproximationDetailSize (currently K=2),
+		// and otherwise falls back to using VoxelBasedSolidApproximation
 		//
-
-		int32 NumVoxWrapLODs = 1;
-		double VoxWrapBaseTolerance = 1.0;
-		int32 VoxWrapMaxTriCountBase = 500;
+		ECoarseApproximationStrategy CoarseLODStrategy = ECoarseApproximationStrategy::Automatic;
+		int32 NumCoarseLODs = 1;
+		double CoarseLODBaseTolerance = 1.0;
+		int32 CoarseLODMaxTriCountBase = 500;
+		double CoarseApproximationDetailSize = 10.0;
 
 		//
 		// Hidden Faces removal options
