moving skeletal mesh tools in MeshModelingToolsEditorOnlyExp for now

diff --git a/Engine/Plugins/Experimental/Animation/SkeletalMeshModelingTools/Source/SkeletalMeshModelingTools/Private/SkeletalMeshModelingToolsEditorMode.cpp b/Engine/Plugins/Experimental/Animation/SkeletalMeshModelingTools/Source/SkeletalMeshModelingTools/Private/SkeletalMeshModelingToolsEditorMode.cpp
--- a/Engine/Plugins/Experimental/Animation/SkeletalMeshModelingTools/Source/SkeletalMeshModelingTools/Private/SkeletalMeshModelingToolsEditorMode.cpp
+++ b/Engine/Plugins/Experimental/Animation/SkeletalMeshModelingTools/Source/SkeletalMeshModelingTools/Private/SkeletalMeshModelingToolsEditorMode.cpp
@@ -42,17 +42,17 @@
 #include "RemeshMeshTool.h"
 #include "RemoveOccludedTrianglesTool.h"
 #include "SimplifyMeshTool.h"
-#include "SkinWeightsBindingTool.h"
-#include "SkinWeightsPaintTool.h"
 #include "SmoothMeshTool.h"
 #include "ToolTargetManager.h"
 #include "WeldMeshEdgesTool.h"
 
 #include "SkeletalMeshNotifier.h"
 #include "Animation/DebugSkelMeshComponent.h"
-#include "SkeletalMesh/SkeletalMeshEditionInterface.h"
 
+#include "SkeletalMesh/SkeletalMeshEditionInterface.h"
 #include "SkeletalMesh/SkeletonEditingTool.h"
+#include "SkeletalMesh/SkinWeightsBindingTool.h"
+#include "SkeletalMesh/SkinWeightsPaintTool.h"
 
 #define LOCTEXT_NAMESPACE "SkeletalMeshModelingToolsEditorMode"
 
diff --git a/Engine/Plugins/Experimental/MeshModelingToolsetExp/Source/MeshModelingToolsEditorOnlyExp/MeshModelingToolsEditorOnlyExp.Build.cs b/Engine/Plugins/Experimental/MeshModelingToolsetExp/Source/MeshModelingToolsEditorOnlyExp/MeshModelingToolsEditorOnlyExp.Build.cs
--- a/Engine/Plugins/Experimental/MeshModelingToolsetExp/Source/MeshModelingToolsEditorOnlyExp/MeshModelingToolsEditorOnlyExp.Build.cs
+++ b/Engine/Plugins/Experimental/MeshModelingToolsetExp/Source/MeshModelingToolsEditorOnlyExp/MeshModelingToolsEditorOnlyExp.Build.cs
@@ -40,6 +40,8 @@ public class MeshModelingToolsEditorOnlyExp : ModuleRules
 				"ModelingComponentsEditorOnly",
 				"ModelingOperators",
 				"ModelingOperatorsEditorOnly",
+				"SkeletalMeshUtilitiesCommon",
+				"SkeletalMeshModifiers"
 				// ... add other public dependencies that you statically link with here ...
 			}
 		);
diff --git a/Engine/Plugins/Experimental/MeshModelingToolsetExp/Source/MeshModelingToolsEditorOnlyExp/Private/MeshModelingToolsEditorOnlyExp.cpp b/Engine/Plugins/Experimental/MeshModelingToolsetExp/Source/MeshModelingToolsEditorOnlyExp/Private/MeshModelingToolsEditorOnlyExp.cpp
--- a/Engine/Plugins/Experimental/MeshModelingToolsetExp/Source/MeshModelingToolsEditorOnlyExp/Private/MeshModelingToolsEditorOnlyExp.cpp
+++ b/Engine/Plugins/Experimental/MeshModelingToolsetExp/Source/MeshModelingToolsEditorOnlyExp/Private/MeshModelingToolsEditorOnlyExp.cpp
@@ -4,6 +4,8 @@
 
 #define LOCTEXT_NAMESPACE "FMeshModelingToolsEditorOnlyExpModule"
 
+DEFINE_LOG_CATEGORY(LogMeshModelingToolsEditor);
+
 void FMeshModelingToolsEditorOnlyExpModule::StartupModule()
 {
 }
diff --git a/Engine/Plugins/Experimental/MeshModelingToolsetExp/Source/MeshModelingToolsEditorOnlyExp/Private/SkeletalMesh/SkeletalMeshEditionInterface.cpp b/Engine/Plugins/Experimental/MeshModelingToolsetExp/Source/MeshModelingToolsEditorOnlyExp/Private/SkeletalMesh/SkeletalMeshEditionInterface.cpp
new file mode 100644
--- /dev/null
+++ b/Engine/Plugins/Experimental/MeshModelingToolsetExp/Source/MeshModelingToolsEditorOnlyExp/Private/SkeletalMesh/SkeletalMeshEditionInterface.cpp
@@ -0,0 +1,59 @@
+﻿// Copyright Epic Games, Inc. All Rights Reserved.
+
+#include "SkeletalMesh/SkeletalMeshEditionInterface.h"
+
+IMPLEMENT_HIT_PROXY(HBoneHitProxy, HHitProxy)
+
+ISkeletalMeshNotifier& ISkeletalMeshEditionInterface::GetNotifier()
+{
+	if (!Notifier)
+	{
+		Notifier.Reset(new FSkeletalMeshToolNotifier(TWeakInterfacePtr<ISkeletalMeshEditionInterface>(this)));
+	}
+	
+	return *Notifier;
+}
+
+bool ISkeletalMeshEditionInterface::NeedsNotification() const
+{
+	return Notifier && Notifier->Delegate().IsBound(); 
+}
+
+void ISkeletalMeshEditionInterface::BindTo(TSharedPtr<ISkeletalMeshEditorBinding> InBinding)
+{
+	Binding = InBinding;
+}
+
+void ISkeletalMeshEditionInterface::Unbind()
+{
+	Binding.Reset();
+}
+
+TOptional<FName> ISkeletalMeshEditionInterface::GetBoneName(HHitProxy* InHitProxy) const
+{
+	if (const HBoneHitProxy* BoneProxy = HitProxyCast<HBoneHitProxy>(InHitProxy))
+	{
+		return BoneProxy->BoneName;
+	}
+
+	return Binding.IsValid() && Binding.Pin()->GetNameFunction() ? Binding.Pin()->GetNameFunction()(InHitProxy) : TOptional<FName>();
+}
+
+TArray<FName> ISkeletalMeshEditionInterface::GetSelectedBones() const
+{
+	static const TArray<FName> Dummy;
+	return Binding.IsValid() ? Binding.Pin()->GetSelectedBones() : Dummy;
+}
+
+FSkeletalMeshToolNotifier::FSkeletalMeshToolNotifier(TWeakInterfacePtr<ISkeletalMeshEditionInterface> InInterface)
+	: ISkeletalMeshNotifier()
+	, Interface(InInterface)
+{}
+
+void FSkeletalMeshToolNotifier::HandleNotification(const TArray<FName>& BoneNames, const ESkeletalMeshNotifyType InNotifyType)
+{
+	if (Interface.IsValid())
+	{
+		Interface->HandleSkeletalMeshModified(BoneNames, InNotifyType);
+	}
+}
diff --git a/Engine/Plugins/Experimental/MeshModelingToolsetExp/Source/MeshModelingToolsEditorOnlyExp/Private/SkeletalMesh/SkeletonEditingTool.cpp b/Engine/Plugins/Experimental/MeshModelingToolsetExp/Source/MeshModelingToolsEditorOnlyExp/Private/SkeletalMesh/SkeletonEditingTool.cpp
new file mode 100644
--- /dev/null
+++ b/Engine/Plugins/Experimental/MeshModelingToolsetExp/Source/MeshModelingToolsEditorOnlyExp/Private/SkeletalMesh/SkeletonEditingTool.cpp
@@ -0,0 +1,1341 @@
+﻿// Copyright Epic Games, Inc. All Rights Reserved.
+
+
+#include "SkeletalMesh/SkeletonEditingTool.h"
+
+#include "ContextObjectStore.h"
+#include "InteractiveToolManager.h"
+#include "ModelingToolTargetUtil.h"
+#include "SkeletalDebugRendering.h"
+#include "ToolTargetManager.h"
+#include "BaseBehaviors/SingleClickBehavior.h"
+#include "Engine/SkeletalMesh.h"
+#include "Components/SkeletalMeshComponent.h"
+#include "UnrealClient.h"
+#include "HitProxies.h"
+#include "BaseBehaviors/ClickDragBehavior.h"
+#include "BaseGizmos/GizmoViewContext.h"
+
+#include "Algo/Count.h"
+#include "BaseGizmos/TransformGizmoUtil.h"
+#include "SkeletalMesh/SkeletonTransformProxy.h"
+#include "TargetInterfaces/PrimitiveComponentBackedTarget.h"
+
+#define LOCTEXT_NAMESPACE "USkeletonEditingTool"
+
+namespace SkeletonEditingTool
+{
+	
+FRefSkeletonChange::FRefSkeletonChange(const USkeletonEditingTool* InTool)
+	: FToolCommandChange()
+	, PreChangeSkeleton(InTool->Modifier->GetReferenceSkeleton())
+	, PreBoneTracker(InTool->Modifier->GetBoneIndexTracker())
+	, PostChangeSkeleton(InTool->Modifier->GetReferenceSkeleton())
+	, PostBoneTracker(InTool->Modifier->GetBoneIndexTracker())
+{}
+
+void FRefSkeletonChange::StoreSkeleton(const USkeletonEditingTool* InTool)
+{
+	PostChangeSkeleton = InTool->Modifier->GetReferenceSkeleton();
+	PostBoneTracker = InTool->Modifier->GetBoneIndexTracker();
+}
+
+void FRefSkeletonChange::Apply(UObject* Object)
+{ // redo
+	USkeletonEditingTool* Tool = CastChecked<USkeletonEditingTool>(Object);
+	Tool->Modifier->ExternalUpdate(PostChangeSkeleton, PostBoneTracker);
+	Tool->UpdateGizmo();
+}
+
+void FRefSkeletonChange::Revert(UObject* Object)
+{ // undo
+	USkeletonEditingTool* Tool = CastChecked<USkeletonEditingTool>(Object);
+	Tool->Modifier->ExternalUpdate(PreChangeSkeleton, PreBoneTracker);
+	Tool->UpdateGizmo();
+}
+
+}
+
+/**
+ * USkeletonEditingToolBuilder
+ */
+
+bool USkeletonEditingToolBuilder::CanBuildTool(const FToolBuilderState& SceneState) const
+{
+	return SceneState.TargetManager->CountSelectedAndTargetable(SceneState, GetTargetRequirements()) == 1;
+}
+
+const FToolTargetTypeRequirements& USkeletonEditingToolBuilder::GetTargetRequirements() const
+{
+	static const FToolTargetTypeRequirements TypeRequirements(UPrimitiveComponentBackedTarget::StaticClass());
+	return TypeRequirements;
+}
+
+UInteractiveTool* USkeletonEditingToolBuilder::BuildTool(const FToolBuilderState& SceneState) const
+{
+	USkeletonEditingTool* NewTool = NewObject<USkeletonEditingTool>(SceneState.ToolManager);
+	
+	UToolTarget* Target = SceneState.TargetManager->BuildFirstSelectedTargetable(SceneState, GetTargetRequirements());
+	NewTool->SetTarget(Target);
+	NewTool->Init(SceneState);
+	
+	return NewTool;
+}
+
+/**
+ * USkeletonEditingTool
+ */
+
+void USkeletonEditingTool::Init(const FToolBuilderState& InSceneState)
+{
+	TargetWorld = InSceneState.World;
+
+	const UContextObjectStore* ContextObjectStore = InSceneState.ToolManager->GetContextObjectStore();
+	ViewContext = ContextObjectStore->FindContext<UGizmoViewContext>();
+	GizmoContext = ContextObjectStore->FindContext<USkeletalMeshGizmoContextObjectBase>();
+}
+
+void USkeletonEditingTool::Setup()
+{
+	Super::Setup();
+
+	const IPrimitiveComponentBackedTarget* TargetComponent = Cast<IPrimitiveComponentBackedTarget>(Target);
+	USkeletalMeshComponent* Component = Cast<USkeletalMeshComponent>(TargetComponent->GetOwnerComponent());
+
+	USkeletalMesh* SkeletalMesh = Component ? Component->GetSkeletalMeshAsset() : nullptr;
+	if (!SkeletalMesh)
+	{
+		return;
+	}
+
+	// setup modifier
+	Modifier = NewObject<USkeletonModifier>(this);
+	Modifier->Init(SkeletalMesh);
+
+	// setup current bone
+	const FReferenceSkeleton& RefSkeleton = Modifier->GetReferenceSkeleton();
+	const int32 NumBones = RefSkeleton.GetNum();
+	const FName& RootBoneName = NumBones ? RefSkeleton.GetBoneName(0) : NAME_None;
+
+	if (NumBones)
+	{
+		Selection = {RootBoneName};
+	}
+
+	// setup preview
+	{
+		PreviewMesh = NewObject<UPreviewMesh>(this);
+		PreviewMesh->bBuildSpatialDataStructure = true;
+		PreviewMesh->CreateInWorld(TargetWorld.Get(), FTransform::Identity);
+
+		PreviewMesh->SetTransform(UE::ToolTarget::GetLocalToWorldTransform(Target));
+
+		PreviewMesh->ReplaceMesh(UE::ToolTarget::GetDynamicMeshCopy(Target));
+
+		const FComponentMaterialSet MaterialSet = UE::ToolTarget::GetMaterialSet(Target);
+		PreviewMesh->SetMaterials(MaterialSet.Materials);
+
+		// hide the skeletal mesh component
+		UE::ToolTarget::HideSourceObject(Target);
+	}
+
+	// setup properties
+	{
+		Properties = NewObject<USkeletonEditingProperties>();
+		Properties->Initialize(this);
+		Properties->Name = GetCurrentBone();
+		Properties->RestoreProperties(this);
+		AddToolPropertySource(Properties);
+
+		ProjectionProperties = NewObject<UProjectionProperties>();
+		ProjectionProperties->Initialize(PreviewMesh);
+		ProjectionProperties->RestoreProperties(this);
+		AddToolPropertySource(ProjectionProperties);
+
+		MirroringProperties = NewObject<UMirroringProperties>();
+		MirroringProperties->Initialize(this);
+		MirroringProperties->RestoreProperties(this);
+		AddToolPropertySource(MirroringProperties);
+
+		OrientingProperties = NewObject<UOrientingProperties>();
+		OrientingProperties->Initialize(this);
+		OrientingProperties->RestoreProperties(this);
+		AddToolPropertySource(OrientingProperties);
+	}
+
+	// setup drag & drop behaviour
+	{
+		UClickDragInputBehavior* ClickDragBehavior = NewObject<UClickDragInputBehavior>(this);
+		ClickDragBehavior->Initialize(this);
+		ClickDragBehavior->Modifiers.RegisterModifier(AddToSelectionModifier, FInputDeviceState::IsShiftKeyDown);
+		ClickDragBehavior->Modifiers.RegisterModifier(ToggleSelectionModifier, FInputDeviceState::IsCtrlKeyDown);
+		AddInputBehavior(ClickDragBehavior);
+	}
+
+	// setup gizmo
+	if (GizmoContext.IsValid())
+	{
+		UGizmoLambdaStateTarget* NewTarget = NewObject<UGizmoLambdaStateTarget>(this);
+		NewTarget->BeginUpdateFunction = [this]()
+		{
+			if (GizmoWrapper && GizmoWrapper->CanInteract())
+			{
+				TGuardValue OperationGuard(Operation, EEditingOperation::Transform);
+				BeginChange();
+			}
+		};
+		NewTarget->EndUpdateFunction = [this]()
+		{
+			if (GizmoWrapper && GizmoWrapper->CanInteract())
+			{
+				TGuardValue OperationGuard(Operation, EEditingOperation::Transform);
+				EndChange();
+
+				const IToolsContextQueriesAPI* ToolsContextQueries = GetToolManager()->GetPairedGizmoManager()->GetContextQueriesAPI();
+				check(ToolsContextQueries);
+				if (ToolsContextQueries->GetCurrentCoordinateSystem() == EToolContextCoordinateSystem::World)
+				{
+					UpdateGizmo();
+				}
+			}
+		};
+		
+		GizmoWrapper = GizmoContext->GetNewWrapper(GetToolManager(), this, NewTarget);
+		if (GizmoWrapper)
+		{
+			GizmoWrapper->Component = Component;
+		}
+	}
+
+	// setup watchers
+	SelectionWatcher.Initialize(
+	[this]()
+	{
+		return this->Selection;
+	},
+	[this](TArray<FName> InBoneNames)
+	{
+		UpdateGizmo();
+		if (NeedsNotification())
+		{
+			GetNotifier().Notify(InBoneNames, ESkeletalMeshNotifyType::BonesSelected);
+		}
+	},
+	this->Selection);
+
+	IToolsContextQueriesAPI* ToolsContextQueries = GetToolManager()->GetPairedGizmoManager()->GetContextQueriesAPI();
+	CoordinateSystemWatcher.Initialize(
+	[ToolsContextQueries]()
+	{
+		check(ToolsContextQueries);
+		return ToolsContextQueries->GetCurrentCoordinateSystem();
+	},
+	[this](EToolContextCoordinateSystem)
+	{
+		UpdateGizmo();
+	},
+	ToolsContextQueries->GetCurrentCoordinateSystem());
+	
+}
+
+void USkeletonEditingTool::UpdateGizmo() const
+{	
+	if (!GizmoWrapper)
+	{
+		return;
+	}
+
+	GizmoWrapper->Initialize();
+
+	const bool bUseGizmo = Operation == EEditingOperation::Select || Operation == EEditingOperation::Transform; 
+	if (Selection.IsEmpty() || !bUseGizmo)
+	{
+		return;
+	}
+
+	const IToolsContextQueriesAPI* ToolsContextQueries = GetToolManager()->GetPairedGizmoManager()->GetContextQueriesAPI();
+	const bool bWorld = ToolsContextQueries->GetCurrentCoordinateSystem() == EToolContextCoordinateSystem::World;
+	
+	const TArray<int32> BoneIndexes = GetSelectedBoneIndexes();
+	FTransform InitTransform = Modifier->GetTransform(BoneIndexes.Last(), true);
+	if (bWorld)
+	{
+		InitTransform = FTransform(FQuat::Identity, InitTransform.GetTranslation(), InitTransform.GetScale3D());
+	}
+	GizmoWrapper->Initialize(InitTransform, ToolsContextQueries->GetCurrentCoordinateSystem());
+
+	const FReferenceSkeleton& ReferenceSkeleton = Modifier->GetReferenceSkeleton();
+	const TArray<FMeshBoneInfo>& BoneInfos = ReferenceSkeleton.GetRawRefBoneInfo();
+	for (const int32 BoneIndex: BoneIndexes)
+	{
+		if (BoneInfos.IsValidIndex(BoneIndex))
+		{
+			const FName BoneName = BoneInfos[BoneIndex].Name;
+			const int32 ParentBoneIndex = BoneInfos[BoneIndex].ParentIndex;
+			const FTransform& ParentGlobal = Modifier->GetTransform(ParentBoneIndex, true);
+		
+			GizmoWrapper->HandleBoneTransform(
+				[this, BoneIndex, bWorld]()
+				{
+					return Modifier->GetTransform(BoneIndex, bWorld);
+				},
+				[this, BoneName, ParentGlobal, bWorld](const FTransform& InNewTransform)
+				{
+					if (bWorld)
+					{
+						const FTransform NewLocal = InNewTransform.GetRelativeTransform(ParentGlobal);
+						return Modifier->SetBoneTransform(BoneName, NewLocal, Properties->bUpdateChildren);
+					}
+					return Modifier->SetBoneTransform(BoneName, InNewTransform, Properties->bUpdateChildren);
+				});
+		}
+	}
+}
+
+
+void USkeletonEditingTool::Shutdown(EToolShutdownType ShutdownType)
+{
+	if (ShutdownType == EToolShutdownType::Accept)
+	{
+		GetToolManager()->BeginUndoTransaction(LOCTEXT("SkeletonEditingTool", "Commit Skeleton Editing"));
+		Modifier->CommitSkeletonToSkeletalMesh();
+		GetToolManager()->EndUndoTransaction();
+
+		// to force to refresh the tree
+		if (NeedsNotification())
+		{
+			GetNotifier().Notify({}, ESkeletalMeshNotifyType::BonesAdded);
+		}
+	}
+	
+	Super::Shutdown(ShutdownType);
+
+	// remove preview mesh
+	if (PreviewMesh != nullptr)
+	{
+		PreviewMesh->SetVisible(false);
+		PreviewMesh->Disconnect();
+		PreviewMesh = nullptr;
+	}
+
+	// show the skeletal mesh component
+	UE::ToolTarget::ShowSourceObject(Target);
+
+	// save properties	
+	Properties->SaveProperties(this);
+	ProjectionProperties->SaveProperties(this);
+	MirroringProperties->SaveProperties(this);
+	OrientingProperties->SaveProperties(this);
+
+	// clear gizmo
+	if (GizmoWrapper)
+	{
+		GizmoWrapper->Clear();
+	}
+}
+
+void USkeletonEditingTool::RegisterActions(FInteractiveToolActionSet& ActionSet)
+{
+	Super::RegisterActions(ActionSet);
+
+	int32 ActionId = static_cast<int32>(EStandardToolActions::BaseClientDefinedActionID) + 400;
+	auto GetActionId = [&ActionId]
+	{
+		return ActionId++;
+	};
+	
+	// register New key
+	ActionSet.RegisterAction(this, GetActionId(), TEXT("CreateNewBone"),
+		LOCTEXT("CreateNewBone", "Create New Bone"),
+		LOCTEXT("CreateNewBoneDesc", "Create New Bone"),
+		EModifierKey::None, EKeys::N,
+		[this]()
+		{
+			Operation = EEditingOperation::Create;
+			UpdateGizmo();
+			GetToolManager()->DisplayMessage(LOCTEXT("Create", "Click & Drag to place a new bone."), EToolMessageLevel::UserNotification);
+		});
+	
+	// register Delete key
+	ActionSet.RegisterAction(this, GetActionId(), TEXT("DeleteSelectedBones"),
+		LOCTEXT("DeleteSelectedBones", "Delete Selected Bone(s)"),
+		LOCTEXT("DeleteSelectedBonesDesc", "Delete Selected Bone(s)"),
+		EModifierKey::None, EKeys::Delete,
+		[this]()
+		{
+			RemoveBones();
+		});
+
+	// register Select key
+	ActionSet.RegisterAction(this, GetActionId(), TEXT("SelectBones"),
+		LOCTEXT("SelectBone", "Select Bone"),
+		LOCTEXT("SelectDesc", "Select Bone"),
+		EModifierKey::None, EKeys::Escape,
+		[this]()
+		{
+			if (Operation != EEditingOperation::Select)
+			{
+				Operation = EEditingOperation::Select;
+				UpdateGizmo();
+				GetToolManager()->DisplayMessage(LOCTEXT("Select", "Click on a bone to select it."), EToolMessageLevel::UserNotification);
+			}
+		});
+
+	// register UnParent key
+	ActionSet.RegisterAction(this, GetActionId(), TEXT("UnparentBones"),
+		LOCTEXT("UnparentBones", "Unparent Bones"),
+		LOCTEXT("UnparentBonesDesc", "Unparent Bones"),
+		EModifierKey::Shift, EKeys::P,
+		[this]()
+		{
+			UnParentBones();
+		});
+		
+	// register Parent key
+	ActionSet.RegisterAction(this, GetActionId(), TEXT("ParentBones"),
+		LOCTEXT("ParentBones", "Parent Bones"),
+		LOCTEXT("ParentBonesDesc", "Parent Bones"),
+		EModifierKey::None, EKeys::B, // FIXME find another shortcut
+		[this]()
+		{
+			Operation = EEditingOperation::Parent;
+			UpdateGizmo();
+			GetToolManager()->DisplayMessage(LOCTEXT("Parent", "Click on a bone to be set as the new parent."), EToolMessageLevel::UserNotification);
+		});
+}
+
+void USkeletonEditingTool::CreateNewBone()
+{
+	if (Operation != EEditingOperation::Create)
+	{
+		return;
+	}
+
+	BeginChange();
+
+	static const FName DefaultName("joint");
+	const FName BoneName = Modifier->GetUniqueName(DefaultName);
+	const FName ParentName = Selection.IsEmpty() ? NAME_None : Selection.Last(); 
+	const bool bBoneAdded = Modifier->AddBone(BoneName, ParentName, Properties->Transform);
+	if (bBoneAdded)
+	{
+		if (NeedsNotification())
+		{
+			GetNotifier().Notify({BoneName}, ESkeletalMeshNotifyType::BonesAdded);
+		}
+	
+		Selection = {BoneName};
+		Properties->Name = BoneName;
+
+		EndChange();
+		return;
+	}
+
+	CancelChange();
+}
+
+void USkeletonEditingTool::MirrorBones()
+{
+	TGuardValue OperationGuard(Operation, EEditingOperation::Mirror);
+	BeginChange();
+
+	const bool bBonesMirrored = Modifier->MirrorBones(GetSelectedBones(), MirroringProperties->Options);
+	if (bBonesMirrored)
+	{
+		EndChange();
+		return;		
+	}
+	
+	CancelChange();
+}
+
+void USkeletonEditingTool::RemoveBones()
+{
+	const TArray<FName> BonesToRemove = GetSelectedBones();
+	
+	TGuardValue OperationGuard(Operation, EEditingOperation::Remove);
+	BeginChange();
+
+	const bool bBonesRemoved = Modifier->RemoveBones(BonesToRemove, true);
+	if (bBonesRemoved)
+	{
+		Selection.RemoveAll([&](const FName& BoneName)
+		{
+			return BonesToRemove.Contains(BoneName);
+		});
+		
+		if (NeedsNotification())
+		{
+			GetNotifier().Notify(BonesToRemove, ESkeletalMeshNotifyType::BonesRemoved);
+		}
+		
+		EndChange();
+		return;
+	}
+
+	CancelChange();
+}
+
+void USkeletonEditingTool::UnParentBones()
+{
+	static const TArray<FName> Dummy;
+
+	TGuardValue OperationGuard(Operation, EEditingOperation::Parent);
+	BeginChange();
+
+	const bool bBonesUnParented = Modifier->ParentBones(GetSelectedBones(), Dummy);
+	if (bBonesUnParented)
+	{
+		GetToolManager()->DisplayMessage(LOCTEXT("Unparent", "Selected bones have been unparented."), EToolMessageLevel::UserNotification);
+		
+		EndChange();
+		return;
+	}
+	
+	CancelChange();
+}
+
+void USkeletonEditingTool::ParentBones(const FName& InParentName)
+{
+	if (Operation != EEditingOperation::Parent)
+	{
+		return;
+	}
+
+	BeginChange();
+	const bool bBonesParented = Modifier->ParentBones(GetSelectedBones(), {InParentName});
+	if (bBonesParented)
+	{
+		EndChange();
+	}
+	else
+	{
+		CancelChange();
+	}
+	Operation = EEditingOperation::Select;
+}
+
+void USkeletonEditingTool::MoveBones()
+{
+	const FReferenceSkeleton& RefSkeleton = Modifier->GetReferenceSkeleton();
+
+	const TArray<FName> Bones = GetSelectedBones();
+	const bool bHasValidBone = Bones.ContainsByPredicate([&](const FName& InBoneName)
+	{
+		return RefSkeleton.FindRawBoneIndex(InBoneName) > INDEX_NONE;
+	});
+
+	if (!bHasValidBone)
+	{
+		return;
+	}
+
+	TGuardValue OperationGuard(Operation, EEditingOperation::Transform);
+	BeginChange();
+
+	const bool bBonesMoved = Modifier->SetBoneTransform(Bones[0], Properties->Transform, Properties->bUpdateChildren);
+	if (bBonesMoved)
+	{
+		// if (NeedsNotification())
+		// {
+		// 	GetNotifier().Notify(Bones, ESkeletalMeshNotifyType::BonesMoved);
+		// }
+		
+		EndChange();
+		return;
+	}
+
+	CancelChange();
+}
+
+void USkeletonEditingTool::RenameBones()
+{
+	const TArray<int32> BoneIndices = GetSelectedBoneIndexes();
+	if (BoneIndices.IsEmpty() || Properties->Name == NAME_None)
+	{
+		return;
+	}
+	
+	const FName CurrentBone = GetCurrentBone();
+	if (BoneIndices.Num() == 1 && CurrentBone == Properties->Name)
+	{
+		return;
+	}
+
+	TGuardValue OperationGuard(Operation, EEditingOperation::Rename);
+	BeginChange();
+
+	TArray<FName> ReversedSelection = Selection;
+	Algo::Reverse(ReversedSelection);
+
+	TArray<FName> NewNames;
+	NewNames.Init(Properties->Name, Selection.Num());
+
+	const bool bBoneRenamed = Modifier->RenameBones(ReversedSelection, NewNames);
+	if (bBoneRenamed)
+	{
+		// update selection with new names
+		const TArray<FMeshBoneInfo>& BoneInfos = Modifier->GetReferenceSkeleton().GetRawRefBoneInfo();
+		for (int32 Index = 0; Index < BoneIndices.Num(); Index++)
+		{
+			if (BoneIndices[Index] != INDEX_NONE)
+			{
+				Selection[Index] = BoneInfos[BoneIndices[Index]].Name;
+			}
+		}
+		Properties->Name = GetCurrentBone();
+		
+		if (NeedsNotification())
+		{
+			GetNotifier().Notify({CurrentBone}, ESkeletalMeshNotifyType::BonesRenamed);
+		}
+		
+		EndChange();
+		return;
+	}
+
+	CancelChange();
+}
+
+void USkeletonEditingTool::OnClickPress(const FInputDeviceRay& InPressPos)
+{
+	if (PendingFunction)
+	{
+		PendingFunction();
+		PendingFunction.Reset();
+	}
+	BeginChange();
+}
+
+void USkeletonEditingTool::OnClickDrag(const FInputDeviceRay& InDragPos)
+{
+	if (Operation != EEditingOperation::Create)
+	{
+		return;
+	}
+	
+	FVector HitPoint;
+	if (ProjectionProperties->GetProjectionPoint(InDragPos, HitPoint))
+	{
+		const FTransform& ParentGlobal = Modifier->GetTransform(ParentIndex, true);
+		Properties->Transform.SetLocation(ParentGlobal.InverseTransformPosition(HitPoint));
+
+		if (!ActiveChange)
+		{
+			TGuardValue OperationGuard(Operation, EEditingOperation::Transform);
+			BeginChange();
+		}
+
+		const bool bBoneMoved = Modifier->SetBoneTransform(GetCurrentBone(), Properties->Transform, Properties->bUpdateChildren);
+		if (!bBoneMoved)
+		{
+			CancelChange();
+			return;
+		}
+
+		const bool bOrient = Operation == EEditingOperation::Create && OrientingProperties->bAutoOrient;
+		if (bOrient && ParentIndex != INDEX_NONE)
+		{
+			const FReferenceSkeleton& RefSkeleton = Modifier->GetReferenceSkeleton();
+			const FName ParentName = RefSkeleton.GetRawRefBoneInfo()[ParentIndex].Name;
+			Modifier->OrientBone(ParentName, OrientingProperties->Options);
+		}
+	}
+}
+
+void USkeletonEditingTool::OrientBones()
+{
+	const FReferenceSkeleton& RefSkeleton = Modifier->GetReferenceSkeleton();
+
+	const TArray<FName> Bones = GetSelectedBones();
+	const bool bHasValidBone = Bones.ContainsByPredicate([&](const FName& InBoneName)
+	{
+		return RefSkeleton.FindRawBoneIndex(InBoneName) > INDEX_NONE;
+	});
+
+	if (!bHasValidBone)
+	{
+		return;
+	}
+
+	TGuardValue OperationGuard(Operation, EEditingOperation::Transform);
+	BeginChange();
+	
+	const bool bBoneOriented = Modifier->OrientBones(Bones, OrientingProperties->Options);
+	if (bBoneOriented)
+	{
+		UpdateGizmo();
+		// if (NeedsNotification())
+		// {
+		// 	GetNotifier().Notify(Bones, ESkeletalMeshNotifyType::BonesMoved);
+		// }
+		
+		EndChange();
+		return;
+	}
+
+	CancelChange();
+}
+
+void USkeletonEditingTool::OnUpdateModifierState(int InModifierID, bool bIsOn)
+{
+	auto UpdateFlag = [this, bIsOn](EBoneSelectionMode InFlag)
+	{
+		if (bIsOn)
+		{
+			EnumAddFlags(SelectionMode, InFlag);
+		}
+		else
+		{
+			EnumRemoveFlags(SelectionMode, InFlag);
+		}		
+	};
+	
+	if (InModifierID == AddToSelectionModifier)
+	{
+		UpdateFlag(EBoneSelectionMode::Additive);
+	}
+	else if (InModifierID == ToggleSelectionModifier)
+	{
+		UpdateFlag(EBoneSelectionMode::Toggle);
+	}
+}
+
+void USkeletonEditingTool::SelectBone(const FName& InBoneName)
+{
+	TArray<FName> NewSelection = Selection;
+
+	if (EnumHasAnyFlags(SelectionMode, EBoneSelectionMode::Additive))
+	{
+		if (InBoneName != NAME_None)
+		{
+			NewSelection.AddUnique(InBoneName);
+		}		
+	}
+	else if (EnumHasAnyFlags(SelectionMode, EBoneSelectionMode::Toggle))
+	{
+		const bool bSelected = NewSelection.Contains(InBoneName);
+		if (bSelected)
+		{
+			NewSelection.Remove(InBoneName);
+		}
+		else
+		{
+			NewSelection.Add(InBoneName);
+		}
+	}
+	else
+	{
+		NewSelection.Empty();
+		if (InBoneName != NAME_None)
+		{
+			NewSelection.Add(InBoneName);
+		}
+	}
+
+	Selection = MoveTemp(NewSelection);
+	NormalizeSelection();
+}
+
+void USkeletonEditingTool::NormalizeSelection()
+{
+	if (Selection.Num() < 2)
+	{
+		return;
+	}
+	
+	const FReferenceSkeleton& ReferenceSkeleton = Modifier->GetReferenceSkeleton();
+
+	// ensure they are known and unique
+	TArray<int32> Indexes;
+	for (const FName& BoneName: Selection)
+	{
+		const int32 BoneIndex = ReferenceSkeleton.FindRawBoneIndex(BoneName);
+		if (BoneIndex != INDEX_NONE)
+		{
+			Indexes.AddUnique(BoneIndex);
+		}
+	}
+
+	if (Indexes.IsEmpty())
+	{
+		Selection.Reset();
+		return;
+	}
+
+	// sort them so that parents are placed after one of their children 
+	const TArray<FMeshBoneInfo>& BoneInfos = ReferenceSkeleton.GetRawRefBoneInfo();
+	Indexes.StableSort([&](const int32 Index0, const int32 Index1)
+	{
+		int32 P0 = BoneInfos[Index0].ParentIndex;
+		while (P0 != INDEX_NONE)
+		{
+			if (P0 == Index1)
+			{ // Index1 is a parent
+				return true;
+			}
+			if (Indexes.Contains(P0))
+			{ // parent is selected
+				return true;
+			}
+			P0 = BoneInfos[P0].ParentIndex;
+		}
+		return false;
+	});
+
+	// transform back to names
+	Selection.Reset();
+	Algo::Transform(Indexes, Selection, [&](const int32 Index)
+	{
+		return BoneInfos[Index].Name;
+	});	
+}
+
+void USkeletonEditingTool::OnClickRelease(const FInputDeviceRay& InReleasePos)
+{
+	TGuardValue OperationGuard(Operation, EEditingOperation::Transform);
+	EndChange();
+}
+
+void USkeletonEditingTool::OnTerminateDragSequence()
+{}
+
+void USkeletonEditingTool::OnTick(float DeltaTime)
+{
+	if (PendingFunction)
+	{
+		PendingFunction();
+		PendingFunction.Reset();
+	}
+	
+	SelectionWatcher.CheckAndUpdate();
+	CoordinateSystemWatcher.CheckAndUpdate();
+}
+
+FInputRayHit USkeletonEditingTool::CanBeginClickDragSequence(const FInputDeviceRay& InPressPos)
+{
+	PendingFunction.Reset();
+	ParentIndex = INDEX_NONE;
+	
+	FViewport* Viewport = GetToolManager()->GetContextQueriesAPI()->GetFocusedViewport();
+	if (!Viewport)
+	{
+		return FInputRayHit();
+	}
+	
+	if (GizmoWrapper && GizmoWrapper->IsGizmoHit(InPressPos))
+	{
+		return FInputRayHit();
+	}
+	
+	auto PickBone = [&]() -> int32
+	{
+		if (HHitProxy* HitProxy = Viewport->GetHitProxy(InPressPos.ScreenPosition.X, InPressPos.ScreenPosition.Y))
+		{
+			if (TOptional<FName> OptBoneName = GetBoneName(HitProxy))
+			{
+				const FReferenceSkeleton& ReferenceSkeleton = Modifier->GetReferenceSkeleton();
+				return ReferenceSkeleton.FindRawBoneIndex(*OptBoneName);
+			}
+		}
+		return INDEX_NONE;
+	};
+	
+	// pick bone in viewport
+	const int32 BoneIndex = PickBone();
+
+	// update projection properties
+	const FVector GlobalPosition = Modifier->GetTransform(BoneIndex, true).GetTranslation();
+	ProjectionProperties->UpdatePlane(*ViewContext, GlobalPosition);
+	
+	// if we picked a new bone
+	if (BoneIndex > INDEX_NONE)
+	{
+		// parent selection without changing the selection
+		if (Operation == EEditingOperation::Parent)
+		{
+			const FReferenceSkeleton& ReferenceSkeleton = Modifier->GetReferenceSkeleton();
+			ParentBones(ReferenceSkeleton.GetBoneName(BoneIndex));
+			return FInputRayHit();
+		}
+		
+		// otherwise, update current selection
+		PendingFunction = [this, BoneIndex]
+		{
+			const FReferenceSkeleton& ReferenceSkeleton = Modifier->GetReferenceSkeleton();
+			SelectBone(ReferenceSkeleton.GetBoneName(BoneIndex));
+
+			Properties->Name = GetCurrentBone();
+			Properties->Transform = ReferenceSkeleton.GetRefBonePose()[BoneIndex];
+			ParentIndex = ReferenceSkeleton.GetParentIndex(BoneIndex);
+		};
+	
+		return FInputRayHit(0.0);
+	}
+
+	// if we didn't pick anything
+	if (Operation == EEditingOperation::Select)
+	{
+		Selection.Empty();
+		Properties->Name = GetCurrentBone();
+		return FInputRayHit();
+	}
+
+	// if we're in creation mode then create a new bone
+	if (Operation == EEditingOperation::Create)
+	{
+		FVector HitPoint;
+		if (ProjectionProperties->GetProjectionPoint(InPressPos, HitPoint))
+		{
+			// CurrentBone is gonna be the parent
+			PendingFunction = [this, HitPoint]
+			{
+				const FReferenceSkeleton& ReferenceSkeleton = Modifier->GetReferenceSkeleton();
+				ParentIndex = ReferenceSkeleton.FindRawBoneIndex(GetCurrentBone());
+				const FTransform& ParentGlobalTransform = Modifier->GetTransform(ParentIndex, true);
+
+				// Create the new bone under mouse
+				Properties->Transform.SetLocation(ParentGlobalTransform.InverseTransformPosition(HitPoint));
+				CreateNewBone();
+			};
+			
+			return FInputRayHit(0.0);
+		}
+	}
+	
+	return FInputRayHit();
+}
+
+void USkeletonEditingTool::HandleSkeletalMeshModified(const TArray<FName>& InBoneNames, const ESkeletalMeshNotifyType InNotifyType)
+{
+	if (GetNotifier().Notifying())
+	{
+		return;
+	}
+
+	TArray<FName> BoneNames(InBoneNames);
+	const FReferenceSkeleton& RefSkeleton = Modifier->GetReferenceSkeleton();
+	BoneNames.RemoveAll([&](const FName& BoneName)
+	{
+		return RefSkeleton.FindRawBoneIndex(BoneName) == INDEX_NONE;
+	});
+
+	switch (InNotifyType)
+	{
+		case ESkeletalMeshNotifyType::BonesAdded:
+			Selection = BoneNames;
+			break;
+		case ESkeletalMeshNotifyType::BonesRemoved:
+			Selection.RemoveAll([&](const FName& BoneName)
+			{
+				return BoneNames.Contains(BoneName);
+			});
+			break;
+		case ESkeletalMeshNotifyType::BonesMoved:
+			break;
+		case ESkeletalMeshNotifyType::BonesSelected:
+			Selection.Reset();
+			Selection = BoneNames;
+			break;
+		case ESkeletalMeshNotifyType::BonesRenamed:
+			Selection = BoneNames;
+			break;
+		default:
+			break;
+	}
+
+	NormalizeSelection();
+	
+	Properties->Name = GetCurrentBone();
+}
+
+void USkeletonEditingTool::Render(IToolsContextRenderAPI* RenderAPI)
+{
+	// FIXME many things could be caches here and updated lazilly
+	if (!Target)
+	{
+		return;
+	}
+
+	static const FLinearColor DefaultBoneColor(0.0f,0.0f,0.025f,1.0f);
+	static const FLinearColor SelectedBoneColor(0.2f,1.0f,0.2f,1.0f);
+	static const FLinearColor AffectedBoneColor(1.0f,1.0f,1.0f,1.0f);
+	static const FLinearColor ParentOfSelectedBoneColor(0.85f,0.45f,0.12f,1.0f);
+	static FSkelDebugDrawConfig DrawConfig;
+		DrawConfig.BoneDrawMode = EBoneDrawMode::Type::All;
+		DrawConfig.BoneDrawSize = 1.f;
+		DrawConfig.bAddHitProxy = true;
+		DrawConfig.bForceDraw = false;
+		DrawConfig.DefaultBoneColor = DefaultBoneColor;
+		DrawConfig.AffectedBoneColor = AffectedBoneColor;
+		DrawConfig.SelectedBoneColor = SelectedBoneColor;
+		DrawConfig.ParentOfSelectedBoneColor = ParentOfSelectedBoneColor;
+		DrawConfig.AxisConfig.Thickness = Properties->AxisThickness;
+		DrawConfig.AxisConfig.Length = Properties->AxisLength;
+	
+	FPrimitiveDrawInterface* PDI = RenderAPI->GetPrimitiveDrawInterface();
+
+	const IPrimitiveComponentBackedTarget* TargetComponent = Cast<IPrimitiveComponentBackedTarget>(Target);
+	const FTransform ComponentTransform = TargetComponent->GetWorldTransform();
+	
+	const FReferenceSkeleton& RefSkeleton = Modifier->GetReferenceSkeleton();
+
+	const int32 NumBones = RefSkeleton.GetRawBoneNum();
+	TArray<TRefCountPtr<HHitProxy>> HitProxies; HitProxies.Reserve(NumBones);
+	TArray<FBoneIndexType> RequiredBones; RequiredBones.AddUninitialized(NumBones);
+	TArray<FTransform> WorldTransforms; WorldTransforms.AddUninitialized(NumBones);
+	TArray<FLinearColor> BoneColors; BoneColors.AddUninitialized(NumBones);
+	
+	for (int32 Index = 0; Index < NumBones; ++Index)
+	{
+		const FTransform& BoneTransform = Modifier->GetTransform(Index, true);
+		WorldTransforms[Index] = BoneTransform;
+		RequiredBones[Index] = Index;
+		BoneColors[Index] = DefaultBoneColor;
+		HitProxies.Add(new HBoneHitProxy(Index, RefSkeleton.GetBoneName(Index)));
+	}
+
+	SkeletalDebugRendering::DrawBones(
+		PDI,
+		ComponentTransform.GetLocation(),
+		RequiredBones,
+		RefSkeleton,
+		WorldTransforms,
+		GetSelectedBoneIndexes(),
+		BoneColors,
+		HitProxies,
+		DrawConfig
+	);
+}
+
+FBox USkeletonEditingTool::GetWorldSpaceFocusBox()
+{
+	const TArray<int32> BoneIndexes = GetSelectedBoneIndexes();
+	if (!BoneIndexes.IsEmpty())
+	{
+		FBox Box(EForceInit::ForceInit);
+		TSet<int32> AllChildren;
+
+		const FReferenceSkeleton& RefSkeleton = Modifier->GetReferenceSkeleton();
+		for (const int32 BoneIndex: BoneIndexes)
+		{
+			Box += Modifier->GetTransform(BoneIndex, true).GetTranslation();
+
+			// get direct children
+			TArray<int32> Children;
+			RefSkeleton.GetRawDirectChildBones(BoneIndex, Children);
+			AllChildren.Append(Children);
+		}
+
+		for (const int32 ChildIndex: AllChildren)
+		{
+			Box += Modifier->GetTransform(ChildIndex, true).GetTranslation();
+		}
+		
+		return Box;
+	}
+
+	if (PreviewMesh && PreviewMesh->GetActor())
+	{
+		return PreviewMesh->GetActor()->GetComponentsBoundingBox();
+	}
+
+	return USingleSelectionTool::GetWorldSpaceFocusBox();
+}
+
+TArray<FName> USkeletonEditingTool::GetSelectedBones() const
+{
+	return GetSelection();
+}
+
+const TArray<FName>& USkeletonEditingTool::GetSelection() const
+{
+	return Selection;
+}
+
+const FTransform& USkeletonEditingTool::GetTransform(const FName InBoneName, const bool bWorld) const
+{
+	if (Modifier)
+	{
+		const FReferenceSkeleton& RefSkeleton = Modifier->GetReferenceSkeleton();
+		return Modifier->GetTransform(RefSkeleton.FindRawBoneIndex(InBoneName), bWorld);
+	}
+	return FTransform::Identity;
+}
+
+void USkeletonEditingTool::SetTransforms(const TArray<FName>& InBones, const TArray<FTransform>& InTransforms, const bool bWorld) const
+{
+	const int32 NumBonesToMove = InBones.Num();
+	if (NumBonesToMove == 0 || NumBonesToMove != InTransforms.Num())
+	{
+		return;
+	}
+	
+	TArray<FTransform> NewTransforms = InTransforms;
+	if (bWorld)
+	{ // switch to local
+		const FReferenceSkeleton& RefSkeleton = Modifier->GetReferenceSkeleton();
+		const TArray<FMeshBoneInfo>& BoneInfos = RefSkeleton.GetRawRefBoneInfo();
+		
+		for (int32 Index = 0; Index < InBones.Num(); ++Index)
+		{
+			const int32 BoneIndex = RefSkeleton.FindRawBoneIndex(InBones[Index]);
+			check(BoneIndex != INDEX_NONE);
+
+			const int32 ParentBoneIndex = BoneInfos[BoneIndex].ParentIndex;
+			const int32 ParentToBeSet = ParentBoneIndex != INDEX_NONE ?
+				InBones.IndexOfByKey(BoneInfos[ParentBoneIndex].Name) : INDEX_NONE;
+
+			const FTransform& ParentGlobal = ParentToBeSet != INDEX_NONE ?
+				InTransforms[ParentToBeSet] : Modifier->GetTransform(ParentBoneIndex, true);
+			
+			NewTransforms[Index] = NewTransforms[Index].GetRelativeTransform(ParentGlobal);
+		}
+	}
+
+	const bool bBonesMoved = Modifier->SetBonesTransforms(InBones, NewTransforms, Properties->bUpdateChildren);
+	if (bBonesMoved)
+	{
+		UpdateGizmo();
+	}
+}
+
+FName USkeletonEditingTool::GetCurrentBone() const
+{
+	return Selection.IsEmpty() ? NAME_None : Selection.Last(); 
+}
+
+TArray<int32> USkeletonEditingTool::GetSelectedBoneIndexes() const
+{
+	TArray<int32> Indexes;
+	
+	const FReferenceSkeleton& ReferenceSkeleton = Modifier->GetReferenceSkeleton();
+	Algo::Transform(Selection, Indexes, [&](const FName& BoneName)
+	{
+		return ReferenceSkeleton.FindRawBoneIndex(BoneName);
+	});
+	
+	return Indexes;
+}
+
+void USkeletonEditingTool::BeginChange()
+{
+	if (Operation == EEditingOperation::Select)
+	{
+		return;
+	}
+	
+	ensure( ActiveChange == nullptr );
+	ActiveChange = MakeUnique<SkeletonEditingTool::FRefSkeletonChange>(this); 
+}
+
+void USkeletonEditingTool::EndChange()
+{
+	if (!ActiveChange.IsValid())
+	{
+		return;
+	}
+
+	if (Operation == EEditingOperation::Select)
+	{
+		return CancelChange();
+	}
+	
+	ActiveChange->StoreSkeleton(this);
+
+	static const UEnum* OperationEnum = StaticEnum<EEditingOperation>();
+	const FName OperationName = OperationEnum->GetNameByValue(static_cast<int64>(Operation));
+	const FText TransactionDesc = FText::Format(LOCTEXT("RefSkeletonChanged", "Skeleton Edit - {0}"), FText::FromName(OperationName));
+	
+	UInteractiveToolManager* ToolManager = GetToolManager();
+	ToolManager->BeginUndoTransaction(TransactionDesc);
+	ToolManager->EmitObjectChange(this, MoveTemp(ActiveChange), TransactionDesc);
+	ToolManager->EndUndoTransaction();
+}
+
+void USkeletonEditingTool::CancelChange()
+{
+	ActiveChange.Reset();
+}
+
+void USkeletonEditingProperties::Initialize(USkeletonEditingTool* ParentToolIn)
+{
+	ParentTool = ParentToolIn;
+}
+
+#if WITH_EDITOR
+void USkeletonEditingProperties::PostEditChangeProperty(FPropertyChangedEvent& PropertyChangedEvent)
+{
+	Super::PostEditChangeProperty(PropertyChangedEvent);
+	
+	if (PropertyChangedEvent.ChangeType == EPropertyChangeType::ValueSet)
+	{
+		const FName PropertyName = PropertyChangedEvent.GetMemberPropertyName();
+		if (PropertyName == GET_MEMBER_NAME_CHECKED(USkeletonEditingProperties, Name))
+		{
+			ParentTool->RenameBones();
+		}
+	}
+}
+#endif
+
+void UMirroringProperties::Initialize(USkeletonEditingTool* ParentToolIn)
+{
+	ParentTool = ParentToolIn;
+}
+
+void UMirroringProperties::MirrorBones()
+{
+	ParentTool->MirrorBones();
+}
+
+void UOrientingProperties::Initialize(USkeletonEditingTool* ParentToolIn)
+{
+	ParentTool = ParentToolIn;
+}
+
+void UOrientingProperties::OrientBones()
+{
+	ParentTool->OrientBones();
+}
+
+#if WITH_EDITOR
+
+void UOrientingProperties::PostEditChangeProperty(FPropertyChangedEvent &PropertyChangedEvent)
+{
+	Super::PostEditChangeProperty(PropertyChangedEvent);
+
+	if (PropertyChangedEvent.ChangeType == EPropertyChangeType::ValueSet)
+	{
+		auto CheckAxis = [&](const TEnumAsByte<EAxis::Type>& InRef, TEnumAsByte<EAxis::Type>& OutOther)
+		{
+			if (OutOther != InRef)
+			{
+				return;
+			}
+
+			switch (InRef)
+			{
+			case EAxis::X:
+				OutOther = EAxis::Y;
+				break;
+			case EAxis::Y:
+				OutOther = EAxis::Z;
+				break;
+			case EAxis::Z:
+				OutOther = EAxis::X;
+				break;
+			default:
+				break;
+			}
+		};
+		
+		const FName PropertyName = PropertyChangedEvent.GetPropertyName();
+		if (PropertyName == GET_MEMBER_NAME_CHECKED(FOrientOptions, Primary))
+		{
+			if (Options.Primary == EAxis::None)
+			{
+				Options.Primary = EAxis::X;
+				Options.Secondary = EAxis::Y;
+				return;
+			}
+			CheckAxis(Options.Primary, Options.Secondary);
+			return;
+		}
+		
+		if (PropertyName == GET_MEMBER_NAME_CHECKED(FOrientOptions, Secondary))
+		{
+			CheckAxis(Options.Secondary, Options.Primary);
+			return;
+		}
+	}
+}
+
+#endif
+
+void UProjectionProperties::Initialize(TObjectPtr<UPreviewMesh> InPreviewMesh) 
+{
+	PreviewMesh = InPreviewMesh;
+}
+
+void UProjectionProperties::UpdatePlane(const UGizmoViewContext& InViewContext, const FVector& InOrigin)
+{
+	PlaneNormal = -InViewContext.GetViewDirection();
+	PlaneOrigin = InOrigin;
+}
+
+bool UProjectionProperties::GetProjectionPoint(const FInputDeviceRay& InRay, FVector& OutHitPoint) const
+{
+	const FRay& WorldRay = InRay.WorldRay;
+
+	if (PreviewMesh.IsValid())
+	{
+		if (ProjectionType == EProjectionType::OnMesh)
+		{
+			FHitResult Hit;
+			if (PreviewMesh->FindRayIntersection(WorldRay, Hit))
+			{
+				OutHitPoint = WorldRay.Origin + WorldRay.Direction * Hit.Distance;
+				return true;
+			}
+		}
+
+		if (ProjectionType == EProjectionType::WithinMesh)
+		{
+			using namespace UE::Geometry;
+
+			if (FDynamicMeshAABBTree3* MeshAABBTree = PreviewMesh->GetSpatial())
+			{
+				using HitResult = MeshIntersection::FHitIntersectionResult;
+				TArray<HitResult> Hits;
+				
+				if (MeshAABBTree->FindAllHitTriangles(WorldRay, Hits))
+				{
+					if (Hits.Num() == 1)
+					{
+						OutHitPoint = WorldRay.Origin + WorldRay.Direction * Hits[0].Distance;
+						return true;						
+					}
+
+					// const double AverageDistance = Algo::Accumulate(Distances, 0.0) / static_cast<double>(Distances.Num());
+					// const int32 Index0 = Distances.IndexOfByPredicate([AverageDistance](double Distance) {return Distance <= AverageDistance;} );
+					// const int32 Index1 = Distances.IndexOfByPredicate([AverageDistance](double Distance) {return Distance >= AverageDistance;} );
+
+					static constexpr int32 Index0 = 0;
+					static constexpr int32 Index1 = 1;
+
+					const double d0 = Hits[Index0].Distance;
+					const double d1 = Hits[Index1].Distance;
+					OutHitPoint = WorldRay.Origin + WorldRay.Direction * ((d0+d1)*0.5);
+					return true;
+				}
+			}
+
+			FHitResult Hit;
+			if (PreviewMesh->FindRayIntersection(WorldRay, Hit))
+			{
+				OutHitPoint = WorldRay.Origin + WorldRay.Direction * Hit.Distance;
+				return true;
+			}
+		}
+	}
+	
+	// if ray is parallel to plane, nothing has been hit
+	if (FMath::IsNearlyZero(FVector::DotProduct(PlaneNormal, WorldRay.Direction)))
+	{
+		return false;
+	}
+
+	const FPlane Plane(PlaneOrigin, PlaneNormal);
+	const float HitDepth = FMath::RayPlaneIntersectionParam(WorldRay.Origin, WorldRay.Direction, Plane);
+	if (HitDepth < 0)
+	{
+		return false;
+	}
+
+	OutHitPoint = WorldRay.Origin + WorldRay.Direction * HitDepth;
+	
+	return true;
+}
+
+#undef LOCTEXT_NAMESPACE
\ No newline at end of file
diff --git a/Engine/Plugins/Experimental/MeshModelingToolsetExp/Source/MeshModelingToolsEditorOnlyExp/Private/SkeletalMesh/SkeletonTransformProxy.cpp b/Engine/Plugins/Experimental/MeshModelingToolsetExp/Source/MeshModelingToolsEditorOnlyExp/Private/SkeletalMesh/SkeletonTransformProxy.cpp
new file mode 100644
--- /dev/null
+++ b/Engine/Plugins/Experimental/MeshModelingToolsetExp/Source/MeshModelingToolsEditorOnlyExp/Private/SkeletalMesh/SkeletonTransformProxy.cpp
@@ -0,0 +1,49 @@
+﻿// Copyright Epic Games, Inc. All Rights Reserved.
+
+
+#include "SkeletalMesh/SkeletonTransformProxy.h"
+
+FTransform USkeletonTransformProxy::GetTransform() const
+{
+	return SharedTransform;
+}
+
+void USkeletonTransformProxy::Initialize(const FTransform& InTransform, const EToolContextCoordinateSystem& InTransformMode)
+{
+	Objects.Reset();
+	SharedTransform = InitialSharedTransform = InTransform;
+	TransformMode = InTransformMode;
+}
+
+void USkeletonTransformProxy::UpdateSharedTransform()
+{
+	// nothing to do
+}
+
+void USkeletonTransformProxy::UpdateObjects()
+{
+	const FTransform OffsetTransform = SharedTransform.GetRelativeTransform(InitialSharedTransform);
+	const FVector DeltaTranslation = OffsetTransform.GetTranslation();
+	const FQuat DeltaRotation = OffsetTransform.GetRotation();
+	const bool bApplyLocal = TransformMode == EToolContextCoordinateSystem::Local;
+	
+	for (FRelativeObject& Object : Objects)
+	{
+		if (bApplyLocal)
+		{
+			Object.SetTransformFunc(OffsetTransform * Object.StartTransform);
+		}
+		else
+		{
+			FTransform Transform = Object.StartTransform;
+			Transform.SetRotation(DeltaRotation * Transform.GetRotation());
+			Transform.SetTranslation(DeltaTranslation + Transform.GetTranslation());
+			Object.SetTransformFunc(Transform);
+		}
+	}
+}
+
+bool USkeletonTransformProxy::IsValid() const
+{
+	return !Objects.IsEmpty();
+}
\ No newline at end of file
diff --git a/Engine/Plugins/Experimental/MeshModelingToolsetExp/Source/MeshModelingToolsEditorOnlyExp/Private/SkeletalMesh/SkinWeightsBindingTool.cpp b/Engine/Plugins/Experimental/MeshModelingToolsetExp/Source/MeshModelingToolsEditorOnlyExp/Private/SkeletalMesh/SkinWeightsBindingTool.cpp
new file mode 100644
--- /dev/null
+++ b/Engine/Plugins/Experimental/MeshModelingToolsetExp/Source/MeshModelingToolsEditorOnlyExp/Private/SkeletalMesh/SkinWeightsBindingTool.cpp
@@ -0,0 +1,371 @@
+// Copyright Epic Games, Inc. All Rights Reserved.
+
+#include "SkeletalMesh/SkinWeightsBindingTool.h"
+
+#include "BoneWeights.h"
+#include "DynamicMeshToMeshDescription.h"
+#include "DynamicMesh/DynamicVertexSkinWeightsAttribute.h"
+#include "InteractiveToolManager.h"
+#include "MeshDescriptionToDynamicMesh.h"
+#include "MeshOpPreviewHelpers.h"
+#include "SkeletalMeshAttributes.h"
+#include "ToolSetupUtil.h"
+#include "ToolTargetManager.h"
+#include "Components/SkeletalMeshComponent.h"
+#include "Engine/SkeletalMesh.h"
+#include "SceneManagement.h"
+#include "Spatial/MeshWindingNumberGrid.h"
+
+#include "ModelingToolTargetUtil.h"
+#include "SkinningOps/SkinBindingOp.h"
+#include "Spatial/OccupancyGrid3.h"
+
+#include UE_INLINE_GENERATED_CPP_BY_NAME(SkinWeightsBindingTool)
+
+
+#define LOCTEXT_NAMESPACE "USkinWeightsBindingTool"
+
+
+bool USkinWeightsBindingToolBuilder::CanBuildTool(const FToolBuilderState& SceneState) const
+{
+	return SceneState.TargetManager->CountSelectedAndTargetable(SceneState, FToolTargetTypeRequirements()) == 1;
+}
+
+
+UMultiSelectionMeshEditingTool* USkinWeightsBindingToolBuilder::CreateNewTool(const FToolBuilderState& SceneState) const
+{
+	return NewObject<USkinWeightsBindingTool>(SceneState.ToolManager);
+}
+
+
+USkeleton* USkinWeightsBindingToolProperties::GetSkeleton(bool& bInvalidSkeletonIsError, const IPropertyHandle* PropertyHandle)
+{
+	bInvalidSkeletonIsError = false;
+	return SkeletalMesh ? SkeletalMesh->GetSkeleton() : nullptr;
+}
+
+
+USkinWeightsBindingTool::USkinWeightsBindingTool()
+{
+	Properties = CreateDefaultSubobject<USkinWeightsBindingToolProperties>(TEXT("SkinWeightsBindingProperties"));
+	// CreateDefaultSubobject automatically sets RF_Transactional flag, we need to clear it so that undo/redo doesn't affect tool properties
+	Properties->ClearFlags(RF_Transactional);
+}
+
+
+USkinWeightsBindingTool::~USkinWeightsBindingTool()
+{
+}
+
+
+void USkinWeightsBindingTool::Setup()
+{
+	Super::Setup();
+
+	if (ensure(Properties))
+	{
+		Properties->RestoreProperties(this);
+	}
+
+	if (!ensure(Targets.Num() > 0) || !ensure(Targets[0]))
+	{
+		return;
+	}
+	
+	const USkeletalMeshComponent* SkeletalMeshComponent = Cast<USkeletalMeshComponent>(UE::ToolTarget::GetTargetComponent(Targets[0]));
+	
+	if (SkeletalMeshComponent && SkeletalMeshComponent->GetSkeletalMeshAsset())
+	{
+		USkeletalMesh* SkeletalMesh = SkeletalMeshComponent->GetSkeletalMeshAsset();
+
+		ReferenceSkeleton = SkeletalMesh->GetRefSkeleton();
+		
+		// Initialize the bone browser
+		const UE::Anim::FCurveFilterSettings CurveFilterSettings = SkeletalMeshComponent->GetCurveFilterSettings(0 /* Always use the highest LOD */);
+		TArray<FBoneIndexType> AllBones;
+		AllBones.Reserve(ReferenceSkeleton.GetRawBoneNum());
+		for (int32 BoneIndex = 0; BoneIndex < ReferenceSkeleton.GetRawBoneNum(); BoneIndex++)
+		{
+			AllBones.Add(static_cast<FBoneIndexType>(BoneIndex));
+		}
+		BoneContainer.InitializeTo(AllBones, CurveFilterSettings, *SkeletalMesh);
+
+		Properties->SkeletalMesh = SkeletalMesh;
+		Properties->CurrentBone.Initialize(BoneContainer);
+	}
+	
+	UE::ToolTarget::HideSourceObject(Targets[0]);
+	Preview = NewObject<UMeshOpPreviewWithBackgroundCompute>(this, "Preview");
+	Preview->Setup(GetTargetWorld(), this);
+	Preview->SetIsMeshTopologyConstant(true, EMeshRenderAttributeFlags::VertexColors);
+	Preview->OnMeshUpdated.AddLambda([this](UMeshOpPreviewWithBackgroundCompute* Compute)
+	{
+		UpdateVisualization();
+	});
+	
+	FComponentMaterialSet MaterialSet = UE::ToolTarget::GetMaterialSet(Targets[0]);
+
+	UMaterialInterface* VtxColorMaterial = GetToolManager()->GetContextQueriesAPI()->GetStandardMaterial(EStandardToolContextMaterials::VertexColorMaterial);
+	if (VtxColorMaterial != nullptr)
+	{
+		for (UMaterialInterface*& Material: MaterialSet.Materials)
+		{
+			Material = VtxColorMaterial;
+		}
+	}
+	
+	Preview->ConfigureMaterials( MaterialSet.Materials,
+								 ToolSetupUtil::GetDefaultWorkingMaterial(GetToolManager())
+	);
+
+	// setup watchers
+	auto HandlePropertyChange = [this](const bool bUpdateOperator)
+	{
+		if (bUpdateOperator)
+		{
+			Preview->InvalidateResult();
+		}
+		UpdateVisualization();
+	};
+
+	Properties->WatchProperty(Properties->CurrentBone.BoneName,
+							  [HandlePropertyChange](FName) { HandlePropertyChange(false); });
+	Properties->WatchProperty(Properties->BindingType,
+							  [HandlePropertyChange](ESkinWeightsBindType) { HandlePropertyChange(true); });
+	Properties->WatchProperty(Properties->Stiffness,
+							  [HandlePropertyChange](float) { HandlePropertyChange(true); });
+	Properties->WatchProperty(Properties->MaxInfluences,
+							  [HandlePropertyChange](int32) { HandlePropertyChange(true); });
+	Properties->WatchProperty(Properties->VoxelResolution,
+							  [this](int32)
+							  {
+							  	Occupancy = MakeShared<UE::Geometry::FOccupancyGrid3>(*OriginalMesh, Properties->VoxelResolution);
+								Preview->InvalidateResult();
+								UpdateVisualization();
+							  });
+
+	OriginalMesh = MakeShared<FDynamicMesh3, ESPMode::ThreadSafe>();
+	FMeshDescriptionToDynamicMesh Converter;
+	Converter.Convert(UE::ToolTarget::GetMeshDescription(Targets[0]), *OriginalMesh);
+
+	// Enable or override vertex colors on the original mesh.
+	OriginalMesh->EnableAttributes();
+	OriginalMesh->Attributes()->DisablePrimaryColors();
+	OriginalMesh->Attributes()->EnablePrimaryColors();
+	// Create an overlay that has no split elements, init with zero value.
+	OriginalMesh->Attributes()->PrimaryColors()->CreateFromPredicate([](int ParentVID, int TriIDA, int TriIDB){return true;}, 0.f);
+
+	Preview->PreviewMesh->SetTransform((FTransform) UE::ToolTarget::GetLocalToWorldTransform(Targets[0]));
+	Preview->PreviewMesh->SetTangentsMode(EDynamicMeshComponentTangentsMode::AutoCalculated);
+	Preview->PreviewMesh->SetShadowsEnabled(false);
+	Preview->PreviewMesh->UpdatePreview(OriginalMesh.Get());
+
+	Occupancy = MakeShared<UE::Geometry::FOccupancyGrid3>(*OriginalMesh, Properties->VoxelResolution);
+	
+	UpdateVisualization(/*bForce=*/true);
+	
+	// add properties to GUI
+	AddToolPropertySource(Properties);
+
+	Preview->InvalidateResult();
+
+	SetToolDisplayName(LOCTEXT("ToolName", "Bind Skin"));
+	GetToolManager()->DisplayMessage(
+		LOCTEXT("OnStartTool", "Creates a rigid binding for the skin weights."),
+		EToolMessageLevel::UserNotification);
+}
+
+
+void USkinWeightsBindingTool::OnShutdown(EToolShutdownType ShutdownType)
+{
+	Properties->SaveProperties(this);
+
+	UE::ToolTarget::ShowSourceObject(Targets[0]);
+	
+	const FDynamicMeshOpResult Result = Preview->Shutdown();
+	if (ShutdownType == EToolShutdownType::Accept)
+	{
+		GenerateAsset(Result);
+	}
+}
+
+
+void USkinWeightsBindingTool::OnTick(float DeltaTime)
+{
+	Preview->Tick(DeltaTime);
+}
+
+static void DrawBox(IToolsContextRenderAPI* RenderAPI, const FTransform& Transform, const FBox &Box, const FLinearColor &Color, float LineThickness)
+{
+	FPrimitiveDrawInterface* PDI = RenderAPI->GetPrimitiveDrawInterface();
+	const float PDIScale = RenderAPI->GetCameraState().GetPDIScalingFactor();
+
+	const FVector Corners[2] = {
+		Transform.TransformPosition(Box.Min),
+		Transform.TransformPosition(Box.Max)
+	};
+
+	static constexpr UE::Geometry::FVector3i Offsets[12][2] =
+	{
+		// Bottom
+		{{0, 0, 0}, {1, 0, 0}},
+		{{1, 0, 0}, {1, 1, 0}},
+		{{1, 1, 0}, {0, 1, 0}},
+		{{0, 1, 0}, {0, 0, 0}},
+		
+		// Top
+		{{0, 0, 1}, {1, 0, 1}},
+		{{1, 0, 1}, {1, 1, 1}},
+		{{1, 1, 1}, {0, 1, 1}},
+		{{0, 1, 1}, {0, 0, 1}},
+		
+		// Sides
+		{{0, 0, 0}, {0, 0, 1}},
+		{{1, 0, 0}, {1, 0, 1}},
+		{{1, 1, 0}, {1, 1, 1}},
+		{{0, 1, 0}, {0, 1, 1}},
+	}; 
+
+	for (int32 Index = 0; Index < 12; Index++)
+	{
+		const UE::Geometry::FVector3i* LineOffsets = Offsets[Index];
+		FVector  A(Corners[LineOffsets[0].X].X, Corners[LineOffsets[0].Y].Y, Corners[LineOffsets[0].Z].Z);
+		FVector  B(Corners[LineOffsets[1].X].X, Corners[LineOffsets[1].Y].Y, Corners[LineOffsets[1].Z].Z);
+		
+		PDI->DrawTranslucentLine(A, B, Color, 1, LineThickness * PDIScale);
+	}
+}
+
+
+void USkinWeightsBindingTool::Render(IToolsContextRenderAPI* RenderAPI)
+{
+	using namespace UE::Geometry;
+
+	if (Occupancy && Properties->bDebugDraw)
+	{
+		constexpr bool bShowInterior = false;
+		constexpr bool bShowBoundary = true;
+
+		FPrimitiveDrawInterface* PDI = RenderAPI->GetPrimitiveDrawInterface();
+		const FTransform Transform = (FTransform) UE::ToolTarget::GetLocalToWorldTransform(Targets[0]);
+		float PDIScale = RenderAPI->GetCameraState().GetPDIScalingFactor();
+		const TDenseGrid3<FOccupancyGrid3::EDomain>& OccupancyGrid = Occupancy->GetOccupancyStateGrid(); 
+
+		for (int32 Index = 0; Index < OccupancyGrid.Size(); Index++)
+		{
+			FVector3i OccupancyIndex(OccupancyGrid.ToIndex(Index));
+
+			const FOccupancyGrid3::EDomain Domain = OccupancyGrid[OccupancyIndex];
+			if (bShowBoundary && Domain == FOccupancyGrid3::EDomain::Boundary)
+			{
+				FBox Box{Occupancy->GetCellBoxFromIndex(OccupancyIndex)};
+				DrawBox(RenderAPI, Transform, Box, FLinearColor(1.0, 1.0, 0.0, 0.5), 0.5f);
+			}
+			
+		}
+	}
+	
+}
+
+
+bool USkinWeightsBindingTool::CanAccept() const
+{
+	return Super::CanAccept() && Preview->HaveValidResult();
+}
+
+TUniquePtr<UE::Geometry::FDynamicMeshOperator> USkinWeightsBindingTool::MakeNewOperator()
+{
+	TUniquePtr<UE::Geometry::FSkinBindingOp> Op = MakeUnique<UE::Geometry::FSkinBindingOp>();
+
+	Op->ProfileName = FSkeletalMeshAttributes::DefaultSkinWeightProfileName;
+	Op->BindType = static_cast<UE::Geometry::ESkinBindingType>(Properties->BindingType);
+	Op->Stiffness = Properties->Stiffness;
+	Op->MaxInfluences = Properties->MaxInfluences;
+	Op->VoxelResolution = Properties->VoxelResolution;
+	
+	Op->OriginalMesh = OriginalMesh;
+	Op->SetTransformHierarchyFromReferenceSkeleton(ReferenceSkeleton);
+
+	const FTransform3d LocalToWorld = UE::ToolTarget::GetLocalToWorldTransform(Targets[0]);
+	Op->SetResultTransform(LocalToWorld);
+	
+	return Op;
+}
+
+
+void USkinWeightsBindingTool::GenerateAsset(const FDynamicMeshOpResult& Result)
+{
+	// TODO: Update FDynamicMeshToMeshDescription to allow update the skin weights only.
+	GetToolManager()->BeginUndoTransaction(LOCTEXT("SkinWeightsBindingToolTransactionName", "Create Rigid Binding"));
+
+	check(Result.Mesh.Get() != nullptr);
+	UE::ToolTarget::CommitMeshDescriptionUpdateViaDynamicMesh(Targets[0], *Result.Mesh.Get(), true);
+
+	GetToolManager()->EndUndoTransaction();	
+}
+
+FVector4f USkinWeightsBindingTool::WeightToColor(float Value)
+{
+	Value = FMath::Clamp(Value, 0.0f, 1.0f);
+
+	{
+		// A close approximation of the skeletal mesh editor's bone weight ramp. 
+		const FLinearColor HSV((1.0f - Value) * 285.0f, 100.0f, 85.0f);
+		return UE::Geometry::ToVector4<float>(HSV.HSVToLinearRGB());
+	}
+}
+
+void USkinWeightsBindingTool::UpdateVisualization(bool bInForce)
+{
+	using namespace UE::AnimationCore;
+	using namespace UE::Geometry;
+
+	const int32 RawBoneIndex = ReferenceSkeleton.FindRawBoneIndex(Properties->CurrentBone.BoneName);
+	if ((bInForce || Preview->HaveValidNonEmptyResult()) && RawBoneIndex != INDEX_NONE)
+	{
+		const FBoneIndexType BoneIndex = static_cast<FBoneIndexType>(RawBoneIndex);
+
+		// update mesh with new value colors
+		Preview->PreviewMesh->EditMesh([&](FDynamicMesh3& InMesh)
+		{
+			const FDynamicMeshVertexSkinWeightsAttribute* SkinWeights = InMesh.Attributes()->GetSkinWeightsAttribute(FSkeletalMeshAttributes::DefaultSkinWeightProfileName);
+			FDynamicMeshColorOverlay* ColorOverlay = InMesh.Attributes()->PrimaryColors();
+
+			if (!ColorOverlay)
+			{
+				InMesh.EnableAttributes();
+				InMesh.Attributes()->EnablePrimaryColors();
+				// Create an overlay that has no split elements, init with zero value.
+				ColorOverlay = InMesh.Attributes()->PrimaryColors();
+				ColorOverlay->CreateFromPredicate([](int /*ParentVID*/, int /*TriIDA*/, int /*TriIDB*/){return true;}, 0.f);
+			}
+			
+			FBoneWeights BoneWeights;
+			
+			for (const int32 ElementId : ColorOverlay->ElementIndicesItr())
+			{
+				const int32 VertexId = ColorOverlay->GetParentVertex(ElementId);
+				SkinWeights->GetValue(VertexId, BoneWeights);
+
+				float Weight = 0.0f;
+				for (const FBoneWeight BoneWeight: BoneWeights)
+				{
+					if (BoneWeight.GetBoneIndex() == BoneIndex)
+					{
+						Weight = BoneWeight.GetWeight();
+						break;
+					}
+				}
+				
+				const FVector4f Color(WeightToColor(Weight));
+				ColorOverlay->SetElement(ElementId, Color);
+			}
+		});
+	}
+}
+
+
+#undef LOCTEXT_NAMESPACE
+
+// #pragma optimize( "", on )
+
diff --git a/Engine/Plugins/Experimental/MeshModelingToolsetExp/Source/MeshModelingToolsEditorOnlyExp/Private/SkeletalMesh/SkinWeightsPaintTool.cpp b/Engine/Plugins/Experimental/MeshModelingToolsetExp/Source/MeshModelingToolsEditorOnlyExp/Private/SkeletalMesh/SkinWeightsPaintTool.cpp
new file mode 100644
--- /dev/null
+++ b/Engine/Plugins/Experimental/MeshModelingToolsetExp/Source/MeshModelingToolsEditorOnlyExp/Private/SkeletalMesh/SkinWeightsPaintTool.cpp
@@ -0,0 +1,1804 @@
+// Copyright Epic Games, Inc. All Rights Reserved.
+
+#include "SkeletalMesh/SkinWeightsPaintTool.h"
+#include "Engine/SkeletalMesh.h"
+#include "InteractiveToolManager.h"
+#include "SkeletalMeshAttributes.h"
+#include "Math/UnrealMathUtility.h"
+#include "Components/SkeletalMeshComponent.h"
+
+#include "TargetInterfaces/PrimitiveComponentBackedTarget.h"
+#include "ModelingToolTargetUtil.h"
+
+#include "MeshDescription.h"
+#include "MeshModelingToolsEditorOnlyExp.h"
+#include "PointSetAdapter.h"
+#include "Animation/MirrorDataTable.h"
+#include "DynamicMesh/NonManifoldMappingSupport.h"
+#include "Parameterization/MeshLocalParam.h"
+#include "Spatial/FastWinding.h"
+#include "Async/Async.h"
+#include "BaseGizmos/BrushStampIndicator.h"
+#include "DynamicMesh/MeshAdapterUtil.h"
+#include "Selection/PolygonSelectionMechanic.h"
+#include "Spatial/PointSetHashTable.h"
+#include "Util/ColorConstants.h"
+
+#include UE_INLINE_GENERATED_CPP_BY_NAME(SkinWeightsPaintTool)
+
+#define LOCTEXT_NAMESPACE "USkinWeightsPaintTool"
+
+using namespace SkinPaintTool;
+
+// thread pool to use for async operations
+static EAsyncExecution SkinPaintToolAsyncExecTarget = EAsyncExecution::ThreadPool;
+
+// any weight below this value is ignored, since it won't be representable in unsigned 16-bit precision
+constexpr float MinimumWeightThreshold = 1.0f / 65535.0f;
+
+USkinWeightsPaintToolProperties::USkinWeightsPaintToolProperties()
+{
+	ColorRamp.Add(FLinearColor::Blue);
+	ColorRamp.Add(FLinearColor::Yellow);
+
+	MinColor = FLinearColor::Black;
+	MaxColor = FLinearColor::White;
+}
+
+void FMultiBoneWeightEdits::MergeSingleEdit(
+	const int32 BoneIndex,
+	const int32 VertexID,
+	const float OldWeight,
+	const float NewWeight)
+{
+	FSingleBoneWeightEdits& BoneWeightEdit = PerBoneWeightEdits.FindOrAdd(BoneIndex);
+	BoneWeightEdit.BoneIndex = BoneIndex;
+	BoneWeightEdit.NewWeights.Add(VertexID, NewWeight);
+	BoneWeightEdit.OldWeights.FindOrAdd(VertexID, OldWeight);
+}
+
+void FMultiBoneWeightEdits::MergeEdits(const FSingleBoneWeightEdits& BoneWeightEdits)
+{
+	// make sure bone has an entry in the map of weight edits
+	const int32 BoneIndex = BoneWeightEdits.BoneIndex;
+	PerBoneWeightEdits.FindOrAdd(BoneIndex);
+	PerBoneWeightEdits[BoneIndex].BoneIndex = BoneIndex;
+	
+	for (const TTuple<int32, float>& NewWeight : BoneWeightEdits.NewWeights)
+	{
+		int32 VertexIndex = NewWeight.Key;
+		PerBoneWeightEdits[BoneIndex].NewWeights.Add(VertexIndex, NewWeight.Value);
+		PerBoneWeightEdits[BoneIndex].OldWeights.FindOrAdd(VertexIndex, BoneWeightEdits.OldWeights[VertexIndex]);
+	}
+}
+
+float FMultiBoneWeightEdits::GetVertexDeltaFromEdits(const int32 BoneIndex, const int32 VertexIndex)
+{
+	PerBoneWeightEdits.FindOrAdd(BoneIndex);
+	if (const float* NewVertexWeight = PerBoneWeightEdits[BoneIndex].NewWeights.Find(VertexIndex))
+	{
+		return *NewVertexWeight - PerBoneWeightEdits[BoneIndex].OldWeights[VertexIndex];
+	}
+
+	return 0.0f;
+}
+
+void FSkinToolDeformer::Initialize(const USkeletalMeshComponent* SkeletalMeshComponent, const FMeshDescription* Mesh)
+{
+	// get all bone transforms in the reference pose store a copy in component space
+	Component = SkeletalMeshComponent;
+	const FReferenceSkeleton& RefSkeleton = Component->GetSkeletalMeshAsset()->GetRefSkeleton();
+	const TArray<FTransform> &LocalSpaceBoneTransforms = RefSkeleton.GetRefBonePose();
+	const int32 NumBones = LocalSpaceBoneTransforms.Num();
+	InvCSRefPoseTransforms.SetNumUninitialized(NumBones);
+	for (int32 BoneIndex=0; BoneIndex<NumBones; ++BoneIndex)
+	{
+		const int32 ParentBoneIndex = RefSkeleton.GetParentIndex(BoneIndex);
+		const FTransform& LocalTransform = LocalSpaceBoneTransforms[BoneIndex];
+		if (ParentBoneIndex != INDEX_NONE)
+		{
+			InvCSRefPoseTransforms[BoneIndex] = LocalTransform * InvCSRefPoseTransforms[ParentBoneIndex];
+		}
+		else
+		{
+			InvCSRefPoseTransforms[BoneIndex] = LocalTransform;
+		}
+	}
+	
+	for (int32 BoneIndex=0; BoneIndex<NumBones; ++BoneIndex)
+	{
+		// pre-invert the transforms so we don't have to at runtime
+		InvCSRefPoseTransforms[BoneIndex] = InvCSRefPoseTransforms[BoneIndex].Inverse();
+
+		// store map of bone indices to bone names
+		FName BoneName = RefSkeleton.GetBoneName(BoneIndex);
+		BoneNames.Add(BoneName);
+		BoneNameToIndexMap.Add(BoneName, BoneIndex);
+	}
+
+	// store reference pose vertex positions
+	const TArrayView<const UE::Math::TVector<float>> VertexPositions = Mesh->GetVertexPositions().GetRawArray();
+	RefPoseVertexPositions = VertexPositions;
+
+	// force all vertices to be updated initially
+	VerticesWithModifiedWeights.Empty(RefPoseVertexPositions.Num());
+	for (int32 VertexID=0; VertexID<RefPoseVertexPositions.Num(); ++VertexID)
+	{
+		VerticesWithModifiedWeights.Add(VertexID);
+	}
+}
+
+void FSkinToolDeformer::UpdateVertexDeformation(USkinWeightsPaintTool* Tool)
+{
+	TRACE_CPUPROFILER_EVENT_SCOPE(SkinTool::UpdateDeformationTotal);
+
+	if (VerticesWithModifiedWeights.IsEmpty())
+	{
+		return;
+	}
+	
+	// update vertex positions
+	UPreviewMesh* PreviewMesh = Tool->PreviewMesh;
+	const TArray<VertexWeights>& CurrentWeights = Tool->Weights.CurrentWeights;
+	PreviewMesh->DeferredEditMesh([this, &CurrentWeights](FDynamicMesh3& Mesh)
+	{
+		TRACE_CPUPROFILER_EVENT_SCOPE(SkinTool::UpdateDeformation);
+		const TArray<FTransform>& CurrentBoneTransforms = Component->GetComponentSpaceTransforms();
+		const TArray<int32> VertexIndices = VerticesWithModifiedWeights.Array();
+		
+		ParallelFor( VerticesWithModifiedWeights.Num(), [this, &VertexIndices, &Mesh, &CurrentBoneTransforms, &CurrentWeights](int32 Index)
+		{
+			const int32 VertexID = VertexIndices[Index];
+			FVector VertexNewPosition = FVector::ZeroVector;
+			const VertexWeights& VertexPerBoneData = CurrentWeights[VertexID];
+			for (const FVertexBoneWeight& VertexData : VertexPerBoneData)
+			{
+				const FTransform& CurrentTransform = CurrentBoneTransforms[VertexData.BoneIndex];
+				VertexNewPosition += CurrentTransform.TransformPosition(VertexData.VertexInBoneSpace) * VertexData.Weight;
+			}
+			
+			Mesh.SetVertex(VertexID, VertexNewPosition, false);
+		});
+	}, false);
+	PreviewMesh->NotifyDeferredEditCompleted(UPreviewMesh::ERenderUpdateMode::FastUpdate, EMeshRenderAttributeFlags::Positions, false);
+
+	// update vertex acceleration structure
+	{
+		TRACE_CPUPROFILER_EVENT_SCOPE(SkinTool::UpdateVertexOctree);
+		Tool->VerticesOctree.RemoveVertices(VerticesWithModifiedWeights);
+		Tool->VerticesOctree.InsertVertices(VerticesWithModifiedWeights);
+	}
+
+	// update triangle acceleration structure
+	{
+		TRACE_CPUPROFILER_EVENT_SCOPE(SkinTool::UpdateTriangleOctree);
+
+		// ensure previous async update is finished before queuing the next one...
+		Tool->TriangleOctreeFuture.Wait();
+
+		// reusable buffer of triangles to reinsert into the octree
+		TArray<int32>& TrianglesToReinsert = Tool->TrianglesToReinsert;
+		TrianglesToReinsert.Reset();
+
+		// reinsert all triangles containing an updated vertex
+		const FDynamicMesh3* DynamicMesh = PreviewMesh->GetMesh();
+		for (const int32 TriangleID : DynamicMesh->TriangleIndicesItr())
+		{
+			UE::Geometry::FIndex3i TriVerts = DynamicMesh->GetTriangle(TriangleID);
+			bool bIsTriangleAffected = VerticesWithModifiedWeights.Contains(TriVerts[0]);
+			bIsTriangleAffected = VerticesWithModifiedWeights.Contains(TriVerts[1]) ? true : bIsTriangleAffected;
+			bIsTriangleAffected = VerticesWithModifiedWeights.Contains(TriVerts[2]) ? true : bIsTriangleAffected;
+			if (bIsTriangleAffected)
+			{
+				TrianglesToReinsert.Add(TriangleID);
+			}
+		}
+
+		// asynchronously update the octree, this normally finishes well before the next update
+		// but in the unlikely event that it does not, it would result in a frame where the paint brush
+		// is not perfectly aligned with the mesh; not a deal breaker.
+		UE::Geometry::FDynamicMeshOctree3& OctreeToUpdate = Tool->TrianglesOctree;
+		Tool->TriangleOctreeFuture = Async(SkinPaintToolAsyncExecTarget, [&OctreeToUpdate, &TrianglesToReinsert]()
+		{
+			TRACE_CPUPROFILER_EVENT_SCOPE(SkinTool::TriangleOctreeReinsert);	
+			OctreeToUpdate.ReinsertTriangles(TrianglesToReinsert);
+		});
+	}
+
+	// empty queue of vertices to update
+	VerticesWithModifiedWeights.Reset();
+}
+
+void FSkinToolDeformer::SetVertexNeedsUpdated(int32 VertexIndex)
+{
+	VerticesWithModifiedWeights.Add(VertexIndex);
+}
+
+void FSkinToolWeights::InitializeSkinWeights(
+	const USkeletalMeshComponent* SkeletalMeshComponent,
+	FMeshDescription* Mesh)
+{
+	// initialize deformer data
+	Deformer.Initialize(SkeletalMeshComponent, Mesh);
+
+	// initialize current weights (using compact format: num_verts * max_influences)
+	const FSkeletalMeshConstAttributes MeshAttribs(*Mesh);
+	const FSkinWeightsVertexAttributesConstRef VertexSkinWeights = MeshAttribs.GetVertexSkinWeights();
+	const int32 NumVertices = Mesh->Vertices().Num();
+	CurrentWeights.SetNum(NumVertices);
+	for (int32 VertexIndex = 0; VertexIndex < NumVertices; VertexIndex++)
+	{
+		const FVertexID VertexID(VertexIndex);
+		int32 InfluenceIndex = 0;
+		for (UE::AnimationCore::FBoneWeight BoneWeight: VertexSkinWeights.Get(VertexID))
+		{
+			check(InfluenceIndex < MAX_TOTAL_INFLUENCES);
+			const int32 BoneIndex = BoneWeight.GetBoneIndex();
+			const float Weight = BoneWeight.GetWeight();
+			const FVector& RefPoseVertexPosition = Deformer.RefPoseVertexPositions[VertexIndex];
+			const FTransform& InvRefPoseTransform = Deformer.InvCSRefPoseTransforms[BoneIndex];
+			const FVector& BoneLocalPositionInRefPose = InvRefPoseTransform.TransformPosition(RefPoseVertexPosition);
+			CurrentWeights[VertexIndex].Emplace(BoneIndex, BoneLocalPositionInRefPose, Weight);
+			++InfluenceIndex;
+		}
+	}
+	
+	// maintain duplicate weight map
+	PreChangeWeights = CurrentWeights;
+
+	// maintain relax-per stroke map
+	MaxFalloffPerVertexThisStroke.SetNumZeroed(NumVertices);
+}
+
+void FSkinToolWeights::EditVertexWeightAndNormalize(
+	const int32 BoneToHoldIndex,
+	const int32 VertexID,
+	float NewWeightValue,
+	FMultiBoneWeightEdits& WeightEdits)
+{
+	// clamp new weight
+	NewWeightValue = FMath::Clamp(NewWeightValue, 0.0f, 1.0f);
+	
+	// calculate the sum of all the weights on this vertex (not including the one we currently applied)
+	TArray<int32> BonesAffectingVertex;
+	TArray<float> ValuesToNormalize;
+	float Total = 0.0f;
+	const VertexWeights& VertexData = CurrentWeights[VertexID];
+	for (const FVertexBoneWeight& VertexBoneData : VertexData)
+	{
+		if (VertexBoneData.BoneIndex == BoneToHoldIndex)
+		{
+			continue;
+		}
+		
+		if (VertexBoneData.Weight < MinimumWeightThreshold)
+		{
+			continue;
+		}
+		
+		BonesAffectingVertex.Add(VertexBoneData.BoneIndex);
+		ValuesToNormalize.Add(VertexBoneData.Weight);
+		Total += VertexBoneData.Weight;
+	}
+
+	// are there no OTHER influences on this vertex?
+	const bool bVertexHasNoOtherInfluences = Total <= MinimumWeightThreshold;
+	
+	// if user applied any weight to this vertex AND there's no other weights of any significance,
+	// then simply apply full weight to this vertex, set all other influences to zero and return
+	const bool bApplyWeightToThisVertex = NewWeightValue >= MinimumWeightThreshold;
+	if (bApplyWeightToThisVertex && bVertexHasNoOtherInfluences)
+	{
+		// set all other influences to 0.0f
+		for (int32 i=0; i<ValuesToNormalize.Num(); ++i)
+		{
+			const int32 BoneIndex = BonesAffectingVertex[i];
+			const float OldWeight = ValuesToNormalize[i];
+			constexpr float NewWeight = 0.0f;
+			WeightEdits.MergeSingleEdit(BoneIndex, VertexID, OldWeight, NewWeight);
+		}
+
+		// set current bone value to 1.0f
+		const float PrevWeight = GetWeightOfBoneOnVertex(BoneToHoldIndex, VertexID, PreChangeWeights);
+		WeightEdits.MergeSingleEdit(
+			BoneToHoldIndex,
+			VertexID,
+			PrevWeight,
+			1.0f);
+		
+		return;
+	}
+
+	// is the user trying to prune ALL weight from this vertex, AND all other weights are equal to zero?
+	// in this case, we have two options:
+	// 1. if there are other influences recorded for this vertex, then split remaining influence among them
+	// 2. if there are NO other influences recorded for this vertex, then move the weight to the root as last ditch effort
+	const bool bApplyZeroWeightToThisVertex = NewWeightValue <= MinimumWeightThreshold;
+	if (bApplyZeroWeightToThisVertex && bVertexHasNoOtherInfluences)
+	{
+		if (ValuesToNormalize.IsEmpty())
+		{
+			// assign all weight to the root
+			constexpr int32 BoneIndex = 0;
+			constexpr float OldWeight = 0.f;
+			constexpr float NewWeight = 1.f;
+			WeightEdits.MergeSingleEdit(BoneIndex, VertexID, OldWeight, NewWeight);
+		}
+		else
+		{
+			// evenly distribute weight among other influences
+			const float NewWeight = 1.f / ValuesToNormalize.Num();
+			for (int32 i=0; i<ValuesToNormalize.Num(); ++i)
+			{
+				const int32 BoneIndex = BonesAffectingVertex[i];
+				const float OldWeight = ValuesToNormalize[i];
+				WeightEdits.MergeSingleEdit(BoneIndex, VertexID, OldWeight, NewWeight);
+			}
+		}
+
+		// set current bone value to 0.0f
+		const float PrevWeight = GetWeightOfBoneOnVertex(BoneToHoldIndex, VertexID, PreChangeWeights);
+		WeightEdits.MergeSingleEdit(
+			BoneToHoldIndex,
+			VertexID,
+			PrevWeight,
+			0.f);
+		
+		return;
+	}
+
+	// calculate amount we have to spread across the other bones affecting this vertex
+	const float AvailableTotal = 1.0f - NewWeightValue;
+
+	// normalize weights into available space not set by current bone
+	for (int32 i=0; i<ValuesToNormalize.Num(); ++i)
+	{
+		float NormalizedValue = 0.f;
+		if (AvailableTotal > MinimumWeightThreshold && Total > KINDA_SMALL_NUMBER)
+		{
+			NormalizedValue = (ValuesToNormalize[i] / Total) * AvailableTotal;	
+		}
+		const int32 BoneIndex = BonesAffectingVertex[i];
+		const float OldWeight = ValuesToNormalize[i];
+		const float NewWeight = NormalizedValue;
+		WeightEdits.MergeSingleEdit(BoneIndex, VertexID, OldWeight, NewWeight);
+	}
+
+	// record current bone edit
+	const float PrevWeight = GetWeightOfBoneOnVertex(BoneToHoldIndex, VertexID, PreChangeWeights);
+	WeightEdits.MergeSingleEdit(
+		BoneToHoldIndex,
+		VertexID,
+		PrevWeight,
+		NewWeightValue);
+}
+
+void FSkinToolWeights::ApplyCurrentWeightsToMeshDescription(FMeshDescription* EditedMesh)
+{
+	FSkeletalMeshAttributes MeshAttribs(*EditedMesh);
+	FSkinWeightsVertexAttributesRef VertexSkinWeights = MeshAttribs.GetVertexSkinWeights();
+	
+	UE::AnimationCore::FBoneWeightsSettings Settings;
+	Settings.SetNormalizeType(UE::AnimationCore::EBoneWeightNormalizeType::None);
+
+	TArray<UE::AnimationCore::FBoneWeight> SourceBoneWeights;
+	SourceBoneWeights.Reserve(UE::AnimationCore::MaxInlineBoneWeightCount);
+
+	const int32 NumVertices = EditedMesh->Vertices().Num();
+	for (int32 VertexIndex = 0; VertexIndex < NumVertices; VertexIndex++)
+	{
+		SourceBoneWeights.Reset();
+
+		const VertexWeights& VertexWeights = CurrentWeights[VertexIndex];
+		for (const FVertexBoneWeight& SingleBoneWeight : VertexWeights)
+		{
+			SourceBoneWeights.Add(UE::AnimationCore::FBoneWeight(SingleBoneWeight.BoneIndex, SingleBoneWeight.Weight));
+		}
+
+		VertexSkinWeights.Set(FVertexID(VertexIndex), UE::AnimationCore::FBoneWeights::Create(SourceBoneWeights, Settings));
+	}
+}
+
+float FSkinToolWeights::GetWeightOfBoneOnVertex(
+	const int32 BoneIndex,
+	const int32 VertexID,
+	const TArray<VertexWeights>& InVertexWeights)
+{
+	const VertexWeights& VertexWeights = InVertexWeights[VertexID];
+	for (const FVertexBoneWeight& BoneWeight : VertexWeights)
+	{
+		if (BoneWeight.BoneIndex == BoneIndex)
+		{
+			return BoneWeight.Weight;
+		}
+	}
+
+	return 0.f;
+}
+
+void FSkinToolWeights::SetWeightOfBoneOnVertex(
+	const int32 BoneIndex,
+	const int32 VertexID,
+	const float Weight,
+	TArray<VertexWeights>& InOutVertexWeights)
+{
+	Deformer.SetVertexNeedsUpdated(VertexID);
+	
+	// incoming weights are assumed to be normalized already, so set it directly
+	VertexWeights& VertexWeights = InOutVertexWeights[VertexID];
+	for (FVertexBoneWeight& BoneWeight : VertexWeights)
+	{
+		if (BoneWeight.BoneIndex == BoneIndex)
+		{
+			BoneWeight.Weight = Weight;
+			return;
+		}
+	}
+
+	// bone not already an influence on this vertex, so we need to add it..
+
+	// if vertex has room for more influences, then simply add it
+	if (VertexWeights.Num() < UE::AnimationCore::MaxInlineBoneWeightCount)
+	{
+		// add a new influence to this vertex
+		const FVector PosLocalToBone = Deformer.InvCSRefPoseTransforms[BoneIndex].TransformPosition(Deformer.RefPoseVertexPositions[VertexID]);
+		VertexWeights.Emplace(BoneIndex, PosLocalToBone, Weight);
+		return;
+	}
+
+	//
+	// uh oh, we're out of room for more influences on this vertex, so lets kick the smallest influence to make room
+	//
+
+	// find the smallest influence
+	float SmallestInfluence = TNumericLimits<float>::Max();
+	int32 SmallestInfluenceIndex = INDEX_NONE;
+	for (int32 InfluenceIndex=0; InfluenceIndex<VertexWeights.Num(); ++InfluenceIndex)
+	{
+		const FVertexBoneWeight& BoneWeight = VertexWeights[InfluenceIndex];
+		if (BoneWeight.Weight <= SmallestInfluence)
+		{
+			SmallestInfluence = BoneWeight.Weight;
+			SmallestInfluenceIndex = InfluenceIndex;
+		}
+	}
+
+	// replace smallest influence
+	FVertexBoneWeight& BoneWeightToReplace = VertexWeights[SmallestInfluenceIndex];
+	BoneWeightToReplace.Weight = Weight;
+	BoneWeightToReplace.BoneIndex = BoneIndex;
+	BoneWeightToReplace.VertexInBoneSpace = Deformer.InvCSRefPoseTransforms[BoneIndex].TransformPosition(Deformer.RefPoseVertexPositions[VertexID]);
+
+	// now we need to re-normalize because the stamp does not handle maximum influences
+	float TotalWeight = 0.f;
+	for (const FVertexBoneWeight& BoneWeight : VertexWeights)
+	{
+		TotalWeight += BoneWeight.Weight;
+	}
+	for (FVertexBoneWeight& BoneWeight : VertexWeights)
+	{
+		BoneWeight.Weight /= TotalWeight;
+	}
+}
+
+void FSkinToolWeights::ResetAfterChange()
+{
+	PreChangeWeights = CurrentWeights;
+
+	for (int32 i=0; i<MaxFalloffPerVertexThisStroke.Num(); ++i)
+	{
+		MaxFalloffPerVertexThisStroke[i] = 0.f;
+	}
+}
+
+float FSkinToolWeights::SetCurrentFalloffAndGetMaxFalloffThisStroke(int32 VertexID, float CurrentStrength)
+{
+	float& MaxFalloffThisStroke = MaxFalloffPerVertexThisStroke[VertexID];
+	if (MaxFalloffThisStroke < CurrentStrength)
+	{
+		MaxFalloffThisStroke = CurrentStrength;
+	}
+	return MaxFalloffThisStroke;
+}
+
+void FSkinToolWeights::ApplyEditsToWeightMap(
+	const FMultiBoneWeightEdits& Edits, 
+	TArray<VertexWeights>& InOutWeights)
+{
+	for (const TTuple<BoneIndex, FSingleBoneWeightEdits>& BoneWeightEdits : Edits.PerBoneWeightEdits)
+	{
+		const FSingleBoneWeightEdits& WeightEdits = BoneWeightEdits.Value;
+		const int32 BoneIndex = WeightEdits.BoneIndex;
+		for (const TTuple<int32, float>& NewWeight : WeightEdits.NewWeights)
+		{
+			const int32 VertexID = NewWeight.Key;
+			const float Weight = NewWeight.Value;
+			SetWeightOfBoneOnVertex(BoneIndex, VertexID, Weight, InOutWeights);
+		}
+	}
+}
+
+void FMeshSkinWeightsChange::Apply(UObject* Object)
+{
+	USkinWeightsPaintTool* Tool = CastChecked<USkinWeightsPaintTool>(Object);
+	
+	for (TTuple<int32, FSingleBoneWeightEdits>& Pair : AllWeightEdits.PerBoneWeightEdits)
+	{
+		Tool->ExternalUpdateWeights(Pair.Key, Pair.Value.NewWeights);
+	}
+}
+
+void FMeshSkinWeightsChange::Revert(UObject* Object)
+{
+	USkinWeightsPaintTool* Tool = CastChecked<USkinWeightsPaintTool>(Object);
+
+	for (TTuple<int32, FSingleBoneWeightEdits>& Pair : AllWeightEdits.PerBoneWeightEdits)
+	{
+		Tool->ExternalUpdateWeights(Pair.Key, Pair.Value.OldWeights);
+	}
+}
+
+void FMeshSkinWeightsChange::AddBoneWeightEdit(const FSingleBoneWeightEdits& BoneWeightEdit)
+{
+	AllWeightEdits.MergeEdits(BoneWeightEdit);
+}
+
+/*
+ * ToolBuilder
+ */
+
+UMeshSurfacePointTool* USkinWeightsPaintToolBuilder::CreateNewTool(const FToolBuilderState& SceneState) const
+{
+	return NewObject<USkinWeightsPaintTool>(SceneState.ToolManager);
+}
+
+void USkinWeightsPaintTool::Setup()
+{
+	TRACE_CPUPROFILER_EVENT_SCOPE(SkinTool::Setup);
+	
+	UDynamicMeshBrushTool::Setup();
+
+	const IPrimitiveComponentBackedTarget* TargetComponent = Cast<IPrimitiveComponentBackedTarget>(Target);
+	check(TargetComponent);
+	const USkeletalMeshComponent* Component = Cast<USkeletalMeshComponent>(TargetComponent->GetOwnerComponent());
+	check(Component && Component->GetSkeletalMeshAsset())
+
+	// create a mesh description for editing (this must be done before calling UpdateBonePositionInfos) 
+	EditedMesh = MakeUnique<FMeshDescription>();
+	*EditedMesh = *UE::ToolTarget::GetMeshDescription(Target);
+
+	// create a custom set of properties inheriting from the base tool properties
+	WeightToolProperties = NewObject<USkinWeightsPaintToolProperties>(this);
+	WeightToolProperties->RestoreProperties(this);
+	WeightToolProperties->WeightTool = this;
+	WeightToolProperties->bSpecifyRadius = true;
+	// replace the base brush properties
+	ReplaceToolPropertySource(BrushProperties, WeightToolProperties);
+	BrushProperties = WeightToolProperties;
+
+	// default to the root bone as current bone
+	PendingCurrentBone = CurrentBone = Component->GetSkeletalMeshAsset()->GetRefSkeleton().GetBoneName(0);
+
+	// configure preview mesh
+	PreviewMesh->SetTangentsMode(EDynamicMeshComponentTangentsMode::AutoCalculated);
+	PreviewMesh->EnableWireframe(true);
+	PreviewMesh->SetShadowsEnabled(false);
+	// enable vtx colors on preview mesh
+	PreviewMesh->EditMesh([](FDynamicMesh3& Mesh)
+	{
+		Mesh.EnableAttributes();
+		Mesh.Attributes()->DisablePrimaryColors();
+		Mesh.Attributes()->EnablePrimaryColors();
+		// Create an overlay that has no split elements, init with zero value.
+		Mesh.Attributes()->PrimaryColors()->CreateFromPredicate([](int ParentVID, int TriIDA, int TriIDB){return true;}, 0.f);
+	});
+	UMaterialInterface* VtxColorMaterial = GetToolManager()->GetContextQueriesAPI()->GetStandardMaterial(EStandardToolContextMaterials::VertexColorMaterial);
+	if (VtxColorMaterial != nullptr)
+	{
+		PreviewMesh->SetOverrideRenderMaterial(VtxColorMaterial);
+	}
+
+	// build octree for vertices
+	VerticesOctree.Initialize(PreviewMesh->GetMesh(), true);
+
+	// build octree for triangles
+	{
+		TRACE_CPUPROFILER_EVENT_SCOPE(SkinTool::InitTriangleOctree);
+		
+		TriangleOctreeFuture = Async(SkinPaintToolAsyncExecTarget, [&]()
+		{
+			TRACE_CPUPROFILER_EVENT_SCOPE(SkinTool::InitTriangleOctreeRun);
+			TrianglesOctree.Initialize(PreviewMesh->GetMesh());
+		});
+	}
+
+	// initialize weight maps and deformation data
+	Weights.InitializeSkinWeights(Component, EditedMesh.Get());
+	bVisibleWeightsValid = false;
+
+	RecalculateBrushRadius();
+
+	// set up vertex selection mechanic
+	PolygonSelectionMechanic = NewObject<UPolygonSelectionMechanic>(this);
+	PolygonSelectionMechanic->bAddSelectionFilterPropertiesToParentTool = false;
+	PolygonSelectionMechanic->Setup(this);
+	PolygonSelectionMechanic->SetIsEnabled(false);
+	PolygonSelectionMechanic->OnSelectionChanged.AddUObject(this, &USkinWeightsPaintTool::OnSelectionModified);
+	// only select vertices
+	PolygonSelectionMechanic->Properties->bSelectEdges = false;
+	PolygonSelectionMechanic->Properties->bSelectFaces = false;
+	PolygonSelectionMechanic->Properties->bSelectVertices = true;
+	// adjust selection rendering for this context
+	PolygonSelectionMechanic->PolyEdgesRenderer.PointColor = FLinearColor(0.78f, 0.f, 0.78f);
+	PolygonSelectionMechanic->PolyEdgesRenderer.PointSize = 5.0f;
+	PolygonSelectionMechanic->HilightRenderer.PointColor = FLinearColor::Red;
+	PolygonSelectionMechanic->HilightRenderer.PointSize = 10.0f;
+	PolygonSelectionMechanic->SelectionRenderer.LineThickness = 0.0f;
+	PolygonSelectionMechanic->SelectionRenderer.PointColor = FLinearColor::Yellow;
+	PolygonSelectionMechanic->SelectionRenderer.PointSize = 5.0f;
+	PolygonSelectionMechanic->SetShowEdges(false);
+	// initialize the polygon selection mechanic
+	constexpr bool bAutoBuild = true;
+	const FDynamicMesh3* DynamicMesh = PreviewMesh->GetPreviewDynamicMesh();
+	SelectionTopology = MakeUnique<UE::Geometry::FTriangleGroupTopology>(DynamicMesh, bAutoBuild);
+	MeshSpatial = MakeUnique<UE::Geometry::FDynamicMeshAABBTree3>(DynamicMesh, bAutoBuild);
+	PolygonSelectionMechanic->Initialize(
+		DynamicMesh,
+		FTransform::Identity,
+		Component->GetWorld(),
+		SelectionTopology.Get(),
+		[this]() { return MeshSpatial.Get(); }
+	);
+	
+	// inform user of tool keys
+	// TODO talk with UX team about viewport overlay to show hotkeys
+	GetToolManager()->DisplayMessage(
+		LOCTEXT("OnStartSkinWeightsPaint", "Paint per-bone skin weights. [ and ] change brush size, Ctrl to Erase/Subtract, Shift to Smooth"),
+		EToolMessageLevel::UserNotification);
+}
+
+void USkinWeightsPaintTool::DrawHUD(FCanvas* Canvas, IToolsContextRenderAPI* RenderAPI)
+{
+	Super::DrawHUD(Canvas, RenderAPI);
+
+	if (PolygonSelectionMechanic)
+	{
+		PolygonSelectionMechanic->DrawHUD(Canvas, RenderAPI);
+	}
+}
+
+void USkinWeightsPaintTool::Render(IToolsContextRenderAPI* RenderAPI)
+{
+	if (WeightToolProperties->EditingMode == EWeightEditMode::Brush)
+	{
+		Super::Render(RenderAPI);	
+	}
+	else if (PolygonSelectionMechanic && WeightToolProperties->EditingMode == EWeightEditMode::Selection)
+	{
+		PolygonSelectionMechanic->Render(RenderAPI);
+	}
+	
+}
+
+void USkinWeightsPaintTool::OnTick(float DeltaTime)
+{
+	if (bStampPending)
+	{
+		ApplyStamp(LastStamp);
+		bStampPending = false;
+	}
+
+	if (PendingCurrentBone.IsSet())
+	{
+		UpdateCurrentBone(*PendingCurrentBone);
+		PendingCurrentBone.Reset();
+	}
+
+	if (bVisibleWeightsValid == false || WeightToolProperties->bColorModeChanged)
+	{
+		UpdateCurrentBoneVertexColors();
+		bVisibleWeightsValid = true;
+		WeightToolProperties->bColorModeChanged = false;
+	}
+
+	// sparsely updates vertex positions (only on vertices with modified weights)
+	Weights.Deformer.UpdateVertexDeformation(this);
+}
+
+bool USkinWeightsPaintTool::HitTest(const FRay& Ray, FHitResult& OutHit)
+{
+	// do not query the triangle octree until all async ops are finished
+	TriangleOctreeFuture.Wait();
+	
+	// put ray in local space of skeletal mesh component
+	// currently no way to transform skeletal meshes in the editor,
+	// but at some point in the future we may add the ability to move parts around
+	const IPrimitiveComponentBackedTarget* TargetComponent = Cast<IPrimitiveComponentBackedTarget>(Target);
+	const FTransform3d CurTargetTransform(TargetComponent->GetWorldTransform());
+	FRay3d LocalRay(
+		CurTargetTransform.InverseTransformPosition((FVector3d)Ray.Origin),
+		CurTargetTransform.InverseTransformVector((FVector3d)Ray.Direction));
+	UE::Geometry::Normalize(LocalRay.Direction);
+	
+	const FDynamicMesh3* Mesh = PreviewMesh->GetPreviewDynamicMesh();
+
+	FViewCameraState StateOut;
+	GetToolManager()->GetContextQueriesAPI()->GetCurrentViewState(StateOut);
+	FVector3d LocalEyePosition(CurTargetTransform.InverseTransformPosition((FVector3d)StateOut.Position));
+	const int32 TriID = TrianglesOctree.FindNearestHitObject(
+		LocalRay,
+		[this, Mesh, &LocalEyePosition](int TriangleID)
+	{
+		FVector3d Normal, Centroid;
+		double Area;
+		Mesh->GetTriInfo(TriangleID, Normal, Area, Centroid);
+		return Normal.Dot((Centroid - LocalEyePosition)) < 0;
+	});
+	
+	if (TriID != IndexConstants::InvalidID)
+	{	
+		FastTriWinding::FTriangle3d Triangle;
+		Mesh->GetTriVertices(TriID, Triangle.V[0], Triangle.V[1], Triangle.V[2]);
+		UE::Geometry::FIntrRay3Triangle3d Query(LocalRay, Triangle);
+		Query.Find();
+
+		StampLocalPos = LocalRay.PointAt(Query.RayParameter);
+		TriangleUnderStamp = TriID;
+
+		OutHit.FaceIndex = TriID;
+		OutHit.Distance = Query.RayParameter;
+		OutHit.Normal = CurTargetTransform.TransformVector(Mesh->GetTriNormal(TriID));
+		OutHit.ImpactPoint = CurTargetTransform.TransformPosition(StampLocalPos);
+		return true;
+	}
+	
+	return false;
+}
+
+void USkinWeightsPaintTool::OnBeginDrag(const FRay& WorldRay)
+{
+	UDynamicMeshBrushTool::OnBeginDrag(WorldRay);
+
+	if (IsInBrushStroke())
+	{
+		bInvertStroke = GetCtrlToggle();
+		bSmoothStroke = GetShiftToggle();
+		BeginChange();
+		StartStamp = UBaseBrushTool::LastBrushStamp;
+		LastStamp = StartStamp;
+		bStampPending = true;
+	}
+}
+
+void USkinWeightsPaintTool::OnUpdateDrag(const FRay& WorldRay)
+{
+	UDynamicMeshBrushTool::OnUpdateDrag(WorldRay);
+	if (IsInBrushStroke())
+	{
+		LastStamp = UBaseBrushTool::LastBrushStamp;
+		bStampPending = true;
+	}
+}
+
+void USkinWeightsPaintTool::OnEndDrag(const FRay& Ray)
+{
+	UDynamicMeshBrushTool::OnEndDrag(Ray);
+
+	bInvertStroke = false;
+	bSmoothStroke = false;
+	bStampPending = false;
+
+	// close change record
+	TUniquePtr<FMeshSkinWeightsChange> Change = EndChange();
+	GetToolManager()->BeginUndoTransaction(LOCTEXT("PaintWeightChange", "Paint skin weights."));
+	GetToolManager()->EmitObjectChange(this, MoveTemp(Change), LOCTEXT("PaintWeightChange", "Paint skin weights."));
+	GetToolManager()->EndUndoTransaction();
+}
+
+bool USkinWeightsPaintTool::OnUpdateHover(const FInputDeviceRay& DevicePos)
+{
+	UDynamicMeshBrushTool::OnUpdateHover(DevicePos);
+	return true;
+}
+
+void USkinWeightsPaintTool::CalculateVertexROI(
+	const FBrushStampData& Stamp,
+	TArray<VertexIndex>& VertexIDs,
+	TArray<float>& VertexFalloffs)
+{
+	using namespace UE::Geometry;
+
+	TRACE_CPUPROFILER_EVENT_SCOPE(SkinTool::CalculateVertexROI);
+
+	auto DistanceToFalloff = [this](int32 InVertexID, float InDistanceSq)-> float
+	{
+		const float CurrentFalloff = CalculateBrushFalloff(FMath::Sqrt(InDistanceSq));
+		const float UseFalloff = Weights.SetCurrentFalloffAndGetMaxFalloffThisStroke(InVertexID, CurrentFalloff);
+		return UseFalloff;
+	};
+	
+	if (WeightToolProperties->FalloffMode == EWeightBrushFalloffMode::Volume)
+	{
+		const IPrimitiveComponentBackedTarget* TargetComponent = Cast<IPrimitiveComponentBackedTarget>(Target);
+		const FTransform3d Transform(TargetComponent->GetWorldTransform());
+		const FVector3d StampPosLocal = Transform.InverseTransformPosition((FVector3d)Stamp.WorldPosition);
+		const float RadiusSqr = CurrentBrushRadius * CurrentBrushRadius;
+		const FDynamicMesh3* Mesh = PreviewMesh->GetPreviewDynamicMesh();
+		const FAxisAlignedBox3d QueryBox(StampPosLocal, CurrentBrushRadius);
+		VerticesOctree.RangeQuery(QueryBox,
+			[&](int32 VertexID) { return FVector3d::DistSquared(Mesh->GetVertex(VertexID), StampPosLocal) < RadiusSqr; },
+			VertexIDs);
+
+		for (const int32 VertexID : VertexIDs)
+		{
+			const float DistSq = FVector3d::DistSquared(Mesh->GetVertex(VertexID), StampPosLocal);
+			VertexFalloffs.Add(DistanceToFalloff(VertexID, DistSq));
+		}
+		
+		return;
+	}
+
+	if (WeightToolProperties->FalloffMode == EWeightBrushFalloffMode::Surface)
+	{
+		// get coordinate frame from stamp
+		auto GetFrameFromStamp = [](const FBrushStampData& InStamp) -> FFrame3d
+		{
+			const FVector3d Origin = InStamp.WorldPosition;
+			const FVector3d Normal = InStamp.WorldNormal;
+			FVector3d NonCollinear = Normal;
+			// get a guaranteed non collinear vector to the normal
+			// doesn't matter where in the plane, stamp is radially symmetric
+			do 
+			{
+				NonCollinear.X = FMath::RandRange(-1.0f, 1.0f);
+				NonCollinear.Y = FMath::RandRange(-1.0f, 1.0f);
+				NonCollinear.Z = FMath::RandRange(-1.0f, 1.0f);
+				NonCollinear.Normalize();
+					
+			} while (FMath::Abs(NonCollinear.Dot(Normal)) > 0.8f);
+
+			const FVector3d Plane = Normal.Cross(NonCollinear);
+			const FVector3d Cross = Plane.Cross(Normal);
+			return FFrame3d(Origin, Cross, Plane, Normal);
+				
+		};
+		const FFrame3d SeedFrame = GetFrameFromStamp(Stamp);
+			
+		// create the ExpMap generator, computes vertex polar coordinates in a plane tangent to the surface
+		const FDynamicMesh3* Mesh = PreviewMesh->GetPreviewDynamicMesh();
+		TMeshLocalParam<FDynamicMesh3> Param(Mesh);
+		Param.ParamMode = ELocalParamTypes::ExponentialMapUpwindAvg;
+		const FIndex3i TriVerts = Mesh->GetTriangle(TriangleUnderStamp);
+		Param.ComputeToMaxDistance(SeedFrame, TriVerts, Stamp.Radius);
+		{
+			// store vertices under the brush and their distances from the stamp
+			const float StampRadSq = FMath::Pow(Stamp.Radius, 2);
+			for (int32 VertexID : Mesh->VertexIndicesItr())
+			{
+				if (!Param.HasUV(VertexID))
+				{
+					continue;
+				}
+				
+				FVector2d UV = Param.GetUV(VertexID);
+				const float DistSq = UV.SizeSquared();
+				if (DistSq >= StampRadSq)
+				{
+					continue;
+				}
+
+				
+				VertexFalloffs.Add(DistanceToFalloff(VertexID, DistSq));
+				VertexIDs.Add(VertexID);
+			}
+		}
+		
+		return;
+	}
+	
+	checkNoEntry();
+}
+
+FVector4f USkinWeightsPaintTool::WeightToColor(float Value) const
+{
+	// optional greyscale mode
+	if (WeightToolProperties->ColorMode == EWeightColorMode::Greyscale)
+	{
+		return FLinearColor::LerpUsingHSV(FLinearColor::Black, FLinearColor::White, Value);	
+	}
+	
+	// early out zero weights to min color
+	if (Value <= MinimumWeightThreshold)
+	{
+		return WeightToolProperties->MinColor;
+	}
+
+	// early out full weights to max color
+	if (FMath::IsNearlyEqual(Value, 1.0f))
+	{
+		return WeightToolProperties->MaxColor;
+	}
+
+	// get user-specified color ramp for intermediate colors
+	const TArray<FLinearColor>& Colors = WeightToolProperties->ColorRamp;
+	
+	// revert back to simple Lerp(min,max) if user supplied color ramp doesn't have enough colors
+	if (Colors.Num() < 2)
+	{
+		const FLinearColor FinalColor = FLinearColor::LerpUsingHSV(WeightToolProperties->MinColor, WeightToolProperties->MaxColor, Value);
+		return UE::Geometry::ToVector4<float>(FinalColor);
+	}
+
+	// otherwise, interpolate within two nearest ramp colors
+	Value = FMath::Clamp(Value, 0.0f, 1.0f);
+	const float PerColorRange = 1.0f / (Colors.Num() - 1);
+	const int ColorIndex = static_cast<int>(Value / PerColorRange);
+	const float RangeStart = ColorIndex * PerColorRange;
+	const float RangeEnd = (ColorIndex + 1) * PerColorRange;
+	const float Param = (Value - RangeStart) / (RangeEnd - RangeStart);
+	const FLinearColor& StartColor = Colors[ColorIndex];
+	const FLinearColor& EndColor = Colors[ColorIndex+1];
+	const FLinearColor FinalColor = FLinearColor::LerpUsingHSV(StartColor, EndColor, Param);
+	return UE::Geometry::ToVector4<float>(FinalColor);
+}
+
+
+void USkinWeightsPaintTool::UpdateCurrentBoneVertexColors()
+{
+	const int32 CurrentBoneIndex = CurrentBone == NAME_None ? INDEX_NONE : Weights.Deformer.BoneNameToIndexMap[CurrentBone];
+	
+	// update mesh with new value colors
+	PreviewMesh->DeferredEditMesh([this, &CurrentBoneIndex](FDynamicMesh3& Mesh)
+	{
+		const UE::Geometry::FNonManifoldMappingSupport NonManifoldMappingSupport(Mesh);
+		UE::Geometry::FDynamicMeshColorOverlay* ColorOverlay = Mesh.Attributes()->PrimaryColors();
+		for (const int32 ElementId : ColorOverlay->ElementIndicesItr())
+		{
+			// with no bone selected, all vertices are drawn black
+			if (CurrentBoneIndex == INDEX_NONE)
+			{
+				ColorOverlay->SetElement(ElementId, FVector4f(FLinearColor::Black));
+				continue;
+			}
+			
+			const int32 VertexID = ColorOverlay->GetParentVertex(ElementId);	
+			const int32 SrcVertexID = NonManifoldMappingSupport.GetOriginalNonManifoldVertexID(VertexID);
+			const float Value = Weights.GetWeightOfBoneOnVertex(CurrentBoneIndex, SrcVertexID, Weights.CurrentWeights);
+			const FVector4f Color(WeightToColor(Value));
+			ColorOverlay->SetElement(ElementId, Color);
+		}
+	}, false);
+	PreviewMesh->NotifyDeferredEditCompleted(UPreviewMesh::ERenderUpdateMode::FastUpdate, EMeshRenderAttributeFlags::VertexColors, false);
+}
+
+float USkinWeightsPaintTool::CalculateBrushFalloff(float Distance) const
+{
+	const float f = FMathd::Clamp(1.f - BrushProperties->BrushFalloffAmount, 0.f, 1.f);
+	float d = Distance / CurrentBrushRadius;
+	double w = 1;
+	if (d > f)
+	{
+		d = FMathd::Clamp((d - f) / (1.0 - f), 0.0, 1.0);
+		w = (1.0 - d * d);
+		w = w * w * w;
+	}
+	return w;
+}
+
+void USkinWeightsPaintTool::ApplyStamp(const FBrushStampData& Stamp)
+{
+	TRACE_CPUPROFILER_EVENT_SCOPE(SkinTool::ApplyStamp);
+
+	// must select a bone to paint
+	if (CurrentBone == NAME_None)
+	{
+		return;
+	}
+	
+	// get the vertices under the brush, and their squared distances to the brush center
+	// when using "Volume" brush, distances are straight line
+	// when using "Surface" brush, distances are geodesics
+	TArray<int32> VerticesInStamp;
+	TArray<float> VertexFalloffs;
+	CalculateVertexROI(Stamp, VerticesInStamp, VertexFalloffs);
+
+	// gather sparse set of modifications made from this stamp, these edits are merged throughout
+	// the lifetime of a single brush stroke in the "ActiveChange" allowing for undo/redo
+	FMultiBoneWeightEdits WeightEditsFromStamp;
+	{
+		TRACE_CPUPROFILER_EVENT_SCOPE(SkinTool::EditWeightOfVerticesInStamp);
+		// generate a weight edit from this stamp (includes modifications caused by normalization)
+		if (bSmoothStroke || WeightToolProperties->BrushMode == EWeightEditOperation::Relax)
+		{
+			// use mesh topology to iteratively smooth weights across neighboring vertices
+			const float UseStrength = CalculateBrushStrengthToUse(EWeightEditOperation::Relax);
+			RelaxWeightOnVertices(VerticesInStamp, VertexFalloffs, UseStrength, WeightEditsFromStamp);
+		}
+		else
+		{
+			// edit weight; either by "Add", "Remove", "Replace", "Multiply"
+			const float UseStrength = CalculateBrushStrengthToUse(WeightToolProperties->BrushMode);
+			EditWeightOnVertices(
+				WeightToolProperties->BrushMode,
+				VerticesInStamp,
+				VertexFalloffs,
+				UseStrength,
+				WeightEditsFromStamp);
+		}
+	}
+	
+	{
+		TRACE_CPUPROFILER_EVENT_SCOPE(SkinTool::ApplyWeightEditsToActiveChange);
+		// store weight edits from all stamps made during a single stroke (1 transaction per stroke)
+		for (const TTuple<int32, FSingleBoneWeightEdits>& BoneWeightEdits : WeightEditsFromStamp.PerBoneWeightEdits)
+		{
+			ActiveChange->AddBoneWeightEdit(BoneWeightEdits.Value);
+		}
+	}
+	
+	{
+		TRACE_CPUPROFILER_EVENT_SCOPE(SkinTool::ApplyWeightEditsToCurrentWeights);
+		// apply weights to current weights
+		Weights.ApplyEditsToWeightMap(WeightEditsFromStamp, Weights.CurrentWeights);
+	}
+
+	{
+		TRACE_CPUPROFILER_EVENT_SCOPE(SkinTool::UpdateVertexColors);
+		// update vertex colors
+		PreviewMesh->DeferredEditMesh([&](FDynamicMesh3& Mesh)
+		{
+			TArray<int> ElementIds;
+			UE::Geometry::FDynamicMeshColorOverlay* ColorOverlay = Mesh.Attributes()->PrimaryColors();
+			const int32 BoneIndex = Weights.Deformer.BoneNameToIndexMap[CurrentBone];
+			const int32 NumVerticesInStamp = VerticesInStamp.Num();
+			for (int32 Index = 0; Index < NumVerticesInStamp; ++Index)
+			{
+				const int32 VertexID = VerticesInStamp[Index];
+				const float Weight = Weights.GetWeightOfBoneOnVertex(BoneIndex, VertexID, Weights.CurrentWeights);
+				FVector4f NewColor(WeightToColor(Weight));
+				ColorOverlay->GetVertexElements(VertexID, ElementIds);
+				for (const int32 ElementId : ElementIds)
+				{
+					ColorOverlay->SetElement(ElementId, NewColor);
+				}
+				ElementIds.Reset();
+			}
+		}, false);
+		PreviewMesh->NotifyDeferredEditCompleted(UPreviewMesh::ERenderUpdateMode::FastUpdate, EMeshRenderAttributeFlags::VertexColors, false);
+	}
+}
+
+float USkinWeightsPaintTool::CalculateBrushStrengthToUse(EWeightEditOperation EditMode) const
+{
+	float UseStrength = BrushProperties->BrushStrength;
+
+	// invert brush strength differently depending on brush mode
+	switch (EditMode)
+	{
+	case EWeightEditOperation::Add:
+		{
+			UseStrength *= bInvertStroke ? -1.0f : 1.0f;
+			break;
+		}
+	case EWeightEditOperation::Replace:
+		{
+			UseStrength = bInvertStroke ? 1.0f - UseStrength : UseStrength;
+			break;
+		}
+	case EWeightEditOperation::Multiply:
+		{
+			UseStrength = bInvertStroke ? 1.0f + UseStrength : UseStrength;
+			break;
+		}
+	case EWeightEditOperation::Relax:
+		{
+			UseStrength = bInvertStroke ? 1.0f - UseStrength : UseStrength;
+			break;
+		}
+	default:
+		checkNoEntry();
+	}
+
+	return UseStrength;
+}
+
+void USkinWeightsPaintTool::EditWeightOnVertices(
+	EWeightEditOperation EditOperation,
+	const TArray<int32>& VerticesToEdit,
+	const TArray<float>& VertexFalloffs,
+	const float UseStrength,
+	FMultiBoneWeightEdits& InOutWeightEdits)
+{
+	// spin through the vertices in the stamp and store new weight values in NewValuesFromStamp
+	// afterwards, these values are normalized while taking into consideration the user's desired changes
+	const int32 CurrentBoneIndex = Weights.Deformer.BoneNameToIndexMap[CurrentBone];
+	const int32 NumVerticesInStamp = VerticesToEdit.Num();
+	for (int32 Index = 0; Index < NumVerticesInStamp; ++Index)
+	{
+		const int32 VertexID = VerticesToEdit[Index];
+		const float UseFalloff = VertexFalloffs.IsValidIndex(Index) ? VertexFalloffs[Index] : 1.f;
+		const float ValueBeforeStroke = Weights.GetWeightOfBoneOnVertex(CurrentBoneIndex, VertexID, Weights.PreChangeWeights);
+
+		// calculate new weight value
+		float NewValueAfterStamp = ValueBeforeStroke;
+		switch (EditOperation)
+		{
+		case EWeightEditOperation::Add:
+			{
+				NewValueAfterStamp = ValueBeforeStroke + (UseStrength * UseFalloff);
+				break;
+			}
+		case EWeightEditOperation::Replace:
+			{
+				NewValueAfterStamp = FMath::Lerp(ValueBeforeStroke, UseStrength, UseFalloff);
+				break;
+			}
+		case EWeightEditOperation::Multiply:
+			{
+				const float DeltaFromThisStamp = ((ValueBeforeStroke * UseStrength) - ValueBeforeStroke) * UseFalloff;
+				NewValueAfterStamp = ValueBeforeStroke + DeltaFromThisStamp;
+				break;
+			}
+		default:
+			// relax operation not supported by this function, use RelaxWeightOnVertices()
+			checkNoEntry();
+		}
+
+		// normalize the values across all bones affecting the vertices in the stamp, and record the bone edits
+		// normalization is done while holding all weights on the current bone constant so that user edits are not overwritten
+		Weights.EditVertexWeightAndNormalize(
+			CurrentBoneIndex,
+			VertexID,
+			NewValueAfterStamp,
+			InOutWeightEdits);
+	}
+}
+
+void USkinWeightsPaintTool::RelaxWeightOnVertices(
+	TArray<int32> VerticesToEdit,
+	TArray<float> VertexFalloffs,
+	const float UseStrength,
+	FMultiBoneWeightEdits& InOutWeightEdits)
+{
+	const FDynamicMesh3* CurrentMesh = PreviewMesh->GetMesh();
+	const UE::Geometry::FNonManifoldMappingSupport NonManifoldMappingSupport(*CurrentMesh);
+
+	auto NormalizeWeights = [](TMap<BoneIndex, float>& InOutWeights)
+	{
+		float TotalWeight = 0.f;
+		for (const TTuple<BoneIndex, float>& Weight : InOutWeights)
+		{
+			TotalWeight += Weight.Value;
+		}
+		for (TTuple<BoneIndex, float>& Weight : InOutWeights)
+		{
+			Weight.Value /= TotalWeight;
+		}
+	};
+
+	// for each vertex in the stamp...
+	constexpr int32 AvgNumNeighbors = 8;
+	using VertexNeighborWeights = TArray<float, TInlineAllocator<AvgNumNeighbors>>;
+	TArray<int32> AllNeighborVertices;
+	TMap<BoneIndex, VertexNeighborWeights> WeightsOnAllNeighbors;
+	TMap<BoneIndex, float> FinalWeights;
+	for (int32 Index = 0; Index < VerticesToEdit.Num(); ++Index)
+	{
+		const int32 VertexID = VerticesToEdit[Index];
+		const int32 SrcVertexID = NonManifoldMappingSupport.GetOriginalNonManifoldVertexID(VertexID);
+
+		// get list of all neighboring vertices, AND this vertex
+		AllNeighborVertices.Reset();
+		AllNeighborVertices.Add(VertexID);
+		for (const int32 NeighborVertexID : CurrentMesh->VtxVerticesItr(SrcVertexID))
+		{
+			AllNeighborVertices.Add(NeighborVertexID);
+		}
+
+		// get all weights above a given threshold across ALL neighbors (including self)
+		WeightsOnAllNeighbors.Reset();
+		for (const int32 NeighborVertexID : AllNeighborVertices)
+		{
+			for (const FVertexBoneWeight& BoneWeight : Weights.PreChangeWeights[NeighborVertexID])
+			{
+				if (BoneWeight.Weight > MinimumWeightThreshold)
+				{
+					VertexNeighborWeights& BoneWeights = WeightsOnAllNeighbors.FindOrAdd(BoneWeight.BoneIndex);
+					BoneWeights.Add(BoneWeight.Weight);
+				}
+			}
+		}
+
+		// calculate single average weight of each bone on all the neighbors
+		FinalWeights.Reset();
+		for (const TTuple<BoneIndex, VertexNeighborWeights>& NeighborWeights : WeightsOnAllNeighbors)
+		{
+			float TotalWeightOnThisBone = 0.f;
+			for (const float& Value : NeighborWeights.Value)
+			{
+				TotalWeightOnThisBone += Value;
+			}
+			FinalWeights.Add(NeighborWeights.Key, TotalWeightOnThisBone / NeighborWeights.Value.Num());
+		}
+
+		// normalize the weights
+		NormalizeWeights(FinalWeights);
+
+		// lerp weights from previous values, to fully relaxed values by brush strength scaled by falloff
+		const float UseFalloff = VertexFalloffs.IsValidIndex(Index) ? VertexFalloffs[Index] * UseStrength : UseStrength;
+		for (TTuple<BoneIndex, float>& FinalWeight : FinalWeights)
+		{
+			const int32 BoneIndex = FinalWeight.Key;
+			float NewWeight = FinalWeight.Value;
+			float OldWeight = Weights.GetWeightOfBoneOnVertex(BoneIndex, VertexID, Weights.PreChangeWeights);
+			FinalWeight.Value = FMath::Lerp(OldWeight, NewWeight, UseFalloff);
+		}
+
+		// normalize again
+		NormalizeWeights(FinalWeights);
+
+		// apply weight edits
+		for (const TTuple<BoneIndex, float>& FinalWeight : FinalWeights)
+		{
+			// record an edit for this vertex, for this bone
+			const int32 BoneIndex = FinalWeight.Key;
+			const float NewWeight = FinalWeight.Value;
+			const float OldWeight = Weights.GetWeightOfBoneOnVertex(BoneIndex, VertexID, Weights.PreChangeWeights);
+			InOutWeightEdits.MergeSingleEdit(BoneIndex, VertexID, OldWeight, NewWeight);
+		}
+	}
+}
+
+void USkinWeightsPaintTool::ApplyWeightEditsToMesh(
+	const FText& TransactionLabel,
+	const SkinPaintTool::FMultiBoneWeightEdits& WeightEdits)
+{
+	// clear the active change
+	BeginChange();
+
+	// store weight edits in the active change
+	for (const TTuple<int32, FSingleBoneWeightEdits>& BoneWeightEdits : WeightEdits.PerBoneWeightEdits)
+	{
+		ActiveChange->AddBoneWeightEdit(BoneWeightEdits.Value);
+	}
+
+	// apply the weight edits to the actual mesh
+	// (copies weight modifications to the tool's weight data structure and updates the vertex colors
+	// in brush mode, this is normally called by the tool itself after a stroke)
+	ActiveChange->Apply(this);
+
+	// store active change in the transaction buffer
+	UInteractiveToolManager* ToolManager = GetToolManager();
+	ToolManager->BeginUndoTransaction(TransactionLabel);
+	ToolManager->EmitObjectChange(this, EndChange(), TransactionLabel);
+	ToolManager->EndUndoTransaction();
+}
+
+void USkinWeightsPaintTool::UpdateCurrentBone(const FName& BoneName)
+{
+	CurrentBone = BoneName;
+	bVisibleWeightsValid = false;
+}
+
+void USkinWeightsPaintTool::GetVerticesToEdit(TArray<VertexIndex>& OutVertexIndices) const
+{
+	OutVertexIndices.Reset();
+
+	//
+	// 1. Prioritize selected vertices
+	// TODO check selection manager for selected vertices here
+	const FGroupTopologySelection& Selection = PolygonSelectionMechanic->GetActiveSelection();
+	if (!Selection.SelectedCornerIDs.IsEmpty())
+	{
+		OutVertexIndices = Selection.SelectedCornerIDs.Array();
+		return;
+	}
+
+	//
+	// 2. Fallback on vertices weighted to selected bones
+	if (!SelectedBoneIndices.IsEmpty())
+	{
+		VertexIndex VertexID = 0;
+		for (const VertexWeights& VertWeights : Weights.PreChangeWeights)
+		{
+			for (const FVertexBoneWeight& BoneWeight : VertWeights)
+			{
+				if (SelectedBoneIndices.Contains(BoneWeight.BoneIndex) && BoneWeight.Weight > MinimumWeightThreshold)
+				{
+					OutVertexIndices.Add(VertexID);	
+				}
+			}
+			++VertexID;
+		}
+		return;
+	}
+
+	//
+	// 3. Finally, fallback to ALL vertices in the mesh
+	const int32 NumVertices = EditedMesh->Vertices().Num();
+	OutVertexIndices.Reset(NumVertices);
+	for (VertexIndex VertexID = 0; VertexID < NumVertices; VertexID++)
+	{
+		OutVertexIndices.Add(VertexID);
+	}
+}
+
+BoneIndex USkinWeightsPaintTool::GetBoneIndexFromName(const FName BoneName) const
+{
+	return  BoneName == NAME_None ? INDEX_NONE : Weights.Deformer.BoneNameToIndexMap[BoneName];
+}
+
+void USkinWeightsPaintTool::OnShutdown(EToolShutdownType ShutdownType)
+{
+	// save tool properties
+	WeightToolProperties->SaveProperties(this);
+
+	// shutdown polygon selection mechanic
+	if (PolygonSelectionMechanic)
+	{
+		PolygonSelectionMechanic->Shutdown();
+		PolygonSelectionMechanic = nullptr;
+	}
+
+	// apply changes to asset
+	if (ShutdownType == EToolShutdownType::Accept)
+	{
+		// apply the weights to the mesh description
+		Weights.ApplyCurrentWeightsToMeshDescription(EditedMesh.Get());
+
+		// this block bakes the modified DynamicMeshComponent back into the StaticMeshComponent inside an undo transaction
+		GetToolManager()->BeginUndoTransaction(LOCTEXT("SkinWeightsPaintTool", "Paint Skin Weights"));
+		UE::ToolTarget::CommitMeshDescriptionUpdate(Target, EditedMesh.Get());
+		GetToolManager()->EndUndoTransaction();
+	}
+}
+
+void USkinWeightsPaintTool::BeginChange()
+{
+	ActiveChange = MakeUnique<FMeshSkinWeightsChange>();
+}
+
+TUniquePtr<FMeshSkinWeightsChange> USkinWeightsPaintTool::EndChange()
+{
+	Weights.ResetAfterChange();
+	return MoveTemp(ActiveChange);
+}
+
+void USkinWeightsPaintTool::ExternalUpdateWeights(const int32 BoneIndex, const TMap<int32, float>& NewValues)
+{
+	for (const TTuple<int32, float>& Pair : NewValues)
+	{
+		const int32 VertexID = Pair.Key;
+		const float Weight = Pair.Value;
+		Weights.SetWeightOfBoneOnVertex(BoneIndex, VertexID, Weight, Weights.CurrentWeights);
+		Weights.SetWeightOfBoneOnVertex(BoneIndex, VertexID, Weight, Weights.PreChangeWeights);
+	}
+
+	const FName BoneName = Weights.Deformer.BoneNames[BoneIndex];
+	if (BoneName == CurrentBone)
+	{
+		UpdateCurrentBoneVertexColors();
+	}
+}
+
+void FSkinMirrorData::RegenerateMirrorData(
+    const TArray<FName>& BoneNames,
+    const TMap<FName, BoneIndex>& BoneNameToIndexMap,
+	const FReferenceSkeleton& RefSkeleton,
+	const TArray<FVector>& RefPoseVertices,
+	EAxis::Type InMirrorAxis,
+	EMirrorDirection InMirrorDirection,
+	const float MaxSearchRadius)
+{
+	if (bIsInitialized && InMirrorAxis == Axis && InMirrorDirection==Direction)
+	{
+		// already initialized, just re-use cached data
+		return;
+	}
+
+	// need to re-initialize
+	bIsInitialized = false;
+	Axis = InMirrorAxis;
+	Direction = InMirrorDirection;
+	
+	// build bone map for mirroring
+	// TODO, provide some way to edit the mirror bone mapping, either by providing a UMirrorDataTable input or editing directly in the hierarchy view.
+	for (FName BoneName : BoneNames)
+	{
+		FName MirroredBoneName = UMirrorDataTable::FindBestMirroredBone(BoneName, RefSkeleton, Axis);
+
+		int32 BoneIndex = BoneNameToIndexMap[BoneName];
+		int32 MirroredBoneIndex = BoneNameToIndexMap[MirroredBoneName];
+		BoneMap.Add(BoneIndex, MirroredBoneIndex);
+		
+		// debug view bone mapping
+		//UE_LOG(LogTemp, Log, TEXT("Bone    : %s"), *BoneName.ToString());
+		//UE_LOG(LogTemp, Log, TEXT("Mirrored: %s"), *MirroredBoneName.ToString());
+		//UE_LOG(LogTemp, Log, TEXT("-------"));
+	}
+
+	// hash grid constants
+	constexpr float HashGridCellSize = 2.0f;
+	constexpr float ThresholdRadius = 0.1f;
+
+	// build a point set of the rest pose vertices
+	UE::Geometry::FDynamicPointSet3d PointSet;
+	for (int32 PointID = 0; PointID < RefPoseVertices.Num(); ++PointID)
+	{
+		PointSet.InsertVertex(PointID, RefPoseVertices[PointID]);
+	}
+
+	// build a spatial hash map from the point set
+	UE::Geometry::FPointSetAdapterd PointSetAdapter = UE::Geometry::MakePointsAdapter(&PointSet);
+	UE::Geometry::FPointSetHashtable PointHash(&PointSetAdapter);
+	PointHash.Build(HashGridCellSize, FVector3d::Zero());
+	
+	// generate a map of point IDs on the target side, to their equivalent vertex ID on the source side 
+	TArray<int> PointsInThreshold;
+	TArray<int> PointsInSphere;
+	bAllVerticesMirrored = true;
+	for (int32 TargetVertexID = 0; TargetVertexID < RefPoseVertices.Num(); ++TargetVertexID)
+	{
+		const FVector& TargetPosition = RefPoseVertices[TargetVertexID];
+
+		if (Direction == EMirrorDirection::PositiveToNegative && TargetPosition[Axis-1] >= 0.f)
+		{
+			continue; // copying to negative side, but vertex is on positive side
+		}
+		if (Direction == EMirrorDirection::NegativeToPositive && TargetPosition[Axis-1] <= 0.f)
+		{
+			continue; // copying to positive side, but vertex is on negative side
+		}
+
+		// flip position across the mirror axis
+		FVector MirroredPosition = TargetPosition;
+		MirroredPosition[Axis-1] *= -1.f;
+
+		// query spatial hash near mirrored position, gradually increasing search radius until at least 1 point is found
+		PointsInSphere.Reset();
+		float SearchRadius = ThresholdRadius;
+		while(PointsInSphere.IsEmpty())
+		{
+			PointHash.FindPointsInBall(MirroredPosition, SearchRadius, PointsInSphere);
+			SearchRadius += ThresholdRadius;
+
+			// forcibly break out if search radius gets bigger than the mesh bounds.
+			// this could potentially happen if mesh is highly non-symmetrical along mirror axis
+			if (SearchRadius >= MaxSearchRadius)
+			{
+				break;
+			}
+		}
+
+		// no mirrored points?
+		if (PointsInSphere.IsEmpty())
+		{
+			bAllVerticesMirrored = false;
+			continue;
+		}
+
+		// find the closest single point
+		float ClosestDistSq = TNumericLimits<float>::Max();
+		int32 ClosestVertexID = INDEX_NONE;
+		for (const int32 PointInSphereID : PointsInSphere)
+		{
+			const float DistSq = FVector::DistSquared(RefPoseVertices[PointInSphereID], MirroredPosition);
+			if (DistSq < ClosestDistSq)
+			{
+				ClosestDistSq = DistSq;
+				ClosestVertexID = PointInSphereID;
+			}
+		}
+		
+		// record the mirrored vertex ID for this vertex 
+		VertexMap.FindOrAdd(TargetVertexID, ClosestVertexID); // (TO, FROM)
+	}
+	
+	bIsInitialized = true;
+}
+
+
+void USkinWeightsPaintTool::MirrorWeights(EAxis::Type Axis, EMirrorDirection Direction)
+{
+	check(Axis != EAxis::None);
+	
+	// get all ref pose vertices
+	const TArray<FVector>& RefPoseVertices = Weights.Deformer.RefPoseVertexPositions;
+	const FReferenceSkeleton& RefSkeleton = Weights.Deformer.Component->GetSkeletalMeshAsset()->GetRefSkeleton();
+
+	// refresh mirror tables (cached / lazy generated)
+	MirrorData.RegenerateMirrorData(
+		Weights.Deformer.BoneNames,
+		Weights.Deformer.BoneNameToIndexMap,
+		RefSkeleton,
+		RefPoseVertices,
+		Axis,
+		Direction,
+		Weights.Deformer.Component->Bounds.SphereRadius);
+
+	// get a reference to the mirror tables
+	const TMap<int32, int32>& BoneMap = MirrorData.GetBoneMap();
+	const TMap<int32, int32>& VertexMirrorMap = MirrorData.GetVertexMap(); // <Target, Source>
+
+	// get set of vertices to mirror
+	TArray<VertexIndex> SelectedVertices;
+	GetVerticesToEdit(SelectedVertices);
+
+	// convert all vertex indices to the target side of the mirror plane
+	TSet<VertexIndex> VerticesToMirror;
+	
+	for (const VertexIndex SelectedVertex : SelectedVertices)
+	{
+		if (VertexMirrorMap.Contains(SelectedVertex))
+		{
+			// vertex is located across the mirror plane (target side, to copy TO)
+			VerticesToMirror.Add(SelectedVertex);
+		}
+		else
+		{
+			// vertex is located on the source side (to copy FROM), so we need to search for it's mirror target vertex
+			for (const TPair<int32, int32>& ToFromPair : VertexMirrorMap)
+			{
+				if (ToFromPair.Value != SelectedVertex)
+				{
+					continue;
+				}
+				VerticesToMirror.Add(ToFromPair.Key);
+				break;
+			}
+		}
+	}
+	
+	// spin through all target vertices to mirror and copy weights from source
+	FMultiBoneWeightEdits WeightEditsFromMirroring;
+	for (const VertexIndex VertexToMirror : VerticesToMirror)
+	{
+		const int32 SourceVertexID = VertexMirrorMap[VertexToMirror];
+		const int32 TargetVertexID = VertexToMirror;
+
+		// remove all weight on vertex
+		for (const FVertexBoneWeight& TargetBoneWeight : Weights.PreChangeWeights[TargetVertexID])
+		{
+			const float OldWeight = TargetBoneWeight.Weight;
+			constexpr float NewWeight = 0.f;
+			WeightEditsFromMirroring.MergeSingleEdit(TargetBoneWeight.BoneIndex, TargetVertexID, OldWeight, NewWeight);
+		}
+
+		// copy source weights, but with mirrored bones
+		for (const FVertexBoneWeight& SourceBoneWeight : Weights.PreChangeWeights[SourceVertexID])
+		{
+			const int32 MirroredBoneIndex = BoneMap[SourceBoneWeight.BoneIndex];
+			const float OldWeight = Weights.GetWeightOfBoneOnVertex(MirroredBoneIndex, TargetVertexID, Weights.PreChangeWeights);
+			const float NewWeight = SourceBoneWeight.Weight;
+			WeightEditsFromMirroring.MergeSingleEdit(MirroredBoneIndex, TargetVertexID, OldWeight, NewWeight);
+		}
+	}
+
+	// apply the changes
+	const FText TransactionLabel = LOCTEXT("MirrorWeightChange", "Mirror skin weights.");
+	ApplyWeightEditsToMesh(TransactionLabel, WeightEditsFromMirroring);
+
+	// warn if some vertices were not mirrored
+	if (!MirrorData.GetAllVerticesMirrored())
+	{
+		UE_LOG(LogMeshModelingToolsEditor, Log, TEXT("Mirror Skin Weights: some vertex weights were not mirrored because a vertex was not found close enough to the mirrored location."));
+	}
+}
+
+void USkinWeightsPaintTool::FloodWeights(const float Weight, const EWeightEditOperation FloodMode)
+{
+	if (CurrentBone == NAME_None)
+	{
+		return;
+	}
+
+	// flood the weights on selected vertices
+	FMultiBoneWeightEdits WeightEditsFromFlood;
+	TArray<VertexIndex> VerticesToEdit;
+	GetVerticesToEdit(VerticesToEdit);
+	const TArray<float> VertexFalloffs = {};
+	if (FloodMode == EWeightEditOperation::Relax)
+	{
+		RelaxWeightOnVertices(VerticesToEdit, VertexFalloffs, Weight, WeightEditsFromFlood);
+	}
+	else
+	{
+		EditWeightOnVertices(FloodMode, VerticesToEdit, VertexFalloffs, Weight, WeightEditsFromFlood);	
+	}
+
+	// apply the changes
+	const FText TransactionLabel = LOCTEXT("FloodWeightChange", "Flood skin weights.");
+	ApplyWeightEditsToMesh(TransactionLabel, WeightEditsFromFlood);
+}
+
+void USkinWeightsPaintTool::PruneWeights(float Threshold)
+{
+	// set weights below the given threshold to zero
+	FMultiBoneWeightEdits WeightEditsFromPrune;
+	TArray<VertexIndex> VerticesToEdit;
+	GetVerticesToEdit(VerticesToEdit);
+	const bool bPruningAllBones = SelectedBoneIndices.IsEmpty();
+	for (const VertexIndex VertexID : VerticesToEdit)
+	{
+		const VertexWeights& VertexWeights = Weights.CurrentWeights[VertexID];
+		for (const FVertexBoneWeight& BoneWeight : VertexWeights)
+		{
+			// are we pruning weights only on selected bones?
+			if (!bPruningAllBones && !SelectedBoneIndices.Contains(BoneWeight.BoneIndex))
+			{
+				// not the bone we're pruning...
+				continue;
+			}
+			
+			if (BoneWeight.Weight > Threshold)
+			{
+				continue;
+			}
+
+			// note: removing ALL weight from a vertex disqualifies it as a candidate to receive weight during normalization
+			// so there is no need to remove the influence from the per-vertex influence array to "prune" it.
+			Weights.EditVertexWeightAndNormalize(BoneWeight.BoneIndex, VertexID,0.f,WeightEditsFromPrune);
+		}
+	}
+
+	// apply the changes
+	const FText TransactionLabel = LOCTEXT("PruneWeightValuesChange", "Prune skin weights.");
+	ApplyWeightEditsToMesh(TransactionLabel, WeightEditsFromPrune);
+}
+
+void USkinWeightsPaintTool::AverageWeights()
+{
+	// get vertices to edit weights on
+	TArray<VertexIndex> VerticesToEdit;
+	GetVerticesToEdit(VerticesToEdit);
+
+	// accumulate ALL the weights on the vertices
+	TMap<BoneIndex, float> AccumulatedWeightMap;
+	for (const VertexIndex VertexID : VerticesToEdit)
+	{
+		for (const FVertexBoneWeight& BoneWeight : Weights.PreChangeWeights[VertexID])
+		{
+			float& AccumulatedWeight = AccumulatedWeightMap.FindOrAdd(BoneWeight.BoneIndex);
+			AccumulatedWeight += BoneWeight.Weight;
+		}
+	}
+
+	// sort influences by total weight
+	AccumulatedWeightMap.ValueSort([](const float& A, const float& B)
+	{
+		return A > B;
+	});
+
+	// truncate to MaxInfluences
+	int32 Index = 0;
+	for (TMap<BoneIndex, float>::TIterator It(AccumulatedWeightMap); It; ++It)
+	{
+		if (Index >= MAX_TOTAL_INFLUENCES)
+		{
+			It.RemoveCurrent();
+		}
+		else
+		{
+			++Index;
+		}
+	}
+
+	// normalize remaining influences
+	float TotalWeight = 0.f;
+	for (const TTuple<BoneIndex, float>& AccumulatedWeight : AccumulatedWeightMap)
+	{
+		TotalWeight += AccumulatedWeight.Value;
+	}
+	for (TTuple<BoneIndex, float>& AccumulatedWeight : AccumulatedWeightMap)
+	{
+		AccumulatedWeight.Value /= TotalWeight;
+	}
+
+	// apply averaged weights to vertices
+	FMultiBoneWeightEdits WeightEditsFromAveraging;
+	for (const VertexIndex VertexID : VerticesToEdit)
+	{
+		// remove influences not a part of the average results
+		for (const FVertexBoneWeight& BoneWeight : Weights.PreChangeWeights[VertexID])
+		{
+			if (!AccumulatedWeightMap.Contains(BoneWeight.BoneIndex))
+			{
+				const float OldWeight = BoneWeight.Weight;
+				constexpr float NewWeight = 0.f;
+				WeightEditsFromAveraging.MergeSingleEdit(BoneWeight.BoneIndex, VertexID, OldWeight, NewWeight);
+			}
+		}
+
+		// add influences from the averaging results
+		for (const TTuple<BoneIndex, float>& AccumulatedWeight : AccumulatedWeightMap)
+		{
+			const BoneIndex IndexOfBone = AccumulatedWeight.Key;
+			const float OldWeight = Weights.GetWeightOfBoneOnVertex(IndexOfBone, VertexID, Weights.PreChangeWeights);
+			const float NewWeight = AccumulatedWeight.Value;
+			WeightEditsFromAveraging.MergeSingleEdit(IndexOfBone, VertexID, OldWeight, NewWeight);
+		}
+	}
+
+	// apply the changes
+	const FText TransactionLabel = LOCTEXT("AverageWeightValuesChange", "Average skin weights.");
+	ApplyWeightEditsToMesh(TransactionLabel, WeightEditsFromAveraging);
+}
+
+void USkinWeightsPaintTool::NormalizeWeights()
+{
+	// re-set a weight on each vertex to force normalization
+	FMultiBoneWeightEdits WeightEditsFromNormalization;
+	TArray<VertexIndex> SelectedVertices;
+	GetVerticesToEdit(SelectedVertices);
+	for (const VertexIndex VertexID : SelectedVertices)
+	{
+		const VertexWeights& VertexWeights = Weights.CurrentWeights[VertexID];
+		for (const FVertexBoneWeight& BoneWeight : VertexWeights)
+		{
+			// set weight to current value, just to force re-normalization
+			Weights.EditVertexWeightAndNormalize(BoneWeight.BoneIndex, VertexID, BoneWeight.Weight,WeightEditsFromNormalization);
+		}
+	}
+
+	// apply the changes
+	const FText TransactionLabel = LOCTEXT("NormalizeWeightValuesChange", "Normalize skin weights.");
+	ApplyWeightEditsToMesh(TransactionLabel, WeightEditsFromNormalization);
+}
+
+void USkinWeightsPaintTool::HandleSkeletalMeshModified(const TArray<FName>& InBoneNames, const ESkeletalMeshNotifyType InNotifyType)
+{
+	switch (InNotifyType)
+	{
+	case ESkeletalMeshNotifyType::BonesAdded:
+		break;
+	case ESkeletalMeshNotifyType::BonesRemoved:
+		break;
+	case ESkeletalMeshNotifyType::BonesMoved:
+		break;
+	case ESkeletalMeshNotifyType::BonesSelected:
+		{
+			// store selected bones
+			SelectedBoneNames = InBoneNames;
+			PendingCurrentBone = InBoneNames.IsEmpty() ? NAME_None : InBoneNames[0];
+
+			// update selected bone indices from names
+			SelectedBoneIndices.Reset();
+			for (const FName SelectedBoneName : SelectedBoneNames)
+			{
+				SelectedBoneIndices.Add(GetBoneIndexFromName(SelectedBoneName));
+			}
+		}
+		break;
+	default:
+		checkNoEntry();
+	}
+}
+
+void USkinWeightsPaintTool::OnSelectionModified()
+{
+}
+
+void USkinWeightsPaintTool::ToggleEditingMode()
+{
+	SetBrushEnabled(WeightToolProperties->EditingMode == EWeightEditMode::Brush);
+	PolygonSelectionMechanic->SetIsEnabled(WeightToolProperties->EditingMode == EWeightEditMode::Selection);
+	
+}
+
+void USkinWeightsPaintTool::OnPropertyModified(UObject* ModifiedObject, FProperty* ModifiedProperty)
+{
+	Super::OnPropertyModified(ModifiedObject, ModifiedProperty);
+	
+	// invalidate vertex color cache when weight color properties are modified
+	const bool bColorModeModified = ModifiedProperty->GetNameCPP() == GET_MEMBER_NAME_STRING_CHECKED(USkinWeightsPaintToolProperties, ColorMode);
+	const bool bColorRampModified = ModifiedProperty->GetNameCPP() == GET_MEMBER_NAME_STRING_CHECKED(USkinWeightsPaintToolProperties, ColorRamp);
+	const bool bMinColorModified = ModifiedProperty->GetNameCPP() == GET_MEMBER_NAME_STRING_CHECKED(USkinWeightsPaintToolProperties, MinColor);
+	const bool bMaxColorModified = ModifiedProperty->GetNameCPP() == GET_MEMBER_NAME_STRING_CHECKED(USkinWeightsPaintToolProperties, MaxColor);
+	if (bColorModeModified || bColorRampModified || bMinColorModified || bMaxColorModified)
+	{
+		bVisibleWeightsValid = false;
+	}
+}
+
+
+#undef LOCTEXT_NAMESPACE
\ No newline at end of file
diff --git a/Engine/Plugins/Experimental/MeshModelingToolsetExp/Source/MeshModelingToolsEditorOnlyExp/Public/MeshModelingToolsEditorOnlyExp.h b/Engine/Plugins/Experimental/MeshModelingToolsetExp/Source/MeshModelingToolsEditorOnlyExp/Public/MeshModelingToolsEditorOnlyExp.h
--- a/Engine/Plugins/Experimental/MeshModelingToolsetExp/Source/MeshModelingToolsEditorOnlyExp/Public/MeshModelingToolsEditorOnlyExp.h
+++ b/Engine/Plugins/Experimental/MeshModelingToolsetExp/Source/MeshModelingToolsEditorOnlyExp/Public/MeshModelingToolsEditorOnlyExp.h
@@ -5,6 +5,8 @@
 #include "CoreMinimal.h"
 #include "Modules/ModuleManager.h"
 
+DECLARE_LOG_CATEGORY_EXTERN(LogMeshModelingToolsEditor, Log, All);
+
 class FMeshModelingToolsEditorOnlyExpModule : public IModuleInterface
 {
 public:
diff --git a/Engine/Plugins/Experimental/MeshModelingToolsetExp/Source/MeshModelingToolsEditorOnlyExp/Public/SkeletalMesh/SkeletalMeshEditionInterface.h b/Engine/Plugins/Experimental/MeshModelingToolsetExp/Source/MeshModelingToolsEditorOnlyExp/Public/SkeletalMesh/SkeletalMeshEditionInterface.h
new file mode 100644
--- /dev/null
+++ b/Engine/Plugins/Experimental/MeshModelingToolsetExp/Source/MeshModelingToolsEditorOnlyExp/Public/SkeletalMesh/SkeletalMeshEditionInterface.h
@@ -0,0 +1,137 @@
+﻿// Copyright Epic Games, Inc. All Rights Reserved.
+
+#pragma once
+
+#include "UObject/Interface.h"
+#include "UObject/WeakInterfacePtr.h"
+#include "SkeletalMeshNotifier.h"
+
+#include "HitProxies.h"
+#include "ToolContextInterfaces.h"
+#include "GenericPlatform/ICursor.h"
+
+#include "SkeletalMeshEditionInterface.generated.h"
+
+class FSkeletalMeshToolNotifier;
+class UTransformProxy;
+class UInteractiveToolManager;
+class IGizmoStateTarget;
+struct FInputDeviceRay;
+
+/**
+ * USkeletalMeshEditionInterface
+ */
+
+UINTERFACE()
+class MESHMODELINGTOOLSEDITORONLYEXP_API USkeletalMeshEditionInterface : public UInterface
+{
+	GENERATED_BODY()
+};
+
+/**
+ * ISkeletalMeshEditionInterface
+ */
+
+class MESHMODELINGTOOLSEDITORONLYEXP_API ISkeletalMeshEditionInterface
+{
+	GENERATED_BODY()
+
+public:
+	ISkeletalMeshNotifier& GetNotifier();
+	bool NeedsNotification() const;
+
+	virtual TArray<FName> GetSelectedBones() const;
+	
+	void BindTo(TSharedPtr<ISkeletalMeshEditorBinding> InBinding);
+	void Unbind();
+	
+protected:
+	virtual void HandleSkeletalMeshModified(const TArray<FName>& BoneNames, const ESkeletalMeshNotifyType InNotifyType) = 0;
+	TOptional<FName> GetBoneName(HHitProxy* InHitProxy) const;
+	
+	TWeakPtr<ISkeletalMeshEditorBinding> Binding;
+
+private:
+	TUniquePtr<FSkeletalMeshToolNotifier> Notifier;
+	
+	friend FSkeletalMeshToolNotifier;
+};
+
+/**
+ * FSkeletalMeshToolNotifier
+ */
+
+class FSkeletalMeshToolNotifier: public ISkeletalMeshNotifier
+{
+public:
+	FSkeletalMeshToolNotifier(TWeakInterfacePtr<ISkeletalMeshEditionInterface> InInterface);
+	virtual void HandleNotification(const TArray<FName>& BoneNames, const ESkeletalMeshNotifyType InNotifyType) override;
+	
+protected:
+	TWeakInterfacePtr<ISkeletalMeshEditionInterface> Interface;
+};
+
+/**
+ * HBoneHitProxy
+ */
+
+struct HBoneHitProxy : public HHitProxy
+{
+	DECLARE_HIT_PROXY()
+
+	int32 BoneIndex;
+	FName BoneName;
+
+	HBoneHitProxy(int32 InBoneIndex, FName InBoneName)
+		: HHitProxy(HPP_Foreground)
+		, BoneIndex(InBoneIndex)
+		, BoneName(InBoneName)
+	{}
+
+	virtual EMouseCursor::Type GetMouseCursor() override { return EMouseCursor::Crosshairs; }
+};
+
+/**
+ * USkeletalMeshGizmoContextObjectBase
+ */
+
+UCLASS(Abstract)
+class MESHMODELINGTOOLSEDITORONLYEXP_API USkeletalMeshGizmoContextObjectBase : public UObject
+{
+	GENERATED_BODY()
+	
+public:
+	virtual USkeletalMeshGizmoWrapperBase* GetNewWrapper(UInteractiveToolManager* InToolManager, UObject* Outer = nullptr, IGizmoStateTarget* InStateTarget = nullptr)
+	PURE_VIRTUAL(UGizmoContextObject::GetNewWrapper, return nullptr;);
+};
+
+/**
+ * USkeletalMeshGizmoWrapperBase
+ */
+
+UCLASS(Abstract)
+class MESHMODELINGTOOLSEDITORONLYEXP_API USkeletalMeshGizmoWrapperBase : public UObject
+{
+	GENERATED_BODY()
+
+protected:
+	using FGetTransform = TUniqueFunction<FTransform(void)>;
+	using FSetTransform = TUniqueFunction<void(const FTransform&)>;
+	
+public:
+	virtual void Initialize(const FTransform& InTransform = FTransform::Identity, const EToolContextCoordinateSystem& InTransformMode = EToolContextCoordinateSystem::Local)
+	PURE_VIRTUAL(USkeletalMeshGizmoWrapperBase::Initialize, return;)
+	
+	virtual void HandleBoneTransform(FGetTransform GetTransformFunc, FSetTransform SetTransformFunc)
+	PURE_VIRTUAL(USkeletalMeshGizmoWrapperBase::HandleBoneTransform, return;)
+
+	virtual void Clear() PURE_VIRTUAL(USkeletalMeshGizmoWrapperBase::Clear, return;)
+
+	virtual bool CanInteract() const
+	PURE_VIRTUAL(USkeletalMeshGizmoWrapperBase::CanInteract, return false;)
+	
+	virtual bool IsGizmoHit(const FInputDeviceRay& PressPos) const PURE_VIRTUAL(USkeletalMeshGizmoWrapperBase::IsGizmoHit, return false;)
+
+	UPROPERTY()
+	TWeakObjectPtr<USceneComponent> Component;
+};
\ No newline at end of file
diff --git a/Engine/Plugins/Experimental/MeshModelingToolsetExp/Source/MeshModelingToolsEditorOnlyExp/Public/SkeletalMesh/SkeletonEditingTool.h b/Engine/Plugins/Experimental/MeshModelingToolsetExp/Source/MeshModelingToolsEditorOnlyExp/Public/SkeletalMesh/SkeletonEditingTool.h
new file mode 100644
--- /dev/null
+++ b/Engine/Plugins/Experimental/MeshModelingToolsetExp/Source/MeshModelingToolsEditorOnlyExp/Public/SkeletalMesh/SkeletonEditingTool.h
@@ -0,0 +1,364 @@
+﻿// Copyright Epic Games, Inc. All Rights Reserved.
+
+#pragma once
+
+#include "DynamicMeshBrushTool.h"
+#include "BaseTools/MeshSurfacePointMeshEditingTool.h"
+
+#include "SkeletalMeshNotifier.h"
+#include "SkeletonModifier.h"
+#include "SkeletalMesh/SkeletalMeshEditionInterface.h"
+#include "Engine/World.h"
+#include "Changes/ValueWatcher.h"
+#include "Misc/EnumClassFlags.h"
+
+#include "SkeletonEditingTool.generated.h"
+
+class USingleClickInputBehavior;
+class UGizmoViewContext;
+class UTransformGizmo;
+class USkeletonEditingTool;
+class USkeletonTransformProxy;
+class USkeletalMeshGizmoContextObjectBase;
+class USkeletalMeshGizmoWrapperBase;
+
+namespace SkeletonEditingTool
+{
+
+/**
+ * A wrapper change class that stores a reference skeleton and the bones' indexes trackers to be used for undo/redo.
+ */
+class MESHMODELINGTOOLSEDITORONLYEXP_API FRefSkeletonChange : public FToolCommandChange
+{
+public:
+	FRefSkeletonChange(const USkeletonEditingTool* InTool);
+
+	virtual FString ToString() const override
+	{
+		return FString(TEXT("Edit Skeleton"));
+	}
+
+	virtual void Apply(UObject* Object) override;
+	virtual void Revert(UObject* Object) override;
+
+	void StoreSkeleton(const USkeletonEditingTool* InTool);
+
+private:
+	FReferenceSkeleton PreChangeSkeleton;
+	TArray<int32> PreBoneTracker;
+	FReferenceSkeleton PostChangeSkeleton;
+	TArray<int32> PostBoneTracker;
+};
+	
+}
+
+/**
+ * USkeletonEditingToolBuilder
+ */
+
+UCLASS()
+class MESHMODELINGTOOLSEDITORONLYEXP_API USkeletonEditingToolBuilder : public UInteractiveToolWithToolTargetsBuilder
+{
+	GENERATED_BODY()
+
+public:
+	// UInteractiveToolBuilder overrides
+	virtual bool CanBuildTool(const FToolBuilderState& SceneState) const override;
+	virtual UInteractiveTool* BuildTool(const FToolBuilderState& SceneState) const override;
+	
+protected:
+	// UInteractiveToolWithToolTargetsBuilder overrides
+	virtual const FToolTargetTypeRequirements& GetTargetRequirements() const override;
+};
+
+/**
+ * EEditingOperation represents the current tool's operation 
+ */
+
+UENUM()
+enum class EEditingOperation : uint8
+{
+	Select,		// Selecting bones in the viewport.
+	Create,		// Creating bones in the viewport.
+	Remove,		// Removing current selection.
+	Transform,	// Transforming bones in the viewport or thru the details panel.
+	Parent,		// Parenting bones in the viewport.
+	Rename,		// Renaming bones thru the details panel.
+	Mirror		// Mirroring bones thru the details panel.
+};
+
+/**
+ * USkeletonEditingTool is a tool to edit a the ReferenceSkeleton of a SkeletalMesh (target)
+ * Changed are actually commit to the SkeletalMesh and it's mesh description on Accept.
+ */
+
+UCLASS()
+class MESHMODELINGTOOLSEDITORONLYEXP_API USkeletonEditingTool :
+	public USingleSelectionTool,
+	public IClickDragBehaviorTarget,
+	public ISkeletalMeshEditionInterface
+{
+	GENERATED_BODY()
+
+public:
+
+	void Init(const FToolBuilderState& InSceneState);
+	
+	// UInteractiveTool overrides
+	virtual void Setup() override;
+	virtual void Shutdown(EToolShutdownType ShutdownType) override;
+
+	virtual void Render(IToolsContextRenderAPI* RenderAPI) override;
+	
+	virtual bool HasCancel() const override { return true; }
+	virtual bool HasAccept() const override { return true; }
+	virtual bool CanAccept() const override { return true; }
+
+	// ICLickDragBehaviorTarget overrides
+	virtual FInputRayHit CanBeginClickDragSequence(const FInputDeviceRay& InPressPos) override;
+	
+	virtual void OnClickPress(const FInputDeviceRay& InPressPos) override;
+	virtual void OnClickDrag(const FInputDeviceRay& InDragPos) override;
+	virtual void OnClickRelease(const FInputDeviceRay& InReleasePos) override;
+	virtual void OnTerminateDragSequence() override;
+	
+	virtual void OnTick(float DeltaTime) override;
+
+	// IInteractiveToolCameraFocusAPI overrides
+	virtual FBox GetWorldSpaceFocusBox() override;
+
+	// Modifier functions
+	void MirrorBones();
+	void RenameBones();
+	void MoveBones();
+	void OrientBones();
+
+	// ISkeletalMeshEditionInterface overrides
+	virtual TArray<FName> GetSelectedBones() const override;
+
+	const TArray<FName>& GetSelection() const;
+
+	const FTransform& GetTransform(const FName InBoneName, const bool bWorld) const;
+	void SetTransforms(const TArray<FName>& InBones, const TArray<FTransform>& InTransforms, const bool bWorld) const;
+
+	// IModifierToggleBehaviorTarget overrides
+	virtual void OnUpdateModifierState(int ModifierID, bool bIsOn) override;
+
+protected:
+
+	// UInteractiveTool overrides
+	virtual void RegisterActions(FInteractiveToolActionSet& ActionSet) override;
+
+	// ISkeletalMeshEditionInterface overrides
+	virtual void HandleSkeletalMeshModified(const TArray<FName>& BoneNames, const ESkeletalMeshNotifyType InNotifyType) override;
+
+	// Modifier functions
+	void CreateNewBone();
+	void RemoveBones();
+	void UnParentBones();
+	void ParentBones(const FName& InParentName);
+
+	UPROPERTY()
+	TObjectPtr<USkeletonModifier> Modifier;
+	
+	UPROPERTY()
+	TObjectPtr<USkeletonEditingProperties> Properties;
+
+	UPROPERTY()
+	TObjectPtr<UProjectionProperties> ProjectionProperties;
+	
+	UPROPERTY()
+	TObjectPtr<UMirroringProperties> MirroringProperties;
+
+	UPROPERTY()
+	TObjectPtr<UOrientingProperties> OrientingProperties;
+	
+	UPROPERTY()
+	TObjectPtr<USingleClickInputBehavior> LeftClickBehavior = nullptr;
+
+	UPROPERTY()
+	TObjectPtr<UPreviewMesh> PreviewMesh = nullptr;
+
+	UPROPERTY()
+	TWeakObjectPtr<UWorld> TargetWorld = nullptr;
+
+	UPROPERTY()
+	TWeakObjectPtr<UGizmoViewContext> ViewContext = nullptr;
+
+	UPROPERTY()
+	EEditingOperation Operation = EEditingOperation::Select;
+
+	// gizmo
+	UPROPERTY()
+	TWeakObjectPtr<USkeletalMeshGizmoContextObjectBase> GizmoContext = nullptr;
+
+	UPROPERTY()
+	TObjectPtr<USkeletalMeshGizmoWrapperBase> GizmoWrapper = nullptr;
+
+	void UpdateGizmo() const;
+	
+	// ref skeleton transactions
+	void BeginChange();
+	void EndChange();
+	void CancelChange();
+	TUniquePtr<SkeletonEditingTool::FRefSkeletonChange> ActiveChange;
+
+	friend class SkeletonEditingTool::FRefSkeletonChange;
+	
+private:
+	TArray<int32> GetSelectedBoneIndexes() const;
+	
+	enum class EBoneSelectionMode : uint8
+	{
+		Single				= 0,
+		Additive			= 1 << 0,
+		Toggle				= 1 << 1
+	};
+	FRIEND_ENUM_CLASS_FLAGS(USkeletonEditingTool::EBoneSelectionMode)
+	
+
+	// flags used to identify behavior modifier keys/buttons
+	static constexpr int AddToSelectionModifier = 1;
+	static constexpr int ToggleSelectionModifier = 2;
+	EBoneSelectionMode SelectionMode = EBoneSelectionMode::Single;
+
+	TArray<FName> Selection;
+	
+	void SelectBone(const FName& InBoneName);
+	FName GetCurrentBone() const;
+
+	void NormalizeSelection();
+
+	TValueWatcher<TArray<FName>> SelectionWatcher;
+	TValueWatcher<EToolContextCoordinateSystem> CoordinateSystemWatcher; 
+	
+	int32 ParentIndex = INDEX_NONE;
+
+	TFunction<void()> PendingFunction;
+};
+
+ENUM_CLASS_FLAGS(USkeletonEditingTool::EBoneSelectionMode);
+
+/**
+ * USkeletonEditingProperties
+ */
+
+UCLASS()
+class MESHMODELINGTOOLSEDITORONLYEXP_API USkeletonEditingProperties : public UInteractiveToolPropertySet
+{
+	GENERATED_BODY()
+public:
+	void Initialize(USkeletonEditingTool* ParentToolIn);
+
+	UPROPERTY(EditAnywhere, Category = "Details")
+	FName Name;
+
+	UPROPERTY()
+	FTransform Transform;
+	
+	UPROPERTY(EditAnywhere, Category = "Move")
+	bool bUpdateChildren = false;
+
+	UPROPERTY(EditAnywhere, Category = "Viewport Axis Settings",  meta = (ClampMin = "0.0", UIMin = "0.0"))
+	float AxisLength = 1.f;
+
+	UPROPERTY(EditAnywhere, Category = "Viewport Axis Settings",  meta = (ClampMin = "0.0", UIMin = "0.0"))
+	float AxisThickness = 0.f;
+
+#if WITH_EDITOR
+	virtual void PostEditChangeProperty(FPropertyChangedEvent &PropertyChangedEvent) override;
+#endif
+	
+	TWeakObjectPtr<USkeletonEditingTool> ParentTool = nullptr;
+};
+
+/**
+ * EProjectionType
+ */
+
+UENUM()
+enum class EProjectionType : uint8
+{
+	CameraPlane,	// The camera plane is used as the projection plane 
+	OnMesh,			// The mesh surface is used for projection
+	WithinMesh,		// The mesh volume is used for projection
+};
+
+/**
+ * UProjectionProperties
+ */
+
+UCLASS()
+class MESHMODELINGTOOLSEDITORONLYEXP_API UProjectionProperties: public UInteractiveToolPropertySet
+{
+	GENERATED_BODY()
+
+public:
+	void Initialize(TObjectPtr<UPreviewMesh> PreviewMesh);
+
+	void UpdatePlane(const UGizmoViewContext& InViewContext, const FVector& InOrigin);
+	bool GetProjectionPoint(const FInputDeviceRay& InRay, FVector& OutHitPoint) const;
+	
+	UPROPERTY(EditAnywhere, Category = "Project")
+	EProjectionType ProjectionType = EProjectionType::WithinMesh;
+	
+private:
+	TWeakObjectPtr<UPreviewMesh> PreviewMesh = nullptr;
+	
+	UPROPERTY()
+	FVector PlaneOrigin = FVector::ZeroVector;
+	
+	UPROPERTY()
+	FVector PlaneNormal =  FVector::ZAxisVector;
+};
+
+/**
+ * UMirroringProperties
+ */
+
+UCLASS()
+class MESHMODELINGTOOLSEDITORONLYEXP_API UMirroringProperties: public UInteractiveToolPropertySet
+{
+	GENERATED_BODY()
+
+public:
+	void Initialize(USkeletonEditingTool* ParentToolIn);
+
+	UFUNCTION(CallInEditor, Category = "Mirror")
+	void MirrorBones();
+
+	UPROPERTY(EditAnywhere, Category = "Mirror")
+	FMirrorOptions Options;
+
+private:
+	TWeakObjectPtr<USkeletonEditingTool> ParentTool = nullptr;
+};
+
+/**
+ * UOrientingProperties
+ */
+
+UCLASS()
+class MESHMODELINGTOOLSEDITORONLYEXP_API UOrientingProperties: public UInteractiveToolPropertySet
+{
+	GENERATED_BODY()
+
+public:
+	void Initialize(USkeletonEditingTool* ParentToolIn);
+
+	UFUNCTION(CallInEditor, Category = "Orient")
+	void OrientBones();
+	
+	UPROPERTY(EditAnywhere, Category = "Orient")
+	FOrientOptions Options;
+
+	UPROPERTY(EditAnywhere, Category = "Orient")
+	bool bAutoOrient = false;
+
+#if WITH_EDITOR
+	virtual void PostEditChangeProperty(FPropertyChangedEvent &PropertyChangedEvent) override;
+#endif
+	
+private:
+	TWeakObjectPtr<USkeletonEditingTool> ParentTool = nullptr;
+};
\ No newline at end of file
diff --git a/Engine/Plugins/Experimental/MeshModelingToolsetExp/Source/MeshModelingToolsEditorOnlyExp/Public/SkeletalMesh/SkeletonTransformProxy.h b/Engine/Plugins/Experimental/MeshModelingToolsetExp/Source/MeshModelingToolsEditorOnlyExp/Public/SkeletalMesh/SkeletonTransformProxy.h
new file mode 100644
--- /dev/null
+++ b/Engine/Plugins/Experimental/MeshModelingToolsetExp/Source/MeshModelingToolsEditorOnlyExp/Public/SkeletalMesh/SkeletonTransformProxy.h
@@ -0,0 +1,37 @@
+﻿// Copyright Epic Games, Inc. All Rights Reserved.
+
+#pragma once
+
+#include "BaseGizmos/TransformProxy.h"
+
+#include "SkeletonTransformProxy.generated.h"
+
+class UObject;
+
+/**
+ * USkeletonTransformProxy is a derivation of UTransformProxy that manages several bones and update their transform individually.
+ */
+
+UCLASS(Transient)
+class MESHMODELINGTOOLSEDITORONLYEXP_API USkeletonTransformProxy : public UTransformProxy
+{
+	GENERATED_BODY()
+
+public:
+
+	void Initialize(const FTransform& InTransform, const EToolContextCoordinateSystem& InTransformMode);
+	
+	bool IsValid() const;
+
+	// UTransformProxy overrides
+	virtual FTransform GetTransform() const override;
+
+	UPROPERTY()
+	EToolContextCoordinateSystem TransformMode = EToolContextCoordinateSystem::Local;
+	
+protected:
+	// UTransformProxy overrides
+	virtual void UpdateSharedTransform() override;
+	virtual void UpdateObjects() override;
+};
+
diff --git a/Engine/Plugins/Experimental/MeshModelingToolsetExp/Source/MeshModelingToolsEditorOnlyExp/Public/SkeletalMesh/SkinWeightsBindingTool.h b/Engine/Plugins/Experimental/MeshModelingToolsetExp/Source/MeshModelingToolsEditorOnlyExp/Public/SkeletalMesh/SkinWeightsBindingTool.h
new file mode 100644
--- /dev/null
+++ b/Engine/Plugins/Experimental/MeshModelingToolsetExp/Source/MeshModelingToolsEditorOnlyExp/Public/SkeletalMesh/SkinWeightsBindingTool.h
@@ -0,0 +1,126 @@
+// Copyright Epic Games, Inc. All Rights Reserved.
+
+#pragma once
+
+#include "CoreMinimal.h"
+
+#include "BoneContainer.h"
+#include "InteractiveToolBuilder.h"
+#include "ModelingOperators.h"
+#include "BaseTools/MultiSelectionMeshEditingTool.h"
+#include "Interfaces/Interface_BoneReferenceSkeletonProvider.h"
+
+#include "SkinWeightsBindingTool.generated.h"
+
+class UMeshOpPreviewWithBackgroundCompute;
+struct FDynamicMeshOpResult;
+struct FMeshDescription;
+class IMeshDescriptionCommitter;
+struct FOccupancyGrid;
+
+namespace UE::Geometry
+{
+	struct FOccupancyGrid3;
+}
+
+/**
+ *
+ */
+UCLASS()
+class MESHMODELINGTOOLSEDITORONLYEXP_API USkinWeightsBindingToolBuilder : public UMultiSelectionMeshEditingToolBuilder
+{
+	GENERATED_BODY()
+public:
+	virtual bool CanBuildTool(const FToolBuilderState& SceneState) const override;
+	virtual UMultiSelectionMeshEditingTool* CreateNewTool(const FToolBuilderState& SceneState) const override;
+};
+
+// A mirror of UE::Geometry::ESkinBindingType
+UENUM()
+enum class ESkinWeightsBindType : uint8
+{
+	DirectDistance = 0 UMETA(DisplayName = "Direct Distance"),
+	GeodesicVoxel = 1 UMETA(DisplayName = "Geodesic Voxel"),
+};
+
+UCLASS()
+class MESHMODELINGTOOLSEDITORONLYEXP_API USkinWeightsBindingToolProperties : 
+	public UInteractiveToolPropertySet,
+	public IBoneReferenceSkeletonProvider
+{
+	GENERATED_BODY()
+public:
+	UPROPERTY(VisibleAnywhere, Category = Skeleton)
+	FBoneReference CurrentBone;
+
+	/** Binding type to use */
+	UPROPERTY(EditAnywhere, Category = Binding)
+	ESkinWeightsBindType BindingType = ESkinWeightsBindType::DirectDistance;
+
+	/** Stiffness of binding. Lower values allow more distant bones to contribute more */
+	UPROPERTY(EditAnywhere, Category = Binding, meta=(ClampMin="0.0", ClampMax="1.0", UIMin="0.0", UIMax="1.0"))
+	float Stiffness = 0.2f;
+
+	/** Maximum bones that will influence each vertex */
+	UPROPERTY(EditAnywhere, Category = Binding, meta=(ClampMin="1", UIMin="1", UIMax="10"))
+	int32 MaxInfluences = 5;
+
+	/** The resolution of the voxel grid if doing geodesic voxel binding */
+	UPROPERTY(EditAnywhere, Category = Binding, meta=(EditCondition = "BindingType == ESkinWeightsBindType::GeodesicVoxel", ClampMin="1", UIMin="8", UIMax="1024"))
+	int32 VoxelResolution = 128;
+	
+	// UPROPERTY(EditAnywhere, Category = Debug)
+	bool bDebugDraw = false;
+	
+	// IBoneReferenceSkeletonProvider
+	virtual USkeleton* GetSkeleton(bool& bInvalidSkeletonIsError, const IPropertyHandle* PropertyHandle) override;
+
+	TObjectPtr<USkeletalMesh> SkeletalMesh;
+};
+
+/**
+ *
+ */
+UCLASS()
+class MESHMODELINGTOOLSEDITORONLYEXP_API USkinWeightsBindingTool :
+	public UMultiSelectionMeshEditingTool,
+	public UE::Geometry::IDynamicMeshOperatorFactory
+{
+	GENERATED_BODY()
+
+public:
+	USkinWeightsBindingTool();
+	virtual ~USkinWeightsBindingTool() override;
+	
+	virtual void Setup() override;
+	virtual void OnShutdown(EToolShutdownType ShutdownType) override;
+	
+	virtual void OnTick(float DeltaTime) override;
+	virtual void Render(IToolsContextRenderAPI* RenderAPI) override;
+
+	virtual bool HasCancel() const override { return true; }
+	virtual bool HasAccept() const override { return true; }
+	virtual bool CanAccept() const override;
+
+	// IDynamicMeshOperatorFactory API
+	virtual TUniquePtr<UE::Geometry::FDynamicMeshOperator> MakeNewOperator() override;
+
+	UPROPERTY()
+	TObjectPtr<USkinWeightsBindingToolProperties> Properties = nullptr;
+	
+	UPROPERTY()
+	TObjectPtr<UMeshOpPreviewWithBackgroundCompute> Preview = nullptr;
+
+protected:
+	TSharedPtr<UE::Geometry::FOccupancyGrid3> Occupancy;
+	
+	TSharedPtr<UE::Geometry::FDynamicMesh3, ESPMode::ThreadSafe> OriginalMesh;
+
+	FReferenceSkeleton ReferenceSkeleton;
+	FBoneContainer BoneContainer;
+	
+	void GenerateAsset(const FDynamicMeshOpResult& Result);
+
+	static FVector4f WeightToColor(float Value);
+	void UpdateVisualization(bool bInForce = false);
+};
diff --git a/Engine/Plugins/Experimental/MeshModelingToolsetExp/Source/MeshModelingToolsEditorOnlyExp/Public/SkeletalMesh/SkinWeightsPaintTool.h b/Engine/Plugins/Experimental/MeshModelingToolsetExp/Source/MeshModelingToolsEditorOnlyExp/Public/SkeletalMesh/SkinWeightsPaintTool.h
new file mode 100644
--- /dev/null
+++ b/Engine/Plugins/Experimental/MeshModelingToolsetExp/Source/MeshModelingToolsEditorOnlyExp/Public/SkeletalMesh/SkinWeightsPaintTool.h
@@ -0,0 +1,406 @@
+// Copyright Epic Games, Inc. All Rights Reserved.
+
+#pragma once
+
+#include "DynamicMeshBrushTool.h"
+#include "BaseTools/MeshSurfacePointMeshEditingTool.h"
+#include "MeshDescription.h"
+#include "DynamicMesh/DynamicVerticesOctree3.h"
+#include "BoneWeights.h"
+#include "GroupTopology.h"
+#include "Misc/Optional.h"
+#include "Containers/Map.h"
+#include "DynamicMesh/DynamicMeshOctree3.h"
+#include "SkeletalMesh/SkeletalMeshEditionInterface.h"
+#include "Engine/SkeletalMesh.h"
+
+#include "SkinWeightsPaintTool.generated.h"
+
+
+struct FMeshDescription;
+class USkinWeightsPaintTool;
+class UPolygonSelectionMechanic;
+
+using BoneIndex = int32;
+using VertexIndex = int32;
+
+// weight edit mode
+UENUM()
+enum class EWeightEditMode : uint8
+{
+	Brush,
+	Selection,
+};
+
+// weight color mode
+UENUM()
+enum class EWeightColorMode : uint8
+{
+	Greyscale,
+	ColorRamp,
+};
+
+// brush falloff mode
+UENUM()
+enum class EWeightBrushFalloffMode : uint8
+{
+	Surface,
+	Volume,
+};
+
+// operation type when editing weights
+UENUM()
+enum class EWeightEditOperation : uint8
+{
+	Add,
+	Replace,
+	Multiply,
+	Relax
+};
+
+// mirror direction mode
+UENUM()
+enum class EMirrorDirection : uint8
+{
+	PositiveToNegative,
+	NegativeToPositive,
+};
+
+namespace SkinPaintTool
+{
+	struct FSkinToolWeights;
+
+	struct FVertexBoneWeight
+	{
+		FVertexBoneWeight() : BoneIndex(INDEX_NONE), VertexInBoneSpace(FVector::ZeroVector), Weight(0.0f) {}
+		FVertexBoneWeight(int32 InBoneIndex, const FVector& InPosInRefPose, float InWeight) :
+			BoneIndex(InBoneIndex), VertexInBoneSpace(InPosInRefPose), Weight(InWeight){}
+		
+		int32 BoneIndex;
+		FVector VertexInBoneSpace;
+		float Weight;
+	};
+
+	using VertexWeights = TArray<FVertexBoneWeight, TFixedAllocator<UE::AnimationCore::MaxInlineBoneWeightCount>>;
+
+	// data required to preview the skinning deformations as you paint
+	struct FSkinToolDeformer
+	{
+		void Initialize(const USkeletalMeshComponent* SkeletalMeshComponent, const FMeshDescription* Mesh);
+
+		void UpdateVertexDeformation(USkinWeightsPaintTool* Tool);
+
+		void SetVertexNeedsUpdated(int32 VertexIndex);
+		
+		// which vertices require updating (partially re-calculated skinning deformation while painting)
+		TSet<int32> VerticesWithModifiedWeights;
+		// position of all vertices in the reference pose
+		TArray<FVector> RefPoseVertexPositions;
+		// inverted, component space ref pose transform of each bone
+		TArray<FTransform> InvCSRefPoseTransforms;
+		// bone index to bone name
+		TArray<FName> BoneNames;
+		TMap<FName, BoneIndex> BoneNameToIndexMap;
+		// the skeletal mesh to get the current pose from
+		const USkeletalMeshComponent* Component;
+	};
+
+	// store a sparse set of modifications to a set of vertex weights on a SINGLE bone
+	struct FSingleBoneWeightEdits
+	{
+		int32 BoneIndex;
+		TMap<VertexIndex, float> OldWeights;
+		TMap<VertexIndex, float> NewWeights;
+	};
+
+	// store a sparse set of modifications to a set of vertex weights for a SET of bones
+	// with support for merging edits. these are used for transaction history undo/redo.
+	struct FMultiBoneWeightEdits
+	{
+		void MergeSingleEdit(const int32 BoneIndex, const int32 VertexID, const float OldWeight, const float NewWeight);
+		void MergeEdits(const FSingleBoneWeightEdits& BoneWeightEdits);
+		float GetVertexDeltaFromEdits(const int32 BoneIndex, const int32 VertexIndex);
+
+		// map of bone indices to weight edits made to that bone
+		TMap<BoneIndex, FSingleBoneWeightEdits> PerBoneWeightEdits;
+	};
+
+	class FMeshSkinWeightsChange : public FToolCommandChange
+	{
+	public:
+		FMeshSkinWeightsChange() : FToolCommandChange()	{ }
+
+		virtual FString ToString() const override
+		{
+			return FString(TEXT("Paint Skin Weights"));
+		}
+
+		virtual void Apply(UObject* Object) override;
+
+		virtual void Revert(UObject* Object) override;
+
+		void AddBoneWeightEdit(const FSingleBoneWeightEdits& BoneWeightEdit);
+
+	private:
+		FMultiBoneWeightEdits AllWeightEdits;
+	};
+
+	// intermediate storage of the weight maps for duration of tool
+	struct FSkinToolWeights
+	{
+		// copy the initial weight values from the skeletal mesh
+		void InitializeSkinWeights(
+			const USkeletalMeshComponent* SkeletalMeshComponent,
+			FMeshDescription* Mesh);
+
+		// applies an edit to a single vertex weight on a single bone, then normalizes the remaining weights while
+		// keeping the edited weight intact (ie, adapts OTHER influences to achieve normalization)
+		void EditVertexWeightAndNormalize(
+			const int32 BoneIndex,
+			const int32 VertexId,
+			float NewWeightValue,
+			FMultiBoneWeightEdits& WeightEdits);
+
+		void ApplyCurrentWeightsToMeshDescription(FMeshDescription* EditedMesh);
+		
+		static float GetWeightOfBoneOnVertex(
+			const int32 BoneIndex,
+			const int32 VertexID,
+			const TArray<VertexWeights>& InVertexWeights);
+
+		void SetWeightOfBoneOnVertex(
+			const int32 BoneIndex,
+			const int32 VertexID,
+			const float Weight,
+			TArray<VertexWeights>& InOutVertexData);
+
+		void ResetAfterChange();
+
+		float SetCurrentFalloffAndGetMaxFalloffThisStroke(int32 VertexID, float CurrentStrength);
+
+		void ApplyEditsToWeightMap(const FMultiBoneWeightEdits& Edits, TArray<VertexWeights>& InOutWeights);
+		
+		// double-buffer of the entire weight matrix (stored sparsely for fast deformation)
+		// "Pre" is state of weights at stroke start
+		// "Current" is state of weights during stroke
+		// When stroke is over, PreChangeWeights are synchronized with CurrentWeights
+		TArray<VertexWeights> PreChangeWeights;
+		TArray<VertexWeights> CurrentWeights;
+
+		// record the current maximum amount of falloff applied to each vertex during the current stroke
+		// values range from 0-1, this allows brushes to sweep over the same vertex, and apply only the maximum amount
+		// of modification (add/replace/relax etc) that was encountered for the duration of the stroke.
+		TArray<float> MaxFalloffPerVertexThisStroke;
+
+		// update deformation when vertex weights are modified
+		FSkinToolDeformer Deformer;
+	};
+
+	struct FSkinMirrorData
+	{
+		void RegenerateMirrorData(
+			const TArray<FName>& BoneNames,
+			const TMap<FName, BoneIndex>& BoneNameToIndexMap,
+			const FReferenceSkeleton& RefSkeleton,
+			const TArray<FVector>& RefPoseVertices,
+			EAxis::Type InMirrorAxis,
+			EMirrorDirection InMirrorDirection,
+			const float MaxSearchRadius);
+
+		const TMap<int32, int32>& GetBoneMap() const { return BoneMap; };
+		const TMap<int32, int32>& GetVertexMap() const { return VertexMap; };
+		bool GetAllVerticesMirrored() const {return bAllVerticesMirrored; };
+		
+	private:
+		
+		bool bIsInitialized = false;
+		bool bAllVerticesMirrored = false;
+		TEnumAsByte<EAxis::Type> Axis;
+		EMirrorDirection Direction; 
+		TMap<int32, int32> BoneMap;
+		TMap<int32, int32> VertexMap; // <Target, Source>
+	};
+}
+
+UCLASS()
+class MESHMODELINGTOOLSEDITORONLYEXP_API USkinWeightsPaintToolBuilder : public UMeshSurfacePointMeshEditingToolBuilder
+{
+	GENERATED_BODY()
+
+public:
+	UMeshSurfacePointTool* CreateNewTool(const FToolBuilderState& SceneState) const override;
+};
+
+// Container for properties displayed in Details panel while using USkinWeightsPaintTool
+UCLASS()
+class MESHMODELINGTOOLSEDITORONLYEXP_API USkinWeightsPaintToolProperties : public UBrushBaseProperties
+{
+	GENERATED_BODY()
+
+	USkinWeightsPaintToolProperties();
+	
+public:
+
+	// brush vs selection modes
+	UPROPERTY()
+	EWeightEditMode EditingMode;
+
+	// custom brush modes and falloff types
+	UPROPERTY()
+	EWeightEditOperation BrushMode;
+	UPROPERTY()
+	EWeightBrushFalloffMode FalloffMode;
+
+	// weight color properties
+	UPROPERTY(EditAnywhere, Category = WeightColors)
+	EWeightColorMode ColorMode;
+	UPROPERTY(EditAnywhere, Category = WeightColors)
+	TArray<FLinearColor> ColorRamp;
+	UPROPERTY(EditAnywhere, Category = WeightColors)
+	FLinearColor MinColor;
+	UPROPERTY(EditAnywhere, Category = WeightColors)
+	FLinearColor MaxColor;
+	bool bColorModeChanged = false;
+
+	// weight editing arguments
+	UPROPERTY()
+	TEnumAsByte<EAxis::Type> MirrorAxis = EAxis::X;
+	UPROPERTY()
+	EMirrorDirection MirrorDirection = EMirrorDirection::PositiveToNegative;
+	UPROPERTY()
+	EWeightEditOperation FloodMode;
+	UPROPERTY()
+	float FloodValue = 1.f;
+	UPROPERTY()
+	float PruneValue = 0.01;
+
+	TObjectPtr<USkinWeightsPaintTool> WeightTool;
+};
+
+// An interactive tool for painting and editing skin weights.
+UCLASS()
+class MESHMODELINGTOOLSEDITORONLYEXP_API USkinWeightsPaintTool : public UDynamicMeshBrushTool, public ISkeletalMeshEditionInterface
+{
+	GENERATED_BODY()
+
+public:
+
+	// UBaseBrushTool overrides
+	virtual bool HasCancel() const override { return true; }
+	virtual bool HasAccept() const override { return true; }
+	virtual bool CanAccept() const override { return true; }
+	virtual bool HitTest(const FRay& Ray, FHitResult& OutHit) override;
+	virtual void OnBeginDrag(const FRay& Ray) override;
+	virtual void OnUpdateDrag(const FRay& Ray) override;
+	virtual void OnEndDrag(const FRay& Ray) override;
+	virtual bool OnUpdateHover(const FInputDeviceRay& DevicePos) override;
+	
+	// UInteractiveTool
+	virtual void Setup() override;
+	virtual void DrawHUD(FCanvas* Canvas, IToolsContextRenderAPI* RenderAPI) override;
+	virtual void Render(IToolsContextRenderAPI* RenderAPI) override;
+
+	void ExternalUpdateWeights(const int32 BoneIndex, const TMap<int32, float>& IndexValues);
+
+	// weight editing operations (selection based)
+	void MirrorWeights(EAxis::Type Axis, EMirrorDirection Direction);
+	void FloodWeights(const float Weight, const EWeightEditOperation FloodMode);
+	void PruneWeights(const float Threshold);
+	void AverageWeights();
+	void NormalizeWeights();
+
+	// toggle brush / selection mode
+	void ToggleEditingMode();
+
+protected:
+
+	virtual void ApplyStamp(const FBrushStampData& Stamp);
+	void OnShutdown(EToolShutdownType ShutdownType) override;
+	void OnTick(float DeltaTime) override;
+
+	// stamp
+	float CalculateBrushFalloff(float Distance) const;
+	void CalculateVertexROI(
+		const FBrushStampData& Stamp,
+		TArray<VertexIndex>& VertexIDs,
+		TArray<float>& VertexFalloffs);
+	float CalculateBrushStrengthToUse(EWeightEditOperation EditMode) const;
+	bool bInvertStroke = false;
+	bool bSmoothStroke = false;
+	FBrushStampData StartStamp;
+	FBrushStampData LastStamp;
+	bool bStampPending;
+	int32 TriangleUnderStamp;
+	FVector StampLocalPos;
+
+	// modify vertex weights according to the specified operation,
+	// generating bone weight edits to be stored in a transaction
+	void EditWeightOnVertices(
+		EWeightEditOperation EditOperation,
+		const TArray<int32>& VerticesToEdit,
+		const TArray<float>& VertexFalloffs,
+		const float UseStrength,
+		SkinPaintTool::FMultiBoneWeightEdits& InOutWeightEdits);
+	// same as EditWeightOfVertices() but specific to relaxation (topology aware operation)
+	void RelaxWeightOnVertices(
+		TArray<int32> VerticesToEdit,
+		TArray<float> VertexFalloffs,
+		const float UseStrength,
+		SkinPaintTool::FMultiBoneWeightEdits& InOutWeightEdits);
+
+	// used to accelerate mesh queries
+	UE::Geometry::TDynamicVerticesOctree3<FDynamicMesh3> VerticesOctree;
+	UE::Geometry::FDynamicMeshOctree3 TrianglesOctree;
+	TFuture<void> TriangleOctreeFuture;
+	TArray<int32> TrianglesToReinsert;
+
+	// tool properties
+	UPROPERTY()
+	TObjectPtr<USkinWeightsPaintToolProperties> WeightToolProperties;
+	virtual void OnPropertyModified(UObject* ModifiedObject, FProperty* ModifiedProperty) override;
+	
+	// the currently edited mesh description
+	TUniquePtr<FMeshDescription> EditedMesh;
+
+	// storage of vertex weights per bone 
+	SkinPaintTool::FSkinToolWeights Weights;
+
+	// cached mirror data
+	SkinPaintTool::FSkinMirrorData MirrorData;
+
+	// vertex colors updated when switching current bone or editing weights
+	void UpdateCurrentBoneVertexColors();
+	FVector4f WeightToColor(float Value) const;
+	bool bVisibleWeightsValid = false;
+
+	// weight editing transactions
+	void BeginChange();
+	TUniquePtr<SkinPaintTool::FMeshSkinWeightsChange> EndChange();
+	TUniquePtr<SkinPaintTool::FMeshSkinWeightsChange> ActiveChange;
+	// used to apply manual weight modifications (not from the brush)
+	void ApplyWeightEditsToMesh(const FText& TransactionLabel, const SkinPaintTool::FMultiBoneWeightEdits& WeightEdits);
+
+	// which bone are we currently painting?
+	void UpdateCurrentBone(const FName &BoneName);
+	FName CurrentBone = NAME_None;
+	TOptional<FName> PendingCurrentBone;
+	TArray<FName> SelectedBoneNames;
+	TArray<BoneIndex> SelectedBoneIndices;
+
+	// determines the set of vertices to operate on, using selection as the priority
+	void GetVerticesToEdit(TArray<VertexIndex>& OutVertexIndices) const;
+	
+	BoneIndex GetBoneIndexFromName(const FName BoneName) const;
+
+	// ISkeletalMeshEditionInterface
+	virtual void HandleSkeletalMeshModified(const TArray<FName>& InBoneNames, const ESkeletalMeshNotifyType InNotifyType) override;
+
+	// polygon selection mechanic
+	UPROPERTY()
+	TObjectPtr<UPolygonSelectionMechanic> PolygonSelectionMechanic;
+	TUniquePtr<UE::Geometry::FDynamicMeshAABBTree3> MeshSpatial = nullptr;
+	TUniquePtr<UE::Geometry::FTriangleGroupTopology> SelectionTopology = nullptr;
+	void OnSelectionModified();
+
+	friend SkinPaintTool::FSkinToolDeformer;
+};
diff --git a/Engine/Plugins/Experimental/MeshModelingToolsetExp/Source/MeshModelingToolsExp/MeshModelingToolsExp.Build.cs b/Engine/Plugins/Experimental/MeshModelingToolsetExp/Source/MeshModelingToolsExp/MeshModelingToolsExp.Build.cs
--- a/Engine/Plugins/Experimental/MeshModelingToolsetExp/Source/MeshModelingToolsExp/MeshModelingToolsExp.Build.cs
+++ b/Engine/Plugins/Experimental/MeshModelingToolsetExp/Source/MeshModelingToolsExp/MeshModelingToolsExp.Build.cs
@@ -36,9 +36,7 @@ public class MeshModelingToolsExp : ModuleRules
 				"StaticMeshDescription",
 				"SkeletalMeshDescription",
 				"ModelingComponents",
-				"ModelingOperators",
-				"SkeletalMeshUtilitiesCommon",
-				"SkeletalMeshModifiers",
+				"ModelingOperators"
 				// ... add other public dependencies that you statically link with here ...
 			}
 		);
diff --git a/Engine/Plugins/Experimental/MeshModelingToolsetExp/Source/MeshModelingToolsExp/Private/SkeletalMesh/SkeletalMeshEditionInterface.cpp b/Engine/Plugins/Experimental/MeshModelingToolsetExp/Source/MeshModelingToolsExp/Private/SkeletalMesh/SkeletalMeshEditionInterface.cpp
deleted file mode 100644
diff --git a/Engine/Plugins/Experimental/MeshModelingToolsetExp/Source/MeshModelingToolsExp/Private/SkeletalMesh/SkeletonEditingTool.cpp b/Engine/Plugins/Experimental/MeshModelingToolsetExp/Source/MeshModelingToolsExp/Private/SkeletalMesh/SkeletonEditingTool.cpp
deleted file mode 100644
diff --git a/Engine/Plugins/Experimental/MeshModelingToolsetExp/Source/MeshModelingToolsExp/Private/SkeletalMesh/SkeletonTransformProxy.cpp b/Engine/Plugins/Experimental/MeshModelingToolsetExp/Source/MeshModelingToolsExp/Private/SkeletalMesh/SkeletonTransformProxy.cpp
deleted file mode 100644
diff --git a/Engine/Plugins/Experimental/MeshModelingToolsetExp/Source/MeshModelingToolsExp/Private/SkinWeightsBindingTool.cpp b/Engine/Plugins/Experimental/MeshModelingToolsetExp/Source/MeshModelingToolsExp/Private/SkinWeightsBindingTool.cpp
deleted file mode 100644
diff --git a/Engine/Plugins/Experimental/MeshModelingToolsetExp/Source/MeshModelingToolsExp/Private/SkinWeightsPaintTool.cpp b/Engine/Plugins/Experimental/MeshModelingToolsetExp/Source/MeshModelingToolsExp/Private/SkinWeightsPaintTool.cpp
deleted file mode 100644
diff --git a/Engine/Plugins/Experimental/MeshModelingToolsetExp/Source/MeshModelingToolsExp/Public/SkeletalMesh/SkeletalMeshEditionInterface.h b/Engine/Plugins/Experimental/MeshModelingToolsetExp/Source/MeshModelingToolsExp/Public/SkeletalMesh/SkeletalMeshEditionInterface.h
deleted file mode 100644
diff --git a/Engine/Plugins/Experimental/MeshModelingToolsetExp/Source/MeshModelingToolsExp/Public/SkeletalMesh/SkeletonEditingTool.h b/Engine/Plugins/Experimental/MeshModelingToolsetExp/Source/MeshModelingToolsExp/Public/SkeletalMesh/SkeletonEditingTool.h
deleted file mode 100644
diff --git a/Engine/Plugins/Experimental/MeshModelingToolsetExp/Source/MeshModelingToolsExp/Public/SkeletalMesh/SkeletonTransformProxy.h b/Engine/Plugins/Experimental/MeshModelingToolsetExp/Source/MeshModelingToolsExp/Public/SkeletalMesh/SkeletonTransformProxy.h
deleted file mode 100644
diff --git a/Engine/Plugins/Experimental/MeshModelingToolsetExp/Source/MeshModelingToolsExp/Public/SkinWeightsBindingTool.h b/Engine/Plugins/Experimental/MeshModelingToolsetExp/Source/MeshModelingToolsExp/Public/SkinWeightsBindingTool.h
deleted file mode 100644
diff --git a/Engine/Plugins/Experimental/MeshModelingToolsetExp/Source/MeshModelingToolsExp/Public/SkinWeightsPaintTool.h b/Engine/Plugins/Experimental/MeshModelingToolsetExp/Source/MeshModelingToolsExp/Public/SkinWeightsPaintTool.h
deleted file mode 100644
diff --git a/Engine/Plugins/Experimental/MeshModelingToolsetExp/Source/ModelingEditorUI/ModelingEditorUI.build.cs b/Engine/Plugins/Experimental/MeshModelingToolsetExp/Source/ModelingEditorUI/ModelingEditorUI.build.cs
--- a/Engine/Plugins/Experimental/MeshModelingToolsetExp/Source/ModelingEditorUI/ModelingEditorUI.build.cs
+++ b/Engine/Plugins/Experimental/MeshModelingToolsetExp/Source/ModelingEditorUI/ModelingEditorUI.build.cs
@@ -37,9 +37,9 @@ public class ModelingEditorUI : ModuleRules
                 "Projects",
 				"ToolWidgets",
 				"EditorWidgets",
-				"DeveloperSettings", 
-				"MeshModelingToolsExp",
-				"Persona"
+				"DeveloperSettings",
+				"Persona",
+				"MeshModelingToolsEditorOnlyExp"
 			}
 			);
 
diff --git a/Engine/Plugins/Experimental/MeshModelingToolsetExp/Source/ModelingEditorUI/Private/ModelingEditorUIModule.cpp b/Engine/Plugins/Experimental/MeshModelingToolsetExp/Source/ModelingEditorUI/Private/ModelingEditorUIModule.cpp
--- a/Engine/Plugins/Experimental/MeshModelingToolsetExp/Source/ModelingEditorUI/Private/ModelingEditorUIModule.cpp
+++ b/Engine/Plugins/Experimental/MeshModelingToolsetExp/Source/ModelingEditorUI/Private/ModelingEditorUIModule.cpp
@@ -3,7 +3,7 @@
 #include "ModelingEditorUIModule.h"
 
 #include "SkinWeightDetailCustomization.h"
-#include "SkinWeightsPaintTool.h"
+#include "SkeletalMesh/SkinWeightsPaintTool.h"
 #include "PropertyEditorModule.h"
 
 #define LOCTEXT_NAMESPACE "FModelingEditorUIModule"
diff --git a/Engine/Plugins/Experimental/MeshModelingToolsetExp/Source/ModelingEditorUI/Private/SkinWeightDetailCustomization.cpp b/Engine/Plugins/Experimental/MeshModelingToolsetExp/Source/ModelingEditorUI/Private/SkinWeightDetailCustomization.cpp
--- a/Engine/Plugins/Experimental/MeshModelingToolsetExp/Source/ModelingEditorUI/Private/SkinWeightDetailCustomization.cpp
+++ b/Engine/Plugins/Experimental/MeshModelingToolsetExp/Source/ModelingEditorUI/Private/SkinWeightDetailCustomization.cpp
@@ -6,7 +6,7 @@
 #include "DetailLayoutBuilder.h"
 #include "DetailWidgetRow.h"
 #include "Widgets/Input/SSegmentedControl.h"
-#include "SkinWeightsPaintTool.h"
+#include "SkeletalMesh/SkinWeightsPaintTool.h"
 #include "SSkinWeightProfileImportOptions.h"
 #include "UObject/UnrealTypePrivate.h"
 #include "Widgets/Input/SSpinBox.h"
