Fixed use of invalid object ptr
    + renamed Init into SetSkeletalMesh to be consistent with USkinWeightModifier

diff --git a/Engine/Plugins/Experimental/MeshModelingToolsetExp/Source/MeshModelingToolsEditorOnlyExp/Private/SkeletalMesh/SkeletonEditingTool.cpp b/Engine/Plugins/Experimental/MeshModelingToolsetExp/Source/MeshModelingToolsEditorOnlyExp/Private/SkeletalMesh/SkeletonEditingTool.cpp
--- a/Engine/Plugins/Experimental/MeshModelingToolsetExp/Source/MeshModelingToolsEditorOnlyExp/Private/SkeletalMesh/SkeletonEditingTool.cpp
+++ b/Engine/Plugins/Experimental/MeshModelingToolsetExp/Source/MeshModelingToolsEditorOnlyExp/Private/SkeletalMesh/SkeletonEditingTool.cpp
@@ -110,7 +110,7 @@ void USkeletonEditingTool::Setup()
 
 	// setup modifier
 	Modifier = NewObject<USkeletonModifier>(this);
-	Modifier->Init(SkeletalMesh);
+	Modifier->SetSkeletalMesh(SkeletalMesh);
 
 	// setup current bone
 	const FReferenceSkeleton& RefSkeleton = Modifier->GetReferenceSkeleton();
diff --git a/Engine/Plugins/Experimental/MeshModelingToolsetExp/Source/SkeletalMeshModifiers/Private/SkeletonModifier.cpp b/Engine/Plugins/Experimental/MeshModelingToolsetExp/Source/SkeletalMeshModifiers/Private/SkeletonModifier.cpp
--- a/Engine/Plugins/Experimental/MeshModelingToolsetExp/Source/SkeletalMeshModifiers/Private/SkeletonModifier.cpp
+++ b/Engine/Plugins/Experimental/MeshModelingToolsetExp/Source/SkeletalMeshModifiers/Private/SkeletonModifier.cpp
@@ -127,7 +127,7 @@ void USkeletonModifier::ExternalUpdate(const FReferenceSkeleton& InRefSkeleton,
 	BoneIndexTracker = InIndexTracker;
 }
 
-bool USkeletonModifier::Init(USkeletalMesh* InSkeletalMesh)
+bool USkeletonModifier::SetSkeletalMesh(USkeletalMesh* InSkeletalMesh)
 {
 	SkeletalMesh = nullptr;
 	MeshDescription.Reset();
@@ -234,7 +234,7 @@ bool USkeletonModifier::IsReferenceSkeletonValid(const bool bLog) const
 
 bool USkeletonModifier::CommitSkeletonToSkeletalMesh()
 {
-	if (!SkeletalMesh || !ReferenceSkeleton || !MeshDescription)
+	if (!SkeletalMesh.IsValid() || !ReferenceSkeleton || !MeshDescription)
 	{
 		UE_LOG(LogAnimation, Error, TEXT("Skeleton Modifier: No mesh loaded. Cannot apply skeleton edits."));
 		return false;
@@ -355,7 +355,7 @@ bool USkeletonModifier::CommitSkeletonToSkeletalMesh()
 	SkeletalMesh->CommitMeshDescription(USkeletonModifierLocals::LODIndex, *MeshDescription);
 
 	// update skeleton
-	if (Skeleton->RecreateBoneTree(SkeletalMesh))
+	if (Skeleton->RecreateBoneTree(SkeletalMesh.Get()))
 	{
 		Skeleton->MarkPackageDirty();	
 	}
diff --git a/Engine/Plugins/Experimental/MeshModelingToolsetExp/Source/SkeletalMeshModifiers/Public/SkeletonModifier.h b/Engine/Plugins/Experimental/MeshModelingToolsetExp/Source/SkeletalMeshModifiers/Public/SkeletonModifier.h
--- a/Engine/Plugins/Experimental/MeshModelingToolsetExp/Source/SkeletalMeshModifiers/Public/SkeletonModifier.h
+++ b/Engine/Plugins/Experimental/MeshModelingToolsetExp/Source/SkeletalMeshModifiers/Public/SkeletonModifier.h
@@ -80,7 +80,7 @@ class SKELETALMESHMODIFIERS_API USkeletonModifier : public UObject
 public:
 
 	UFUNCTION(BlueprintCallable, Category = "Mesh")
-	bool Init(USkeletalMesh* InSkeletalMesh);
+	bool SetSkeletalMesh(USkeletalMesh* InSkeletalMesh);
 
 	/** Creates a new bone in the skeleton hierarchy at desired transform
 	 *  @param InBoneName The new bone's name.
@@ -200,8 +200,10 @@ private:
 
 	// orient function
 	void GetBonesToOrient(const TArray<FName>& InBonesName, const FOrientOptions& InOptions, TArray<int32>& OutBonesToOrient) const;
+
+	UPROPERTY(transient)
+	TWeakObjectPtr<USkeletalMesh> SkeletalMesh = nullptr;
 	
-	TObjectPtr<USkeletalMesh> SkeletalMesh = nullptr;
 	TUniquePtr<FMeshDescription> MeshDescription;
 	TUniquePtr<FReferenceSkeleton> ReferenceSkeleton;
 	TUniquePtr<FTransformComposer> TransformComposer;
