Fixes issue with selection mode in skin paint tool after deformation occurs.

diff --git a/Engine/Plugins/Experimental/MeshModelingToolsetExp/Source/MeshModelingToolsEditorOnlyExp/Private/SkeletalMesh/SkinWeightsPaintTool.cpp b/Engine/Plugins/Experimental/MeshModelingToolsetExp/Source/MeshModelingToolsEditorOnlyExp/Private/SkeletalMesh/SkinWeightsPaintTool.cpp
--- a/Engine/Plugins/Experimental/MeshModelingToolsetExp/Source/MeshModelingToolsEditorOnlyExp/Private/SkeletalMesh/SkinWeightsPaintTool.cpp
+++ b/Engine/Plugins/Experimental/MeshModelingToolsetExp/Source/MeshModelingToolsEditorOnlyExp/Private/SkeletalMesh/SkinWeightsPaintTool.cpp
@@ -120,7 +120,12 @@ void FSkinToolDeformer::Initialize(const USkeletalMeshComponent* SkeletalMeshCom
 	const TArrayView<const UE::Math::TVector<float>> VertexPositions = Mesh->GetVertexPositions().GetRawArray();
 	RefPoseVertexPositions = VertexPositions;
 
-	// force all vertices to be updated initially
+	// set all vertices to be updated on first tick
+	SetAllVerticesToBeUpdated();
+}
+
+void FSkinToolDeformer::SetAllVerticesToBeUpdated()
+{
 	VerticesWithModifiedWeights.Empty(RefPoseVertexPositions.Num());
 	for (int32 VertexID=0; VertexID<RefPoseVertexPositions.Num(); ++VertexID)
 	{
@@ -162,47 +167,65 @@ void FSkinToolDeformer::UpdateVertexDeformation(USkinWeightsPaintTool* Tool)
 	}, false);
 	PreviewMesh->NotifyDeferredEditCompleted(UPreviewMesh::ERenderUpdateMode::FastUpdate, EMeshRenderAttributeFlags::Positions, false);
 
-	// update vertex acceleration structure
-	{
-		TRACE_CPUPROFILER_EVENT_SCOPE(SkinTool::UpdateVertexOctree);
-		Tool->VerticesOctree.RemoveVertices(VerticesWithModifiedWeights);
-		Tool->VerticesOctree.InsertVertices(VerticesWithModifiedWeights);
-	}
-
-	// update triangle acceleration structure
+	// what mode are we in?
+	const EWeightEditMode EditingMode = Tool->WeightToolProperties->EditingMode;
+	
+	// update data structures used by the brush mode	
+	if (EditingMode == EWeightEditMode::Brush)
 	{
-		TRACE_CPUPROFILER_EVENT_SCOPE(SkinTool::UpdateTriangleOctree);
+		// update vertex acceleration structure
+		{
+			TRACE_CPUPROFILER_EVENT_SCOPE(SkinTool::UpdateVertexOctree);
+			Tool->VerticesOctree.RemoveVertices(VerticesWithModifiedWeights);
+			Tool->VerticesOctree.InsertVertices(VerticesWithModifiedWeights);
+		}
+		
+		// update triangle acceleration structure
+		{
+			TRACE_CPUPROFILER_EVENT_SCOPE(SkinTool::UpdateTriangleOctree);
 
-		// ensure previous async update is finished before queuing the next one...
-		Tool->TriangleOctreeFuture.Wait();
+			// create list of triangles that were affected by the vertices that were deformed
+			TArray<int32>& AffectedTriangles = Tool->TrianglesToReinsert; // reusable buffer of triangles to update
+			{
+				AffectedTriangles.Reset();
 
-		// reusable buffer of triangles to reinsert into the octree
-		TArray<int32>& TrianglesToReinsert = Tool->TrianglesToReinsert;
-		TrianglesToReinsert.Reset();
+				// reinsert all triangles containing an updated vertex
+				const FDynamicMesh3* DynamicMesh = PreviewMesh->GetMesh();
+				for (const int32 TriangleID : DynamicMesh->TriangleIndicesItr())
+				{
+					UE::Geometry::FIndex3i TriVerts = DynamicMesh->GetTriangle(TriangleID);
+					bool bIsTriangleAffected = VerticesWithModifiedWeights.Contains(TriVerts[0]);
+					bIsTriangleAffected = VerticesWithModifiedWeights.Contains(TriVerts[1]) ? true : bIsTriangleAffected;
+					bIsTriangleAffected = VerticesWithModifiedWeights.Contains(TriVerts[2]) ? true : bIsTriangleAffected;
+					if (bIsTriangleAffected)
+					{
+						AffectedTriangles.Add(TriangleID);
+					}
+				}
+			}
 
-		// reinsert all triangles containing an updated vertex
-		const FDynamicMesh3* DynamicMesh = PreviewMesh->GetMesh();
-		for (const int32 TriangleID : DynamicMesh->TriangleIndicesItr())
-		{
-			UE::Geometry::FIndex3i TriVerts = DynamicMesh->GetTriangle(TriangleID);
-			bool bIsTriangleAffected = VerticesWithModifiedWeights.Contains(TriVerts[0]);
-			bIsTriangleAffected = VerticesWithModifiedWeights.Contains(TriVerts[1]) ? true : bIsTriangleAffected;
-			bIsTriangleAffected = VerticesWithModifiedWeights.Contains(TriVerts[2]) ? true : bIsTriangleAffected;
-			if (bIsTriangleAffected)
+			// ensure previous async update is finished before queuing the next one...
+			Tool->TriangleOctreeFuture.Wait();
+		
+			// asynchronously update the octree, this normally finishes well before the next update
+			// but in the unlikely event that it does not, it would result in a frame where the paint brush
+			// is not perfectly aligned with the mesh; not a deal breaker.
+			UE::Geometry::FDynamicMeshOctree3& OctreeToUpdate = Tool->TrianglesOctree;
+			Tool->TriangleOctreeFuture = Async(SkinPaintToolAsyncExecTarget, [&OctreeToUpdate, &AffectedTriangles]()
 			{
-				TrianglesToReinsert.Add(TriangleID);
-			}
+				TRACE_CPUPROFILER_EVENT_SCOPE(SkinTool::TriangleOctreeReinsert);	
+				OctreeToUpdate.ReinsertTriangles(AffectedTriangles);
+			});
 		}
+	}
 
-		// asynchronously update the octree, this normally finishes well before the next update
-		// but in the unlikely event that it does not, it would result in a frame where the paint brush
-		// is not perfectly aligned with the mesh; not a deal breaker.
-		UE::Geometry::FDynamicMeshOctree3& OctreeToUpdate = Tool->TrianglesOctree;
-		Tool->TriangleOctreeFuture = Async(SkinPaintToolAsyncExecTarget, [&OctreeToUpdate, &TrianglesToReinsert]()
-		{
-			TRACE_CPUPROFILER_EVENT_SCOPE(SkinTool::TriangleOctreeReinsert);	
-			OctreeToUpdate.ReinsertTriangles(TrianglesToReinsert);
-		});
+	// update data structures used by the selection mode
+	if (EditingMode == EWeightEditMode::Selection)
+	{
+		// update AABB Tree for vertex selection
+		TRACE_CPUPROFILER_EVENT_SCOPE(SkinTool::UpdateAABBTree);
+		Tool->MeshSpatial->Build();
+		Tool->PolygonSelectionMechanic->UMeshTopologySelectionMechanic::GetTopologySelector()->Invalidate(true, false);
 	}
 
 	// empty queue of vertices to update
@@ -1780,9 +1803,9 @@ void USkinWeightsPaintTool::OnSelectionModified()
 
 void USkinWeightsPaintTool::ToggleEditingMode()
 {
+	Weights.Deformer.SetAllVerticesToBeUpdated();
 	SetBrushEnabled(WeightToolProperties->EditingMode == EWeightEditMode::Brush);
 	PolygonSelectionMechanic->SetIsEnabled(WeightToolProperties->EditingMode == EWeightEditMode::Selection);
-	
 }
 
 void USkinWeightsPaintTool::OnPropertyModified(UObject* ModifiedObject, FProperty* ModifiedProperty)
diff --git a/Engine/Plugins/Experimental/MeshModelingToolsetExp/Source/MeshModelingToolsEditorOnlyExp/Public/SkeletalMesh/SkinWeightsPaintTool.h b/Engine/Plugins/Experimental/MeshModelingToolsetExp/Source/MeshModelingToolsEditorOnlyExp/Public/SkeletalMesh/SkinWeightsPaintTool.h
--- a/Engine/Plugins/Experimental/MeshModelingToolsetExp/Source/MeshModelingToolsEditorOnlyExp/Public/SkeletalMesh/SkinWeightsPaintTool.h
+++ b/Engine/Plugins/Experimental/MeshModelingToolsetExp/Source/MeshModelingToolsEditorOnlyExp/Public/SkeletalMesh/SkinWeightsPaintTool.h
@@ -88,6 +88,8 @@ namespace SkinPaintTool
 	{
 		void Initialize(const USkeletalMeshComponent* SkeletalMeshComponent, const FMeshDescription* Mesh);
 
+		void SetAllVerticesToBeUpdated();
+
 		void UpdateVertexDeformation(USkinWeightsPaintTool* Tool);
 
 		void SetVertexNeedsUpdated(int32 VertexIndex);
