EditPivotTool: Fix crash when multiple tool targets are backed by the same asset.

diff --git a/Engine/Plugins/Experimental/MeshModelingToolsetExp/Source/MeshModelingToolsExp/Private/EditPivotTool.cpp b/Engine/Plugins/Experimental/MeshModelingToolsetExp/Source/MeshModelingToolsExp/Private/EditPivotTool.cpp
--- a/Engine/Plugins/Experimental/MeshModelingToolsetExp/Source/MeshModelingToolsExp/Private/EditPivotTool.cpp
+++ b/Engine/Plugins/Experimental/MeshModelingToolsetExp/Source/MeshModelingToolsExp/Private/EditPivotTool.cpp
@@ -441,6 +441,7 @@ void UEditPivotTool::UpdateAssets(const FFrame3d& NewPivotWorldFrame)
 
 	// First pre-compute all the transforms we will need to apply
 	TArray<FTransformSequence3d> BakeTransforms;
+	BakeTransforms.SetNum(Targets.Num());
 	TArray<FTransform> ComponentTransforms;
 	FTransform NewWorldTransform = NewPivotWorldFrame.ToFTransform();
 	FTransform NewWorldInverse = NewPivotWorldFrame.ToInverseFTransform();
@@ -458,7 +459,6 @@ void UEditPivotTool::UpdateAssets(const FFrame3d& NewPivotWorldFrame)
 		if (InstancedComponent != nullptr)
 		{
 			// no bake for instancedcomponents, nothing needs to be precomputed here
-			BakeTransforms.Emplace();
 			ComponentTransforms.Add(FTransform::Identity);
 		}
 		else if (MapToFirstOccurrences[ComponentIdx] == ComponentIdx)
@@ -512,7 +512,7 @@ void UEditPivotTool::UpdateAssets(const FFrame3d& NewPivotWorldFrame)
 				// else do nothing -- scale is not invertible and must be baked
 			}
 
-			FTransformSequence3d& Sequence = BakeTransforms.Emplace_GetRef();
+			FTransformSequence3d& Sequence = BakeTransforms[ComponentIdx];
 			Sequence.Append(ToBake);
 			if (bNeedSeparateScale)
 			{
@@ -528,7 +528,7 @@ void UEditPivotTool::UpdateAssets(const FFrame3d& NewPivotWorldFrame)
 			FTransform ApproxBaked = OriginalTransforms[FirstOccurrence] * NewWorldInverse;
 			FTransform ComponentTransform = ApproxBaked.Inverse() * OriginalTransforms[ComponentIdx];
 			ComponentTransforms.Add(ComponentTransform);
-			BakeTransforms.Add(BakeTransforms[FirstOccurrence]);
+			BakeTransforms[ComponentIdx] = BakeTransforms[FirstOccurrence];
 		}
 	}
 
