ModelingTools: fix Fill and GroupFill modes in Vertex Paint Tool. Improve Lazy Brush Behavior (required adding an AABBTree but fixes quite broken behavior at high lazyness)

diff --git a/Engine/Plugins/Experimental/MeshModelingToolsetExp/Source/MeshModelingToolsExp/Private/MeshVertexPaintTool.cpp b/Engine/Plugins/Experimental/MeshModelingToolsetExp/Source/MeshModelingToolsExp/Private/MeshVertexPaintTool.cpp
--- a/Engine/Plugins/Experimental/MeshModelingToolsetExp/Source/MeshModelingToolsExp/Private/MeshVertexPaintTool.cpp
+++ b/Engine/Plugins/Experimental/MeshModelingToolsetExp/Source/MeshModelingToolsExp/Private/MeshVertexPaintTool.cpp
@@ -202,11 +202,16 @@ void UMeshVertexPaintTool::Setup()
 		Octree.Initialize(Mesh);
 	});
 
+	TFuture<void> AABBTreeFuture = Async(VertexPaintToolAsyncExecTarget, [&]()
+	{
+		AABBTree.SetMesh(Mesh, true);
+	});
+
+
 	// initialize render decomposition
 	TUniquePtr<FMeshRenderDecomposition> Decomp = MakeUnique<FMeshRenderDecomposition>();
 	FMeshRenderDecomposition::BuildChunkedDecomposition(Mesh, &MaterialSet, *Decomp);
 	Decomp->BuildAssociations(Mesh);
-	//UE_LOG(LogTemp, Warning, TEXT("Decomposition has %d groups"), Decomp->Num());
 	DynamicMeshComponent->SetExternalDecomposition(MoveTemp(Decomp));
 
 	// initialize brush radius range interval, brush properties
@@ -368,6 +373,7 @@ void UMeshVertexPaintTool::Setup()
 
 	PrecomputeFuture.Wait();
 	OctreeFuture.Wait();
+	AABBTreeFuture.Wait();
 }
 
 void UMeshVertexPaintTool::Shutdown(EToolShutdownType ShutdownType)
@@ -501,11 +507,13 @@ void UMeshVertexPaintTool::OnBeginStroke(const FRay& WorldRay)
 	{
 		PaintBrushOpProperties->Color = BasicProperties->PaintColor;
 		PaintBrushOpProperties->BlendMode = (EVertexColorPaintBrushOpBlendMode)(int32)BasicProperties->BlendMode;
+		PaintBrushOpProperties->bApplyFalloff = (BasicProperties->SubToolType == EMeshVertexPaintInteractionType::Brush);
 	}
 	if (EraseBrushOpProperties)
 	{
 		EraseBrushOpProperties->Color = BasicProperties->EraseColor;
 		EraseBrushOpProperties->BlendMode = EVertexColorPaintBrushOpBlendMode::Lerp;
+		EraseBrushOpProperties->bApplyFalloff = (BasicProperties->SubToolType == EMeshVertexPaintInteractionType::Brush);
 	}
 
 	TUniquePtr<FMeshSculptBrushOp>& UseBrushOp = GetActiveBrushOp();
@@ -582,7 +590,17 @@ void UMeshVertexPaintTool::UpdateROI(const FSculptBrushStamp& BrushStamp)
 
 	TriangleROI.Reset();
 
+	// With Lazy Brush, GetBrushTriangleID() may not necessarily return a triangle that contains BrushPos,
+	// so in that case we will try to find the triangle containing BrushPos. Note that this may not
+	// be correct in some geometric cases...potentially we could check if the initial CenterTID 
+	// contains the BrushPos point first
 	int32 CenterTID = GetBrushTriangleID();
+	if (BrushProperties->Lazyness > 0)
+	{
+		double NearDistSqr;
+		CenterTID = AABBTree.FindNearestTriangle(BrushPos, NearDistSqr);
+	}
+
 	if (Mesh->IsTriangle(CenterTID))
 	{
 		TriangleROI.Add(CenterTID);
@@ -770,7 +788,6 @@ bool UMeshVertexPaintTool::ApplyStamp()
 
 	LastStamp = CurrentStamp;
 	LastStamp.TimeStamp = FDateTime::Now();
-
 	return bUpdated;
 }
 
diff --git a/Engine/Plugins/Experimental/MeshModelingToolsetExp/Source/MeshModelingToolsExp/Public/MeshVertexPaintTool.h b/Engine/Plugins/Experimental/MeshModelingToolsetExp/Source/MeshModelingToolsExp/Public/MeshVertexPaintTool.h
--- a/Engine/Plugins/Experimental/MeshModelingToolsetExp/Source/MeshModelingToolsExp/Public/MeshVertexPaintTool.h
+++ b/Engine/Plugins/Experimental/MeshModelingToolsetExp/Source/MeshModelingToolsExp/Public/MeshVertexPaintTool.h
@@ -633,7 +633,11 @@ protected:
 	bool UpdateStampPosition(const FRay& WorldRay);
 	bool ApplyStamp();
 
+	// initial code here was ported from MeshVertexSculptTool, which requires an Octree. However since
+	// mesh shape is static, we can actually use an AABBTree, and in one case a required query (nearest-point)
+	// is not supported by Octree (currently). So currently using both (gross)
 	UE::Geometry::FDynamicMeshOctree3 Octree;
+	UE::Geometry::FDynamicMeshAABBTree3 AABBTree;
 
 	bool UpdateBrushPosition(const FRay& WorldRay);
 
diff --git a/Engine/Plugins/Experimental/MeshModelingToolsetExp/Source/MeshModelingToolsExp/Public/Sculpting/MeshVertexPaintBrushOps.h b/Engine/Plugins/Experimental/MeshModelingToolsetExp/Source/MeshModelingToolsExp/Public/Sculpting/MeshVertexPaintBrushOps.h
--- a/Engine/Plugins/Experimental/MeshModelingToolsetExp/Source/MeshModelingToolsExp/Public/Sculpting/MeshVertexPaintBrushOps.h
+++ b/Engine/Plugins/Experimental/MeshModelingToolsetExp/Source/MeshModelingToolsExp/Public/Sculpting/MeshVertexPaintBrushOps.h
@@ -133,6 +133,11 @@ public:
 	UPROPERTY(EditAnywhere, Category = PaintBrush)
 	EVertexColorPaintBrushOpBlendMode BlendMode = EVertexColorPaintBrushOpBlendMode::Mix;
 
+	/** If bApplyFalloff is disabled, 1.0 is used as "falloff" for all vertices */
+	UPROPERTY(EditAnywhere, Category = PaintBrush)
+	bool bApplyFalloff = true;
+
+
 	virtual float GetStrength() override { return Strength; }
 	virtual void SetStrength(float NewStrength) override { Strength = FMathf::Clamp(NewStrength, 0.0f, 1.0f); }
 	virtual float GetFalloff() override { return Falloff; }
@@ -233,6 +238,7 @@ public:
 		{
 			SetColor.W = 1.0f;
 		}
+		bool bApplyFalloff = Props->bApplyFalloff;
 
 		int32 NumElements = ElementIDs.Num();
 		ensure(NewColorsOut.Num() == NumElements);
@@ -241,7 +247,7 @@ public:
 			int32 ElementID = ElementIDs[k];
 			int32 VertexID = VertexColors->GetParentVertex(ElementID);
 			FVector3d VertexPos = Mesh->GetVertex(VertexID);
-			double PositionFalloff = GetFalloff().Evaluate(Stamp, VertexPos);
+			double PositionFalloff = (bApplyFalloff) ? GetFalloff().Evaluate(Stamp, VertexPos) : 1.0;
 			double T = PositionFalloff * UsePower;
 
 			if (Props->BlendMode == EVertexColorPaintBrushOpBlendMode::Mix)
