Skeleton Editing Tool: fixed projection in orthographic views

diff --git a/Engine/Plugins/Experimental/MeshModelingToolsetExp/Source/MeshModelingToolsEditorOnlyExp/Private/SkeletalMesh/SkeletonEditingTool.cpp b/Engine/Plugins/Experimental/MeshModelingToolsetExp/Source/MeshModelingToolsEditorOnlyExp/Private/SkeletalMesh/SkeletonEditingTool.cpp
--- a/Engine/Plugins/Experimental/MeshModelingToolsetExp/Source/MeshModelingToolsEditorOnlyExp/Private/SkeletalMesh/SkeletonEditingTool.cpp
+++ b/Engine/Plugins/Experimental/MeshModelingToolsetExp/Source/MeshModelingToolsEditorOnlyExp/Private/SkeletalMesh/SkeletonEditingTool.cpp
@@ -847,6 +847,7 @@ FInputRayHit USkeletonEditingTool::CanBeginClickDragSequence(const FInputDeviceR
 	// update projection properties
 	const FVector GlobalPosition = Modifier->GetTransform(BoneIndex, true).GetTranslation();
 	ProjectionProperties->UpdatePlane(*ViewContext, GlobalPosition);
+	GetToolManager()->GetContextQueriesAPI()->GetCurrentViewState(ProjectionProperties->CameraState);
 	
 	// if we picked a new bone
 	if (BoneIndex > INDEX_NONE)
@@ -1327,8 +1328,15 @@ bool UProjectionProperties::GetProjectionPoint(const FInputDeviceRay& InRay, FVe
 	}
 
 	const FPlane Plane(PlaneOrigin, PlaneNormal);
-	const float HitDepth = FMath::RayPlaneIntersectionParam(WorldRay.Origin, WorldRay.Direction, Plane);
-	if (HitDepth < 0)
+	
+	if (CameraState.bIsOrthographic)
+	{
+		OutHitPoint = FVector::PointPlaneProject(WorldRay.Origin, Plane);
+		return true;
+	}
+	
+	const double HitDepth = FMath::RayPlaneIntersectionParam(WorldRay.Origin, WorldRay.Direction, Plane);
+	if (HitDepth < 0.0)
 	{
 		return false;
 	}
diff --git a/Engine/Plugins/Experimental/MeshModelingToolsetExp/Source/MeshModelingToolsEditorOnlyExp/Public/SkeletalMesh/SkeletonEditingTool.h b/Engine/Plugins/Experimental/MeshModelingToolsetExp/Source/MeshModelingToolsEditorOnlyExp/Public/SkeletalMesh/SkeletonEditingTool.h
--- a/Engine/Plugins/Experimental/MeshModelingToolsetExp/Source/MeshModelingToolsEditorOnlyExp/Public/SkeletalMesh/SkeletonEditingTool.h
+++ b/Engine/Plugins/Experimental/MeshModelingToolsetExp/Source/MeshModelingToolsEditorOnlyExp/Public/SkeletalMesh/SkeletonEditingTool.h
@@ -301,6 +301,8 @@ public:
 	
 	UPROPERTY(EditAnywhere, Category = "Project")
 	EProjectionType ProjectionType = EProjectionType::WithinMesh;
+
+	FViewCameraState CameraState;
 	
 private:
 	TWeakObjectPtr<UPreviewMesh> PreviewMesh = nullptr;
