Skeleton Editing Tool: use the plane defined by the bones to compute the default secondary axis when orienting bones

diff --git a/Engine/Plugins/Experimental/MeshModelingToolsetExp/Source/SkeletalMeshModifiers/Private/SkeletonModifier.cpp b/Engine/Plugins/Experimental/MeshModelingToolsetExp/Source/SkeletalMeshModifiers/Private/SkeletonModifier.cpp
--- a/Engine/Plugins/Experimental/MeshModelingToolsetExp/Source/SkeletalMeshModifiers/Private/SkeletonModifier.cpp
+++ b/Engine/Plugins/Experimental/MeshModelingToolsetExp/Source/SkeletalMeshModifiers/Private/SkeletonModifier.cpp
@@ -838,7 +838,9 @@ bool USkeletonModifier::SetBonesTransforms(
 		FReferenceSkeletonModifier Modifier(*ReferenceSkeleton, nullptr);
 		for (int32 Index = 0; Index < BoneIndices.Num(); Index++)
 		{
-			Modifier.UpdateRefPoseTransform(BoneIndices[Index], InNewTransforms[Offsets[Index]]);
+			FTransform NewTransform = InNewTransforms[Offsets[Index]];
+			NewTransform.NormalizeRotation();
+			Modifier.UpdateRefPoseTransform(BoneIndices[Index], NewTransform);
 
 			// invalidate cached global transform
 			TransformComposer->Invalidate(BoneIndices[Index]);
@@ -852,7 +854,8 @@ bool USkeletonModifier::SetBonesTransforms(
 				const int32 ChildrenIndex = ChildrenToFix[Index];
 				const int32 ParentIndex = ReferenceSkeleton->GetRawParentIndex(ChildrenIndex);
 				const FTransform& NewParentGlobal = TransformComposer->GetGlobalTransform(ParentIndex);
-				const FTransform& NewLocal = GlobalTransforms[Index].GetRelativeTransform(NewParentGlobal);
+				FTransform NewLocal = GlobalTransforms[Index].GetRelativeTransform(NewParentGlobal);
+				NewLocal.NormalizeRotation();
 				Modifier.UpdateRefPoseTransform(ChildrenIndex, NewLocal);
 				TransformComposer->Invalidate(ChildrenIndex);
 			}
@@ -1120,6 +1123,22 @@ bool USkeletonModifier::OrientBones(const TArray<FName>& InBoneNames, const FOri
 			return BoneGlobal;
 		}
 
+		// compute the secondary target based on the plane formed by the bones if needed
+		if (InOptions.bUsePlaneAsSecondary)
+		{
+			const FVector SecondaryDirection = (BoneGlobal.GetLocation() - ParentGlobal.GetLocation()).GetSafeNormal();
+			if (!SecondaryDirection.IsNearlyZero())
+			{
+				const bool bComputePlane = (FMath::Abs(FVector::DotProduct(Direction, SecondaryDirection)) - 1.0f) < KINDA_SMALL_NUMBER;
+				if (bComputePlane)
+				{ // use the plane normal as the secondary target, otherwise use InOptions' SecondaryTarget
+					FOrientOptions Options = InOptions;
+					Options.SecondaryTarget = FVector::CrossProduct(Direction, SecondaryDirection);
+					return Options.OrientTransform(Direction, BoneGlobal);
+				}
+			}
+		}
+
 		return InOptions.OrientTransform(Direction, BoneGlobal);
 	};
 
diff --git a/Engine/Plugins/Experimental/MeshModelingToolsetExp/Source/SkeletalMeshModifiers/Public/SkeletonModifier.h b/Engine/Plugins/Experimental/MeshModelingToolsetExp/Source/SkeletalMeshModifiers/Public/SkeletonModifier.h
--- a/Engine/Plugins/Experimental/MeshModelingToolsetExp/Source/SkeletalMeshModifiers/Public/SkeletonModifier.h
+++ b/Engine/Plugins/Experimental/MeshModelingToolsetExp/Source/SkeletalMeshModifiers/Public/SkeletonModifier.h
@@ -59,6 +59,9 @@ struct SKELETALMESHMODIFIERS_API FOrientOptions
 	UPROPERTY(BlueprintReadWrite, EditAnywhere, Category = "Orient")
 	TEnumAsByte<EAxis::Type> Secondary = EAxis::Y;
 
+	UPROPERTY(BlueprintReadWrite, EditAnywhere, Category = "Orient")
+	bool bUsePlaneAsSecondary = true;
+	
 	UPROPERTY(BlueprintReadWrite, EditAnywhere, Category = "Orient")
 	FVector SecondaryTarget = FVector::YAxisVector;
 
