CombineMeshInstances: Add IntersectSweptPlanarProjections mode, intersects 3 polygon-sweeps, which can produce good approximations of L-shape parts w/ holes. Added bAutoGenerateMissingUVs and bAutoGenerateTangents flags to FOptions, supported in impl. Updated settings on internal usage.

diff --git a/Engine/Plugins/Experimental/MeshModelingToolsetExp/Source/GeometryProcessingAdapters/Private/CombineMeshInstancesImpl.cpp b/Engine/Plugins/Experimental/MeshModelingToolsetExp/Source/GeometryProcessingAdapters/Private/CombineMeshInstancesImpl.cpp
--- a/Engine/Plugins/Experimental/MeshModelingToolsetExp/Source/GeometryProcessingAdapters/Private/CombineMeshInstancesImpl.cpp
+++ b/Engine/Plugins/Experimental/MeshModelingToolsetExp/Source/GeometryProcessingAdapters/Private/CombineMeshInstancesImpl.cpp
@@ -19,6 +19,7 @@
 #include "Parameterization/DynamicMeshUVEditor.h"
 #include "DynamicMesh/DynamicMeshAttributeSet.h"
 #include "DynamicMesh/MeshNormals.h"
+#include "DynamicMesh/MeshTangents.h"
 #include "DynamicMesh/MeshTransforms.h"
 #include "DynamicMesh/DynamicMeshAABBTree3.h"
 
@@ -37,6 +38,7 @@
 #include "DynamicMesh/Operations/MergeCoincidentMeshEdges.h"
 #include "Operations/RemoveOccludedTriangles.h"
 #include "Operations/MeshResolveTJunctions.h"
+#include "Operations/MeshBoolean.h"
 
 #include "MeshBoundaryLoops.h"
 #include "Curve/PlanarComplex.h"
@@ -1044,6 +1046,12 @@ static void ComputeSweptSolidApproximation(
 		return false;
 	});
 
+	if (Triangulator.Triangles.Num() == 0)
+	{
+		ComputeSimplePartApproximation(SourcePartMesh, DestMesh, EApproximatePartMethod::OrientedBox);
+		return;
+	}
+
 	FFlatTriangulationMeshGenerator TriangulationMeshGen;
 	TriangulationMeshGen.Vertices2D = Triangulator.Vertices;
 	TriangulationMeshGen.Triangles2D = Triangulator.Triangles;
@@ -1887,7 +1895,8 @@ static void RetriangulatePlanarFacePolygons(FDynamicMesh3& TargetMesh)
 	}
 
 	// if we have vertex colors, transfer them to new meshes by finding a value at nearest vertex
-	FDynamicMeshColorOverlay* SourceColors = TargetMesh.Attributes()->PrimaryColors();
+	bool bTargetHasAttributes = TargetMesh.HasAttributes();
+	FDynamicMeshColorOverlay* SourceColors = bTargetHasAttributes ? TargetMesh.Attributes()->PrimaryColors() : nullptr;
 	FDynamicMeshAABBTree3 TargetMeshSpatial(&TargetMesh, (SourceColors != nullptr));
 
 	// TODO: we do not actually have to split here, we can just send the triangle ROI to ComputePlanarPolygonApproximation and
@@ -1910,6 +1919,7 @@ static void RetriangulatePlanarFacePolygons(FDynamicMesh3& TargetMesh)
 		if (NewPlanarMesh.TriangleCount() < Mesh.TriangleCount())
 		{
 			Mesh = MoveTemp(NewPlanarMesh);
+			if (!bTargetHasAttributes) continue;
 
 			Mesh.EnableAttributes();
 			Mesh.Attributes()->SetNumUVLayers(0);
@@ -1940,7 +1950,10 @@ static void RetriangulatePlanarFacePolygons(FDynamicMesh3& TargetMesh)
 	}
 
 	FDynamicMesh3 NewMesh;
-	NewMesh.EnableMatchingAttributes(TargetMesh);
+	if (bTargetHasAttributes)
+	{
+		NewMesh.EnableMatchingAttributes(TargetMesh);
+	}
 
 	FDynamicMeshEditor Editor(&NewMesh);
 	for (FDynamicMesh3& Mesh : SplitMeshes)
@@ -1956,14 +1969,17 @@ static void RetriangulatePlanarFacePolygons(FDynamicMesh3& TargetMesh)
 
 	// currently assuming input mesh has been split by MaterialID, so no projection here,
 	// just use any MaterialID
-	if (const FDynamicMeshMaterialAttribute* SourceMaterialIDs = TargetMesh.Attributes()->GetMaterialID())
+	if (bTargetHasAttributes)
 	{
-		if (FDynamicMeshMaterialAttribute* TargetMaterialIDs = NewMesh.Attributes()->GetMaterialID())
+		if (const FDynamicMeshMaterialAttribute* SourceMaterialIDs = TargetMesh.Attributes()->GetMaterialID())
 		{
-			int32 ConstantMaterialID = SourceMaterialIDs->GetValue(0);
-			for (int32 tid : NewMesh.TriangleIndicesItr())
+			if (FDynamicMeshMaterialAttribute* TargetMaterialIDs = NewMesh.Attributes()->GetMaterialID())
 			{
-				TargetMaterialIDs->SetValue(tid, ConstantMaterialID);
+				int32 ConstantMaterialID = SourceMaterialIDs->GetValue(0);
+				for (int32 tid : NewMesh.TriangleIndicesItr())
+				{
+					TargetMaterialIDs->SetValue(tid, ConstantMaterialID);
+				}
 			}
 		}
 	}
@@ -2011,8 +2027,9 @@ static void PostProcessHiddenFaceRemovedMesh(
 		return;
 	}
 
+	bool bTargetHasAttributes = TargetMesh.HasAttributes();
 	const FDynamicMeshMaterialAttribute* MaterialIDs =
-		(TargetMesh.HasAttributes() && TargetMesh.Attributes()->HasMaterialID()) ? TargetMesh.Attributes()->GetMaterialID() : nullptr;
+		(bTargetHasAttributes && TargetMesh.Attributes()->HasMaterialID()) ? TargetMesh.Attributes()->GetMaterialID() : nullptr;
 	
 	TMap<FMergeTriInfo, int> UniqueMatIndices;
 	TArray<int32> TriSortIndex;
@@ -2063,7 +2080,10 @@ static void PostProcessHiddenFaceRemovedMesh(
 		// todo: could discard UVs here to improve merging. Possibly also do planar remesh, it 
 		// seems like the current method still will not remove boundary vertices...
 
-		SubRegionMesh.Attributes()->SplitAllBowties();
+		if (SubRegionMesh.HasAttributes())
+		{
+			SubRegionMesh.Attributes()->SplitAllBowties();
+		}
 		SubRegionMesh.CompactInPlace();
 
 		// simplify to planar
@@ -2091,7 +2111,10 @@ static void PostProcessHiddenFaceRemovedMesh(
 	}
 
 	TargetMesh.Clear();
-	TargetMesh.EnableMatchingAttributes(SplitMeshes[0], true, true);
+	if (bTargetHasAttributes)
+	{
+		TargetMesh.EnableMatchingAttributes(SplitMeshes[0], true, true);
+	}
 	FDynamicMeshEditor Editor(&TargetMesh);
 	for (FDynamicMesh3& SubRegionMesh : SplitMeshes)
 	{
@@ -2116,12 +2139,11 @@ static void PostProcessHiddenFaceRemovedMesh(
 
 
 	// make sure we have necessary attribute sets
-	if (TargetMesh.Attributes()->NumUVLayers() == 0)
+	if (bTargetHasAttributes && TargetMesh.Attributes()->NumUVLayers() == 0)
 	{
 		TargetMesh.Attributes()->SetNumUVLayers(1);
 	}
 
-
 	TargetMesh.CompactInPlace();
 
 	if ( bVerbose )
@@ -2259,6 +2281,89 @@ static void ComputeBestFullProjectionMesh(
 
 
 
+
+/**
+ * Computes best cardinal-axis swept-solid approximation to CombinedMesh and returns in ResultMesh.
+ * The swept-solid approximation is found by flattening ResultMesh along that axis, and then doing
+ * polygon booleans and topological closure. The various simplification parameters are derived from
+ * the ClosureDistance parameter.
+ */
+static void ComputeProjectionMeshIntersection(
+	const FDynamicMesh3& CombinedMesh,
+	FDynamicMeshAABBTree3& CombinedMeshSpatial,
+	FDynamicMesh3& ResultMesh,
+	double ClosureDistance)
+{
+	TArray<FVector3d, TInlineAllocator<3>> Directions;
+	Directions.Add(FVector3d::UnitZ());
+	Directions.Add(FVector3d::UnitX());
+	Directions.Add(FVector3d::UnitY());
+	int32 N = Directions.Num();
+
+	TArray<FDynamicMesh3, TInlineAllocator<3>> DirectionMeshes;
+	DirectionMeshes.SetNum(N);
+
+	TArray<FVector2d> DeviationMeasures;
+	DeviationMeasures.SetNum(N);
+
+	ParallelFor(Directions.Num(), [&](int32 k)
+	{
+		FDynamicMesh3& UseMesh = DirectionMeshes[k];
+		FVector3d UseDirection = Directions[k];
+		ComputeSweptSolidApproximation(CombinedMesh, UseMesh, UseDirection,
+			ClosureDistance, ClosureDistance / 4, 4.0f * ClosureDistance * ClosureDistance);
+		UseMesh.DiscardAttributes();
+
+		// simplify to planar
+		FQEMSimplification Simplifier(&UseMesh);
+
+		Simplifier.CollapseMode = FQEMSimplification::ESimplificationCollapseModes::MinimalExistingVertexError;
+		// no constraints as we discarded attributes
+		Simplifier.SimplifyToMinimalPlanar(ClosureDistance/2);
+
+		DeviationMeasures[k] = ComputeGeometricDeviation(UseMesh, CombinedMeshSpatial);
+	});
+
+	// intersect
+	FMeshBoolean MeshBoolean1(
+		&DirectionMeshes[0], FTransformSRT3d::Identity(),
+		&DirectionMeshes[1], FTransformSRT3d::Identity(),
+		&ResultMesh, FMeshBoolean::EBooleanOp::Intersect);
+	MeshBoolean1.bPutResultInInputSpace = true;
+	MeshBoolean1.bSimplifyAlongNewEdges = true;
+	MeshBoolean1.Compute();
+	
+	FDynamicMesh3 TmpMesh;
+	FMeshBoolean MeshBoolean2(
+		&ResultMesh, FTransformSRT3d::Identity(),
+		&DirectionMeshes[2], FTransformSRT3d::Identity(),
+		&TmpMesh, FMeshBoolean::EBooleanOp::Intersect);
+	MeshBoolean2.bPutResultInInputSpace = true;
+	MeshBoolean2.bSimplifyAlongNewEdges = true;
+	MeshBoolean2.Compute();
+
+	ResultMesh = MoveTemp(TmpMesh);
+
+	//RetriangulatePlanarFacePolygons(ResultMesh);
+
+	//FMergeCoincidentMeshEdges Welder(&ResultMesh);
+	//Welder.MergeVertexTolerance = ClosureDistance/5;
+	//Welder.OnlyUniquePairs = false;
+	//Welder.bWeldAttrsOnMergedEdges = true;
+	//Welder.Apply();
+
+	//FQEMSimplification Simplifier2(&ResultMesh);
+	//Simplifier2.CollapseMode = FQEMSimplification::ESimplificationCollapseModes::MinimalExistingVertexError;
+	//Simplifier2.SimplifyToEdgeLength(ClosureDistance);
+
+	PostProcessHiddenFaceRemovedMesh(ResultMesh, 1.0, true, true,
+		[](const FDynamicMesh3& Mesh, int32 TriangleID) { return FIndex3i::Zero(); });
+
+	ResultMesh.CompactInPlace();
+}
+
+
+
 template<typename SimplificationType>
 void DoSimplifyMesh(
 	FDynamicMesh3& EditMesh,
@@ -2396,6 +2501,39 @@ static void ComputeVoxWrapMeshAutoUV(FDynamicMesh3& EditMesh)
 
 
 
+
+
+
+
+
+static void ComputeMissingUVs(FDynamicMesh3& EditMesh)
+{
+	check(EditMesh.IsCompact());
+	check(EditMesh.HasAttributes());
+
+	FDynamicMeshUVEditor UVEditor(&EditMesh, 0, true);
+	FDynamicMeshUVOverlay* UVOverlay = UVEditor.GetOverlay();
+
+	TArray<int32> UnsetTriangles;
+	for (int32 tid : EditMesh.TriangleIndicesItr())
+	{
+		if (UVOverlay->IsSetTriangle(tid) == false)
+		{
+			UnsetTriangles.Add(tid);
+		}
+	}
+
+	FAxisAlignedBox3d Bounds = EditMesh.GetBounds();
+	UVEditor.SetTriangleUVsFromBoxProjection(UnsetTriangles,
+		[](const FVector3d& V) { return V; },
+		FFrame3d(Bounds.Center()), Bounds.Diagonal(), 1);
+	UVEditor.ScaleUVAreaToBoundingBox(UnsetTriangles,
+		FAxisAlignedBox2f(FVector2f::Zero(), FVector2f::One()), true, true);
+}
+
+
+
+
 static void InitializeNormalsFromAngleThreshold(
 	FDynamicMesh3& TargetMesh,
 	double NormalAngleThreshDeg)
@@ -2845,6 +2983,11 @@ void BuildCombinedMesh(
 				ComputeBestFullProjectionMesh(SourceVoxWrapMesh, SourceSpatial, InitialCoarseApproximation, CombineOptions.CoarseApproximationDetailSize);
 				bUsingCoarseSweepApproximation = true;
 			}
+			else if (CombineOptions.CoarseLODStrategy == IGeometryProcessing_CombineMeshInstances::ECoarseApproximationStrategy::IntersectSweptPlanarProjections)
+			{
+				ComputeProjectionMeshIntersection(SourceVoxWrapMesh, SourceSpatial, InitialCoarseApproximation, CombineOptions.CoarseApproximationDetailSize);
+				bUsingCoarseSweepApproximation = true;
+			}
 			else  // Automatic
 			{
 				// try swept-planar-projection as it is cheaper and generally better. If it deviates too much, fall back to voxel
@@ -2954,6 +3097,25 @@ void BuildCombinedMesh(
 		check(false);		// for future use
 	}
 
+	// parallel regenerate UVs and potentially tangents for any areas of LODs that are missing UVs
+	TArray<UE::Tasks::FTask> PendingAutoUVTasks;
+	bool bComputeTangents = (CombineOptions.bAutoGenerateMissingUVs && CombineOptions.bAutoGenerateTangents);
+	if (CombineOptions.bAutoGenerateMissingUVs)
+	{
+		for (int32 LODIndex = 0; LODIndex < NumLODs && LODIndex < FirstVoxWrappedIndex; ++LODIndex)
+		{
+			UE::Tasks::FTask AutoUVTask = UE::Tasks::Launch(UE_SOURCE_LOCATION, [&MeshLODs, &CombineOptions, LODIndex, bComputeTangents]()
+			{
+				ComputeMissingUVs(MeshLODs[LODIndex].Mesh);
+				if (bComputeTangents)
+				{
+					FMeshTangentsd::ComputeDefaultOverlayTangents(MeshLODs[LODIndex].Mesh);
+				}
+			});
+			PendingAutoUVTasks.Add(AutoUVTask);
+		}
+	}
+
 	// remove hidden faces on voxel LODs (todo: can do this via shape sorting, much faster)
 	if (bRemoveHiddenFaces && bUsingCoarseSweepApproximation == false)
 	{
@@ -2974,6 +3136,9 @@ void BuildCombinedMesh(
 	}
 
 
+	// make sure AutoUV is done
+	UE::Tasks::Wait(PendingAutoUVTasks);
+
 
 	// can't replace voxel LODs if we are generating UVs for them!
 	// (no way to communicate this upwards...)
diff --git a/Engine/Source/Developer/GeometryProcessingInterfaces/Public/GeometryProcessingInterfaces/CombineMeshInstances.h b/Engine/Source/Developer/GeometryProcessingInterfaces/Public/GeometryProcessingInterfaces/CombineMeshInstances.h
--- a/Engine/Source/Developer/GeometryProcessingInterfaces/Public/GeometryProcessingInterfaces/CombineMeshInstances.h
+++ b/Engine/Source/Developer/GeometryProcessingInterfaces/Public/GeometryProcessingInterfaces/CombineMeshInstances.h
@@ -93,7 +93,8 @@ public:
 	{
 		Automatic = 0,
 		VoxelBasedSolidApproximation = 1,
-		SweptPlanarProjection = 2
+		SweptPlanarProjection = 2,
+		IntersectSweptPlanarProjections = 3
 	};
 
 	enum class EVertexColorMappingMode
@@ -216,6 +217,14 @@ public:
 		double PartLODPromotionBudgetMultiplier = 2.0;
 
 
+
+		//
+		// 
+		//
+		bool bAutoGenerateMissingUVs = false;
+		bool bAutoGenerateTangents = false;
+
+
 		//
 		// Debug/utility options
 		// 
