ModelingToolset: Make Pattern Tool emit correct actors using a new CreateNewActor function.

Added a CreateNewActor function to EditorModelingObjectsCreationAPI which finds a factory capable of producing actors of the same type as a provided template actor to create the new actor without needing to create entirely new assets as would have been necessary with the existing functions in EditorModelingObjectsCreationAPI.

diff --git a/Engine/Plugins/Experimental/MeshModelingToolsetExp/Source/MeshModelingToolsExp/Private/PatternTool.cpp b/Engine/Plugins/Experimental/MeshModelingToolsetExp/Source/MeshModelingToolsExp/Private/PatternTool.cpp
--- a/Engine/Plugins/Experimental/MeshModelingToolsetExp/Source/MeshModelingToolsExp/Private/PatternTool.cpp
+++ b/Engine/Plugins/Experimental/MeshModelingToolsetExp/Source/MeshModelingToolsExp/Private/PatternTool.cpp
@@ -1830,20 +1830,43 @@ void UPatternTool::EmitResults()
 					FTransformSRT3d PatternTransform = CurrentPattern[k];
 					FTransform WorldTransform;
 					ComputeWorldTransform(WorldTransform, (FTransform)ElementTransform, (FTransform)PatternTransform);
-					
-					FActorSpawnParameters SpawnInfo;
-					AActor* NewActor = GetTargetWorld()->SpawnActor<AActor>(SpawnInfo);
-					if (NewActor != nullptr)
+
+					FCreateActorParams SpawnInfo;
+					SpawnInfo.BaseName = FString::Printf(TEXT("Pattern_%d_%d"), ElemIdx, k);
+					SpawnInfo.TargetWorld = GetTargetWorld();
+					SpawnInfo.TemplateActor = SourceComponent->GetOwner();
+					SpawnInfo.Transform = WorldTransform;
+
+					FCreateActorResult Result = UE::Modeling::CreateNewActor(GetToolManager(), MoveTemp(SpawnInfo));
+					if (Result.IsOK())
 					{
-						NewActors.Add(NewActor);
+						AActor* NewActor = Result.NewActor;
+						UStaticMeshComponent* NewComponent = NewActor->GetComponentByClass<UStaticMeshComponent>();
+
+						// Not sure this will ever happen but we probably don't want to proceed if it does
+						ensureMsgf(NewActor->GetRootComponent(), TEXT("New actor has no root component."));
+
+						if (NewComponent)
+						{
+							NewComponent->SetStaticMesh(SourceComponent->GetStaticMesh());
+							for (int32 j = 0; j < Element.SourceMaterials.Num(); ++j)
+							{
+								NewComponent->SetMaterial(j, Element.SourceMaterials[j]);
+							}
+						}
+						else
+						{
+							NewComponent = DuplicateObject<UStaticMeshComponent>(SourceComponent, NewActor);
+							NewActor->SetRootComponent(NewComponent);
+							NewComponent->OnComponentCreated();
+							NewActor->AddInstanceComponent(NewComponent);
+							NewComponent->RegisterComponent();
+						}
+
+						NewComponent->SetWorldTransform(WorldTransform);
 
-						UStaticMeshComponent* TemplateComponent = DuplicateObject<UStaticMeshComponent>(SourceComponent, NewActor);
-						TemplateComponent->ClearFlags(RF_DefaultSubObject);
-						TemplateComponent->OnComponentCreated();
-						NewActor->SetRootComponent(TemplateComponent);
-						TemplateComponent->RegisterComponent();
-						TemplateComponent->SetWorldTransform(WorldTransform);
-						TemplateComponent->SetVisibility(false, bPropagateToChildren);
+						NewActors.Add(NewActor);
+						SetActorComponentsVisibility(NewActor, false);
 					}
 				}
 			}
@@ -1887,41 +1910,74 @@ void UPatternTool::EmitResults()
 			}
 			else
 			{
-				// Emit a single actor with multiple StaticMeshComponents
-				FActorSpawnParameters SpawnInfo;
-				AActor* ParentActor = GetTargetWorld()->SpawnActor<AActor>(SpawnInfo);
-				if (ParentActor != nullptr)
+				AActor* NewActor = nullptr;
+				UStaticMeshComponent* TemplateComponent = nullptr;
+
+				for (int32 k = 0; k < NumPatternItems; ++k)
 				{
-					NewActors.Add(ParentActor);
-					
-					UStaticMeshComponent* TemplateComponent = nullptr;
-					for (int32 k = 0; k < NumPatternItems; ++k)
+					FTransformSRT3d PatternTransform = CurrentPattern[k];
+					FTransform WorldTransform;
+					ComputeWorldTransform(WorldTransform, (FTransform)ElementTransform, (FTransform)PatternTransform);
+
+					if (k == 0)
 					{
-						FTransformSRT3d PatternTransform = CurrentPattern[k];
-                        FTransform WorldTransform;
-                        ComputeWorldTransform(WorldTransform, (FTransform)ElementTransform, (FTransform)PatternTransform);
-						if (k == 0)
+						FCreateActorParams SpawnInfo;
+						SpawnInfo.BaseName = FString::Printf(TEXT("Pattern_%d"), ElemIdx);
+						SpawnInfo.TargetWorld = GetTargetWorld();
+						SpawnInfo.TemplateActor = SourceComponent->GetOwner();
+						SpawnInfo.Transform = FTransform::Identity;
+
+						FCreateActorResult Result = UE::Modeling::CreateNewActor(GetToolManager(), MoveTemp(SpawnInfo));
+						if (Result.IsOK())
 						{
-							TemplateComponent = DuplicateObject<UStaticMeshComponent>(SourceComponent, ParentActor);
-							TemplateComponent->ClearFlags(RF_DefaultSubObject);
-							TemplateComponent->OnComponentCreated();
-							ParentActor->SetRootComponent(TemplateComponent);
-							TemplateComponent->RegisterComponent();
-							TemplateComponent->SetWorldTransform( WorldTransform );
-                            TemplateComponent->SetVisibility(false, bPropagateToChildren);
+							NewActor = Result.NewActor;
+							TemplateComponent = NewActor->GetComponentByClass<UStaticMeshComponent>();
+
+							// Not sure this will ever happen but we probably don't want to proceed if it does
+							ensureMsgf(NewActor->GetRootComponent(), TEXT("New actor has no root component."));
+
+							if (TemplateComponent)
+							{
+								TemplateComponent->SetStaticMesh(SourceComponent->GetStaticMesh());
+								for (int32 j = 0; j < Element.SourceMaterials.Num(); ++j)
+								{
+									TemplateComponent->SetMaterial(j, Element.SourceMaterials[j]);
+								}
+							}
+							else
+							{
+								TemplateComponent = DuplicateObject<UStaticMeshComponent>(SourceComponent, NewActor);
+								TemplateComponent->SetupAttachment(NewActor->GetRootComponent());
+								TemplateComponent->OnComponentCreated();
+								NewActor->AddInstanceComponent(TemplateComponent);
+								TemplateComponent->RegisterComponent();
+							}
+
+							TemplateComponent->SetWorldTransform(WorldTransform);
+
+							NewActors.Add(NewActor);
+							SetActorComponentsVisibility(NewActor, false);
 						}
 						else
 						{
-							UStaticMeshComponent* NewCloneComponent = DuplicateObject<UStaticMeshComponent>(TemplateComponent, ParentActor);
-							NewCloneComponent->ClearFlags(RF_DefaultSubObject);
-							NewCloneComponent->SetupAttachment(ParentActor->GetRootComponent());
-							NewCloneComponent->OnComponentCreated();
-							ParentActor->AddInstanceComponent(NewCloneComponent);
-							NewCloneComponent->RegisterComponent();
-							NewCloneComponent->SetWorldTransform( WorldTransform );
-                            NewCloneComponent->SetVisibility(false, bPropagateToChildren);
+							break;
 						}
 					}
+					else
+					{
+						UStaticMeshComponent* NewCloneComponent = DuplicateObject<UStaticMeshComponent>(TemplateComponent, NewActor);
+						NewCloneComponent->ClearFlags(RF_DefaultSubObject);
+						NewCloneComponent->SetupAttachment(NewActor->GetRootComponent());
+						NewCloneComponent->OnComponentCreated();
+						NewActor->AddInstanceComponent(NewCloneComponent);
+						NewCloneComponent->RegisterComponent();
+						NewCloneComponent->SetWorldTransform(WorldTransform);
+
+						// Using SetVisibility here instead of SetActorComponentsVisibility because in this particular output
+						// case with large patterns, there could be a lot of redundant iteration over and casting for components
+						// which have already been set to invisible.
+						NewCloneComponent->SetVisibility(false, bPropagateToChildren);
+					}
 				}
 			}
 		}
diff --git a/Engine/Plugins/Runtime/MeshModelingToolset/Source/ModelingComponents/Private/ModelingObjectsCreationAPI.cpp b/Engine/Plugins/Runtime/MeshModelingToolset/Source/ModelingComponents/Private/ModelingObjectsCreationAPI.cpp
--- a/Engine/Plugins/Runtime/MeshModelingToolset/Source/ModelingComponents/Private/ModelingObjectsCreationAPI.cpp
+++ b/Engine/Plugins/Runtime/MeshModelingToolset/Source/ModelingComponents/Private/ModelingObjectsCreationAPI.cpp
@@ -131,6 +131,30 @@ FCreateMaterialObjectResult UE::Modeling::CreateMaterialObject(UInteractiveToolM
 
 
 
+FCreateActorResult UE::Modeling::CreateNewActor(UInteractiveToolManager* ToolManager, FCreateActorParams&& CreateActorParams)
+{
+	if (ensure(ToolManager))
+	{
+		UModelingObjectsCreationAPI* UseAPI = ToolManager->GetContextObjectStore()->FindContext<UModelingObjectsCreationAPI>();
+		if (UseAPI)
+		{
+			if (UseAPI->HasMoveVariants())
+			{
+				return UseAPI->CreateNewActor(MoveTemp(CreateActorParams));
+			}
+			else
+			{
+				return UseAPI->CreateNewActor(CreateActorParams);
+			}
+		}
+	}
+	return FCreateActorResult{ ECreateModelingObjectResult::Failed_NoAPIFound };
+}
+
+
+
+
+
 
 
 
diff --git a/Engine/Plugins/Runtime/MeshModelingToolset/Source/ModelingComponents/Public/ModelingObjectsCreationAPI.h b/Engine/Plugins/Runtime/MeshModelingToolset/Source/ModelingComponents/Public/ModelingObjectsCreationAPI.h
--- a/Engine/Plugins/Runtime/MeshModelingToolset/Source/ModelingComponents/Public/ModelingObjectsCreationAPI.h
+++ b/Engine/Plugins/Runtime/MeshModelingToolset/Source/ModelingComponents/Public/ModelingObjectsCreationAPI.h
@@ -377,6 +377,64 @@ struct MODELINGCOMPONENTS_API FCreateMaterialObjectResult
 
 
 
+/**
+ * FCreateActorParams is a collection of input data intended to be passed to
+ * UModelingObjectsCreationAPI::CreateNewActor().
+ */
+USTRUCT(Blueprintable)
+struct MODELINGCOMPONENTS_API FCreateActorParams
+{
+	GENERATED_BODY()
+
+	//
+	// Base data
+	//
+
+	/** 
+	 * The World/Level the new Actor should be created in (if known).
+	 */
+	UPROPERTY(Category = "CreateActorParams", EditAnywhere)
+	TObjectPtr<UWorld> TargetWorld = nullptr;
+
+	/** The base name of the new Actor */
+	UPROPERTY(Category = "CreateActorParams", EditAnywhere)
+	FString BaseName;
+
+	/** The 3D local-to-world transform for the new actor */
+	UPROPERTY(Category = "CreateActorParams", EditAnywhere)
+	FTransform Transform;
+
+	//
+	// input data
+	//
+
+	/** A template Actor the new Actor is based on */
+	UPROPERTY(Category = "CreateActorParams", EditAnywhere)
+	TObjectPtr<AActor> TemplateActor = nullptr;
+};
+
+
+/**
+ * FCreateActorResult is returned by UModelingObjectsCreationAPI::FCreateActorParams()
+ * to indicate success/failure and provide information about created actors
+ */
+USTRUCT(BlueprintType)
+struct MODELINGCOMPONENTS_API FCreateActorResult
+{
+	GENERATED_BODY()
+
+	/** Success/Failure status for the requested operation */
+	UPROPERTY(Category = "CreateMaterialObjectResult", VisibleAnywhere)
+	ECreateModelingObjectResult ResultCode = ECreateModelingObjectResult::Ok;
+
+	/** A pointer to a newly-created Actor */
+	UPROPERTY(Category="CreateMeshObjectResult", VisibleAnywhere)
+	TObjectPtr<AActor> NewActor = nullptr;
+
+
+	bool IsOK() const { return ResultCode == ECreateModelingObjectResult::Ok; }
+};
+
 
 
 
@@ -427,6 +485,12 @@ public:
 	UFUNCTION(BlueprintCallable, Category = "Modeling Objects")
 	virtual FCreateMaterialObjectResult CreateMaterialObject(const FCreateMaterialObjectParams& CreateMaterialParams) { return FCreateMaterialObjectResult{ ECreateModelingObjectResult::Failed_Unknown }; }
 
+	/**
+	 * Create a new material object based on the data in CreateMaterialParams
+	 * @return a results data structure, containing a result code and information about any new objects created
+	 */
+	UFUNCTION(BlueprintCallable, Category = "Modeling Objects")
+	virtual FCreateActorResult CreateNewActor(const FCreateActorParams& CreateActorParams) { return FCreateActorResult{ ECreateModelingObjectResult::Failed_Unknown }; }
 
 	//
 	// Non-UFunction variants that support std::move operators
@@ -441,6 +505,7 @@ public:
 	virtual FCreateMeshObjectResult CreateMeshObject(FCreateMeshObjectParams&& CreateMeshParams) { return FCreateMeshObjectResult{ ECreateModelingObjectResult::Failed_Unknown }; }
 	virtual FCreateTextureObjectResult CreateTextureObject(FCreateTextureObjectParams&& CreateTexParams) { return FCreateTextureObjectResult{ ECreateModelingObjectResult::Failed_Unknown }; }
 	virtual FCreateMaterialObjectResult CreateMaterialObject(FCreateMaterialObjectParams&& CreateMaterialParams) { return FCreateMaterialObjectResult{ ECreateModelingObjectResult::Failed_Unknown }; }
+	virtual FCreateActorResult CreateNewActor(FCreateActorParams&& CreateActorParams) { return FCreateActorResult{ ECreateModelingObjectResult::Failed_Unknown }; }
 
 };
 
@@ -480,6 +545,15 @@ MODELINGCOMPONENTS_API FCreateTextureObjectResult CreateTextureObject(UInteracti
 MODELINGCOMPONENTS_API FCreateMaterialObjectResult CreateMaterialObject(UInteractiveToolManager* ToolManager, FCreateMaterialObjectParams&& CreateMaterialParams);
 
 
+/**
+ * Create a new actor based on the data in CreateActorParams.
+ * This is a convenience function that will try to locate a UModelingObjectsCreationAPI instance in the ToolManager's ContextObjectStore,
+ * and then call UModelingObjectsCreationAPI::CreateNewActor()
+ * @return a results data structure, containing a result code and information about any new objects created
+ */
+MODELINGCOMPONENTS_API FCreateActorResult CreateNewActor(UInteractiveToolManager* ToolManager, FCreateActorParams&& CreateActorParams);
+
+
 
 /**
  * Strip the appended auto-generated _UUID suffix on the given string, if we can detect one
diff --git a/Engine/Plugins/Runtime/MeshModelingToolset/Source/ModelingComponentsEditorOnly/Private/EditorModelingObjectsCreationAPI.cpp b/Engine/Plugins/Runtime/MeshModelingToolset/Source/ModelingComponentsEditorOnly/Private/EditorModelingObjectsCreationAPI.cpp
--- a/Engine/Plugins/Runtime/MeshModelingToolset/Source/ModelingComponentsEditorOnly/Private/EditorModelingObjectsCreationAPI.cpp
+++ b/Engine/Plugins/Runtime/MeshModelingToolset/Source/ModelingComponentsEditorOnly/Private/EditorModelingObjectsCreationAPI.cpp
@@ -37,6 +37,7 @@
 
 using namespace UE::Geometry;
 
+extern UNREALED_API UEditorEngine* GEditor;
 
 UEditorModelingObjectsCreationAPI* UEditorModelingObjectsCreationAPI::Register(UInteractiveToolsContext* ToolsContext)
 {
@@ -103,10 +104,16 @@ FCreateMaterialObjectResult UEditorModelingObjectsCreationAPI::CreateMaterialObj
 {
 	FCreateMaterialObjectParams LocalParams = CreateMaterialParams;
 	return CreateMaterialObject(MoveTemp(LocalParams));
+}
 
+FCreateActorResult UEditorModelingObjectsCreationAPI::CreateNewActor(const FCreateActorParams& CreateActorParams)
+{
+	FCreateActorParams LocalParams = CreateActorParams;
+	return CreateNewActor(MoveTemp(LocalParams));
 }
 
 
+
 FCreateMeshObjectResult UEditorModelingObjectsCreationAPI::CreateMeshObject(FCreateMeshObjectParams&& CreateMeshParams)
 {
 	FCreateMeshObjectResult ResultOut;
@@ -491,6 +498,37 @@ FCreateMaterialObjectResult UEditorModelingObjectsCreationAPI::CreateMaterialObj
 }
 
 
+
+
+FCreateActorResult UEditorModelingObjectsCreationAPI::CreateNewActor(FCreateActorParams&& CreateActorParams)
+{
+	// create new Actor
+	AActor* NewActor = nullptr;
+	if (GEditor)
+	{
+		if (UActorFactory* const ActorFactory = GEditor->FindActorFactoryForActorClass(CreateActorParams.TemplateActor->GetClass()))
+		{
+			NewActor = ActorFactory->CreateActor(nullptr, CreateActorParams.TargetWorld->GetCurrentLevel(), CreateActorParams.Transform);
+			FActorLabelUtilities::SetActorLabelUnique(NewActor, CreateActorParams.BaseName);
+		}
+	}
+
+	if (!NewActor)
+	{
+		return FCreateActorResult{ ECreateModelingObjectResult::Failed_ActorCreationFailed };
+	}
+
+	FCreateActorResult ResultOut;
+	ResultOut.ResultCode = ECreateModelingObjectResult::Ok;
+	ResultOut.NewActor = NewActor;
+
+	OnModelingActorCreated.Broadcast(ResultOut);
+
+	return ResultOut;
+}
+
+
+
 ECreateModelingObjectResult UEditorModelingObjectsCreationAPI::GetNewAssetPath(
 	FString& OutNewAssetPath,
 	const FString& BaseName,
diff --git a/Engine/Plugins/Runtime/MeshModelingToolset/Source/ModelingComponentsEditorOnly/Public/EditorModelingObjectsCreationAPI.h b/Engine/Plugins/Runtime/MeshModelingToolset/Source/ModelingComponentsEditorOnly/Public/EditorModelingObjectsCreationAPI.h
--- a/Engine/Plugins/Runtime/MeshModelingToolset/Source/ModelingComponentsEditorOnly/Public/EditorModelingObjectsCreationAPI.h
+++ b/Engine/Plugins/Runtime/MeshModelingToolset/Source/ModelingComponentsEditorOnly/Public/EditorModelingObjectsCreationAPI.h
@@ -40,6 +40,9 @@ public:
 	// UFUNCTION(BlueprintCallable, Category = "Modeling Objects")
 	virtual FCreateMaterialObjectResult CreateMaterialObject(const FCreateMaterialObjectParams& CreateMaterialParams) override;
 
+	// UFUNCTION(BlueprintCallable, Category = "Modeling Objects")
+	virtual FCreateActorResult CreateNewActor(const FCreateActorParams& CreateActorParams) override;
+
 	//
 	// Non-UFunction variants that support std::move operators
 	//
@@ -52,6 +55,8 @@ public:
 
 	virtual FCreateMaterialObjectResult CreateMaterialObject(FCreateMaterialObjectParams&& CreateMaterialParams) override;
 
+	virtual FCreateActorResult CreateNewActor(FCreateActorParams&& CreateActorParams) override;
+
 
 	//
 	// Callbacks that editor can hook into to handle asset creation
@@ -69,7 +74,8 @@ public:
 	DECLARE_MULTICAST_DELEGATE_OneParam(FModelingMaterialCreatedSignature, const FCreateMaterialObjectResult& CreatedInfo);
 	FModelingMaterialCreatedSignature OnModelingMaterialCreated;
 
-
+	DECLARE_MULTICAST_DELEGATE_OneParam(FModelingActorCreatedSignature, const FCreateActorResult& CreatedInfo);
+	FModelingActorCreatedSignature OnModelingActorCreated;
 
 	//
 	// Utility functions to handle registration/unregistration
