CombineMeshInstances: add support for passing unique mesh LOD chains as array of MeshDescriptions, as an alternative to passing UStaticMeshes.

diff --git a/Engine/Plugins/Experimental/MeshModelingToolsetExp/Source/GeometryProcessingAdapters/Private/CombineMeshInstancesImpl.cpp b/Engine/Plugins/Experimental/MeshModelingToolsetExp/Source/GeometryProcessingAdapters/Private/CombineMeshInstancesImpl.cpp
--- a/Engine/Plugins/Experimental/MeshModelingToolsetExp/Source/GeometryProcessingAdapters/Private/CombineMeshInstancesImpl.cpp
+++ b/Engine/Plugins/Experimental/MeshModelingToolsetExp/Source/GeometryProcessingAdapters/Private/CombineMeshInstancesImpl.cpp
@@ -105,11 +105,32 @@ struct FMeshPartInstance
 // MeshPart represents a set of FMeshPartInstances of a particular FSourceGeometry
 struct FMeshPart
 {
-	UStaticMesh* SourceAsset;
+	// only one of the values below can be non-null
+	UStaticMesh* SourceAsset = nullptr;
+	const IGeometryProcessing_CombineMeshInstances::FMeshLODSet* SourceMeshLODSet = nullptr;
+
 	TArray<FMeshPartInstance> Instances;
 
 	bool bPreserveUVs = false;
 	bool bAllowMerging = true;
+
+
+	int32 GetNumTriangles() const 
+	{
+		if (SourceAsset != nullptr)
+		{
+			return SourceAsset->GetNumTriangles(0);
+		}
+		else if ( SourceMeshLODSet != nullptr )
+		{
+			return SourceMeshLODSet->ReferencedMeshLODs[0]->Triangles().Num();
+		}
+		else
+		{
+			ensure(false);
+			return 0;
+		}
+	}
 };
 
 
@@ -161,15 +182,15 @@ void InitializeMeshPartAssembly(
 	const IGeometryProcessing_CombineMeshInstances::FSourceInstanceList& SourceInstanceList,
 	FMeshPartsAssembly& AssemblyOut)
 {
-	TMap<UStaticMesh*, FMeshPart*> MeshToPartMap;
+	TMap<UStaticMesh*, FMeshPart*> StaticMeshToPartMap;
 
-	int32 NumInstances = SourceInstanceList.StaticMeshInstances.Num();
-	for ( int32 Index = 0; Index < NumInstances; ++Index)
+	int32 NumStaticMeshInstances = SourceInstanceList.StaticMeshInstances.Num();
+	for ( int32 Index = 0; Index < NumStaticMeshInstances; ++Index)
 	{
 		const IGeometryProcessing_CombineMeshInstances::FStaticMeshInstance& SourceMeshInstance = SourceInstanceList.StaticMeshInstances[Index];
 
 		UStaticMesh* StaticMesh = SourceMeshInstance.SourceMesh;
-		FMeshPart** FoundPart = MeshToPartMap.Find(StaticMesh);
+		FMeshPart** FoundPart = StaticMeshToPartMap.Find(StaticMesh);
 		if (FoundPart == nullptr)
 		{
 			TUniquePtr<FMeshPart> NewPart = MakeUnique<FMeshPart>();
@@ -179,12 +200,12 @@ void InitializeMeshPartAssembly(
 			AssemblyOut.Parts.Add(MoveTemp(NewPart));
 			// store source model?
 
-			MeshToPartMap.Add(StaticMesh, Ptr);
+			StaticMeshToPartMap.Add(StaticMesh, Ptr);
 			FoundPart = &Ptr;
 		}
 
 		FMeshPartInstance NewInstance;
-		NewInstance.ExternalInstanceIndex = FIndex3i(Index, -1,-1);
+		NewInstance.ExternalInstanceIndex = FIndex3i(Index, 0, -1);
 
 		if ( SourceMeshInstance.GroupDataIndex >= 0 && SourceMeshInstance.GroupDataIndex < SourceInstanceList.InstanceGroupDatas.Num() )
 		{
@@ -206,23 +227,77 @@ void InitializeMeshPartAssembly(
 		(*FoundPart)->Instances.Add(NewInstance);
 	}
 
+
+	// todo: possibly should not assume that MeshLODSets contains unique sets, ie should find uniques and remap them?
+
+	TMap<const IGeometryProcessing_CombineMeshInstances::FMeshLODSet*, FMeshPart*> MeshLODSetToPartMap;
+
+	int32 NumMeshLODSetInstances = SourceInstanceList.MeshLODSetInstances.Num();
+	for (int32 Index = 0; Index < NumMeshLODSetInstances; ++Index)
+	{
+		const IGeometryProcessing_CombineMeshInstances::FMeshLODSetInstance& SourceMeshInstance = SourceInstanceList.MeshLODSetInstances[Index];
+		int32 MeshSetIndex = SourceMeshInstance.MeshLODSetIndex;
+		if (MeshSetIndex < 0 || MeshSetIndex > SourceInstanceList.MeshLODSets.Num())
+		{
+			ensure(false);
+			continue;
+		}
+		const IGeometryProcessing_CombineMeshInstances::FMeshLODSet* MeshLODSet = &SourceInstanceList.MeshLODSets[MeshSetIndex];
+
+		FMeshPart** FoundPart = MeshLODSetToPartMap.Find(MeshLODSet);
+		if (FoundPart == nullptr)
+		{
+			TUniquePtr<FMeshPart> NewPart = MakeUnique<FMeshPart>();
+			NewPart->SourceMeshLODSet = MeshLODSet;
+			FMeshPart* Ptr = NewPart.Get();
+
+			AssemblyOut.Parts.Add(MoveTemp(NewPart));
+			// store source model?
+
+			MeshLODSetToPartMap.Add(MeshLODSet, Ptr);
+			FoundPart = &Ptr;
+		}
+
+		FMeshPartInstance NewInstance;
+		NewInstance.ExternalInstanceIndex = FIndex3i(Index, 1, -1);
+
+		if (SourceMeshInstance.GroupDataIndex >= 0 && SourceMeshInstance.GroupDataIndex < SourceInstanceList.InstanceGroupDatas.Num())
+		{
+			const IGeometryProcessing_CombineMeshInstances::FMeshInstanceGroupData& GroupData =
+				SourceInstanceList.InstanceGroupDatas[SourceMeshInstance.GroupDataIndex];
+			NewInstance.Materials = GroupData.MaterialSet;
+
+			(*FoundPart)->bPreserveUVs = GroupData.bPreserveUVs;
+			(*FoundPart)->bAllowMerging = GroupData.bAllowMerging;
+		}
+
+		NewInstance.SourceComponent = nullptr;
+		NewInstance.SourceInstanceIndex = 0;
+		NewInstance.DetailLevel = static_cast<EMeshDetailLevel>(static_cast<int32>(SourceMeshInstance.DetailLevel));
+		for (FTransform3d Transform : SourceMeshInstance.TransformSequence)
+		{
+			NewInstance.WorldTransform.Append(Transform);
+		}
+		(*FoundPart)->Instances.Add(NewInstance);
+	}
+
+
+
+
 	// sort parts by largest triangle count first
 	AssemblyOut.Parts.Sort([&](const TUniquePtr<FMeshPart>& A, const TUniquePtr<FMeshPart>& B)
 	{
-		return A->SourceAsset->GetNumTriangles(0) > B->SourceAsset->GetNumTriangles(0);
+		return A->GetNumTriangles() > B->GetNumTriangles();
 	});
 
-	// collect unique materials
-	for (TPair<UStaticMesh*, FMeshPart*>& Pair : MeshToPartMap)
-	{
-		UStaticMesh* StaticMesh = Pair.Key;
-		FMeshPart& Part = *(Pair.Value);
 
-		for (FMeshPartInstance& Instance : Part.Instances)
+	auto CollectUniqueMaterials = [&AssemblyOut](const FMeshPart& Part)
+	{
+		for (const FMeshPartInstance& Instance : Part.Instances)
 		{
 			for (UMaterialInterface* Material : Instance.Materials)
 			{
-				if ( AssemblyOut.MaterialMap.Contains(Material) == false)
+				if (AssemblyOut.MaterialMap.Contains(Material) == false)
 				{
 					int32 NewIndex = AssemblyOut.UniqueMaterials.Num();
 					AssemblyOut.UniqueMaterials.Add(Material);
@@ -230,8 +305,60 @@ void InitializeMeshPartAssembly(
 				}
 			}
 		}
+	};
+
+
+	// todo: why are these map iterations?? can't we just iterate over AssemblyOut.Parts?
+	// collect unique materials
+	for (TPair<UStaticMesh*, FMeshPart*>& Pair : StaticMeshToPartMap)
+	{
+		FMeshPart& Part = *(Pair.Value);
+		CollectUniqueMaterials(Part);
+	}
+	for (TPair<const IGeometryProcessing_CombineMeshInstances::FMeshLODSet*, FMeshPart*>& Pair : MeshLODSetToPartMap)
+	{
+		FMeshPart& Part = *(Pair.Value);
+		CollectUniqueMaterials(Part);
+	}
+}
+
+
+// Fetch a given LOD index of the Part and return in OutputLODMesh. Return false if LOD mesh is not available.
+bool ExtractSourceMeshLOD(
+	FMeshPart& Part,
+	int32 LODIndex,
+	FDynamicMesh3& OutputLODMesh)
+{
+	const FMeshDescription* UseMeshDescription = nullptr;
+	if (Part.SourceAsset != nullptr)
+	{
+		UStaticMesh* StaticMesh = Part.SourceAsset;
+		if (LODIndex < StaticMesh->GetNumSourceModels())
+		{
+			UseMeshDescription = StaticMesh->GetMeshDescription(LODIndex);
+		}
+	}
+	else if (Part.SourceMeshLODSet != nullptr)
+	{
+		if (LODIndex < Part.SourceMeshLODSet->ReferencedMeshLODs.Num())
+		{
+			UseMeshDescription = Part.SourceMeshLODSet->ReferencedMeshLODs[LODIndex];
+		}
+	}
+	else
+	{
+		ensure(false);
 	}
 
+	if (UseMeshDescription != nullptr)
+	{
+		FMeshDescriptionToDynamicMesh Converter;
+		Converter.bEnableOutputGroups = true;
+		Converter.bTransformVertexColorsLinearToSRGB = true;		// possibly this should be false...
+		Converter.Convert(UseMeshDescription, OutputLODMesh);
+		return true;
+	}
+	return false;
 }
 
 
@@ -254,34 +381,16 @@ void InitializeAssemblySourceMeshesFromLOD(
 		FSourceGeometry& Target = Assembly.SourceMeshGeometry[Index];
 		Target.SourceMeshLODs.SetNum(NumSourceLODs);
 
-		UStaticMesh* StaticMesh = Part->SourceAsset;
-
 		for (int32 k = 0; k < NumSourceLODs; ++k)
 		{
 			int32 LODIndex = SourceAssetBaseLOD + k;
-			if (LODIndex < StaticMesh->GetNumSourceModels())
-			{
-				FMeshDescription* UseMeshDescription = StaticMesh->GetMeshDescription(LODIndex);
-				if (UseMeshDescription != nullptr)
-				{
-					FMeshDescriptionToDynamicMesh Converter;
-					Converter.bEnableOutputGroups = true; 
-					Converter.bTransformVertexColorsLinearToSRGB = true;
-					Converter.Convert(UseMeshDescription, Target.SourceMeshLODs[k]);
-				}
-			}
+			ExtractSourceMeshLOD(*Part, LODIndex, Target.SourceMeshLODs[k]);
 		}
 
-		// if first LOD is missing try getting LOD0 again
-		if (Target.SourceMeshLODs[0].TriangleCount() == 0)
+		// if first LOD is missing try getting LOD0 again  (why?? Only doing for static mesh asset path because it was done before...)
+		if (Target.SourceMeshLODs[0].TriangleCount() == 0 && Part->SourceAsset != nullptr)
 		{
-			if (FMeshDescription* UseMeshDescription = StaticMesh->GetMeshDescription(0))
-			{
-				FMeshDescriptionToDynamicMesh Converter;
-				Converter.bEnableOutputGroups = true; 
-				Converter.bTransformVertexColorsLinearToSRGB = true;
-				Converter.Convert(UseMeshDescription, Target.SourceMeshLODs[0]);
-			}
+			ExtractSourceMeshLOD(*Part, 0, Target.SourceMeshLODs[0]);
 		}
 
 		// now if first LOD is missing, just fall back to a box
@@ -309,12 +418,13 @@ void InitializeAssemblySourceMeshesFromLOD(
 		TUniquePtr<FMeshPart>& Part = Assembly.Parts[Index];
 		FSourceGeometry& Target = Assembly.SourceMeshGeometry[Index];
 
-		UStaticMesh* StaticMesh = Part->SourceAsset;
-		UBodySetup* BodySetup = StaticMesh->GetBodySetup();
-		if (BodySetup)
+		if ( UStaticMesh* StaticMesh = Part->SourceAsset )
 		{
-			UE::Geometry::GetShapeSet(BodySetup->AggGeom, Target.CollisionShapes);
-			// todo: detect boxes?
+			if ( UBodySetup* BodySetup = StaticMesh->GetBodySetup() )
+			{
+				UE::Geometry::GetShapeSet(BodySetup->AggGeom, Target.CollisionShapes);
+				// todo: detect boxes?
+			}
 		}
 	}
 
@@ -2820,7 +2930,6 @@ void BuildCombinedMesh(
 		const TUniquePtr<FMeshPart>& Part = Assembly.Parts[SetIndex];
 		const FSourceGeometry& SourceGeometry = Assembly.SourceMeshGeometry[SetIndex];
 		const FOptimizedGeometry& OptimizedGeometry = Assembly.OptimizedMeshGeometry[SetIndex];
-		UStaticMesh* StaticMesh = Part->SourceAsset;
 
 		check(Part->bAllowMerging == true);		// do not support this yet
 
@@ -3353,7 +3462,7 @@ void FCombineMeshInstancesImpl::CombineMeshInstances(
 	{
 		int32 NumApproxLODs = FMath::Max(0, Options.NumLODs - Options.NumCopiedLODs - Options.NumSimplifiedLODs - Options.NumCoarseLODs);
 		UE_LOG(LogGeometry, Log, TEXT("CombineMeshInstances: processing %d Instances into %d LODs (%d Copied, %d Simplified, %d Approx, %d Coarse)"),
-			MeshInstances.StaticMeshInstances.Num(),
+			( MeshInstances.StaticMeshInstances.Num() + MeshInstances.MeshLODSetInstances.Num() ),
 			Options.NumLODs, Options.NumCopiedLODs, Options.NumSimplifiedLODs, NumApproxLODs, Options.NumCoarseLODs);
 	}
 
@@ -3380,8 +3489,13 @@ void FCombineMeshInstancesImpl::CombineMeshInstances(
 	
 	PartAssembly.PreProcessInstanceMeshFunc = [&PartAssembly, &MeshInstances, &Options](FDynamicMesh3& AppendMesh, const FMeshPartInstance& Instance)
 	{
-		int32 SourceInstance = Instance.ExternalInstanceIndex[0];
-		int GroupDataIdx = MeshInstances.StaticMeshInstances[SourceInstance].GroupDataIndex;
+		int32 SourceInstance = Instance.ExternalInstanceIndex.A;
+		bool bIsStaticMeshInstance = (Instance.ExternalInstanceIndex.B == 0);		// a bit of a hack here but we configured this above
+
+		int GroupDataIdx = (bIsStaticMeshInstance) ?
+			MeshInstances.StaticMeshInstances[SourceInstance].GroupDataIndex
+			: MeshInstances.MeshLODSetInstances[SourceInstance].GroupDataIndex;
+
 		if (MeshInstances.InstanceGroupDatas[GroupDataIdx].bHasConstantOverrideVertexColor)
 		{
 			FLinearColor VertexColorLinear(0,0,0,1);
diff --git a/Engine/Source/Developer/GeometryProcessingInterfaces/Public/GeometryProcessingInterfaces/CombineMeshInstances.h b/Engine/Source/Developer/GeometryProcessingInterfaces/Public/GeometryProcessingInterfaces/CombineMeshInstances.h
--- a/Engine/Source/Developer/GeometryProcessingInterfaces/Public/GeometryProcessingInterfaces/CombineMeshInstances.h
+++ b/Engine/Source/Developer/GeometryProcessingInterfaces/Public/GeometryProcessingInterfaces/CombineMeshInstances.h
@@ -10,7 +10,7 @@
 class UPrimitiveComponent;
 class UStaticMesh;
 class UMaterialInterface;
-
+struct FMeshDescription;
 
 /**
  * The CombineMeshInstances modular feature is used to provide a mechanism
@@ -49,14 +49,20 @@ public:
 	};
 
 
+	/**
+	 * FBaseMeshInstance is a base-struct for the various instance types below (FStaticMeshInstance, FMeshLODSetInstance)
+	 */
 	struct FBaseMeshInstance
 	{
 		EMeshDetailLevel DetailLevel = EMeshDetailLevel::Standard;
-		TArray<FTransform3d> TransformSequence;
+		TArray<FTransform3d> TransformSequence;		// set of transforms on this instance. Often just a single transform.
 		int32 GroupDataIndex = -1;		// index into FSourceInstanceList::InstanceGroupDatas
 	};
 
-	// a single instance of a static mesh asset
+
+	/**
+	 * FStaticMeshInstance represents a single instance of a static mesh asset
+	 */
 	struct FStaticMeshInstance : FBaseMeshInstance
 	{
 		UStaticMesh* SourceMesh = nullptr;
@@ -64,6 +70,24 @@ public:
 		int32 SourceInstanceIndex = 0;		// custom index, eg if SourceComponent is an InstancedStaticMeshComponent, this is the Instance index
 	};
 
+	/**
+	 * FMeshLODSet represents a list of LOD meshes that are used by one or more instances (eg like a StaticMesh, but could be other sources).
+	 */
+	struct FMeshLODSet
+	{
+		TArray<const FMeshDescription*> ReferencedMeshLODs;
+	};
+
+	/**
+	 * FMeshLODSetInstance represents a single instance of a FMeshLODSet
+	 */
+	struct FMeshLODSetInstance : FBaseMeshInstance
+	{
+		int32 MeshLODSetIndex = -1;		// index into a list of fixed Mesh LODs shared across multiple instances (eg FSourceInstanceList.MeshLODSets)
+
+		int32 ExternalInstanceID = 0;	// external identifier used for Instance for debugging/convenience purposes, this is not used internally by CombineMeshInstances, 
+	};
+
 	/**
 	 * FSourceInstanceList provides a flattened list of mesh instances to CombineMeshInstances.
 	 * Each Instance can refer to shared data in the InstanceGroupDatas list
@@ -73,6 +97,10 @@ public:
 	struct FSourceInstanceList
 	{
 		TArray<FStaticMeshInstance> StaticMeshInstances;
+		TArray<FMeshLODSetInstance> MeshLODSetInstances;
+
+		// mesh sets shared across instances
+		TArray<FMeshLODSet> MeshLODSets;
 
 		// sets of data shared across multiple instances
 		TArray<FMeshInstanceGroupData> InstanceGroupDatas;
@@ -219,9 +247,14 @@ public:
 
 
 		//
-		// 
+		// Final processing options applied to each LOD after generation
 		//
+
+		// ensure that all triangles have some UV values set. Various geometry processing steps may discard UVs, this flag will
+		// cause missing UVs to be recomputed (currently using a box projection)
 		bool bAutoGenerateMissingUVs = false;
+
+		// generate tangents on the output meshes - requires that UVs be available
 		bool bAutoGenerateTangents = false;
 
 
