Skin Paint Tool: fixes crash when selecting sockets in the skeleton tree

diff --git a/Engine/Plugins/Experimental/MeshModelingToolsetExp/Source/MeshModelingToolsEditorOnlyExp/Private/SkeletalMesh/SkinWeightsPaintTool.cpp b/Engine/Plugins/Experimental/MeshModelingToolsetExp/Source/MeshModelingToolsEditorOnlyExp/Private/SkeletalMesh/SkinWeightsPaintTool.cpp
--- a/Engine/Plugins/Experimental/MeshModelingToolsetExp/Source/MeshModelingToolsEditorOnlyExp/Private/SkeletalMesh/SkinWeightsPaintTool.cpp
+++ b/Engine/Plugins/Experimental/MeshModelingToolsetExp/Source/MeshModelingToolsEditorOnlyExp/Private/SkeletalMesh/SkinWeightsPaintTool.cpp
@@ -1025,7 +1025,7 @@ FVector4f USkinWeightsPaintTool::WeightToColor(float Value) const
 
 void USkinWeightsPaintTool::UpdateCurrentBoneVertexColors()
 {
-	const int32 CurrentBoneIndex = CurrentBone == NAME_None ? INDEX_NONE : Weights.Deformer.BoneNameToIndexMap[CurrentBone];
+	const int32 CurrentBoneIndex = GetBoneIndexFromName(CurrentBone);
 	
 	// update mesh with new value colors
 	PreviewMesh->DeferredEditMesh([this, &CurrentBoneIndex](FDynamicMesh3& Mesh)
@@ -1344,7 +1344,12 @@ void USkinWeightsPaintTool::GetVerticesToEdit(TArray<VertexIndex>& OutVertexIndi
 
 BoneIndex USkinWeightsPaintTool::GetBoneIndexFromName(const FName BoneName) const
 {
-	return  BoneName == NAME_None ? INDEX_NONE : Weights.Deformer.BoneNameToIndexMap[BoneName];
+	if (BoneName == NAME_None)
+	{
+		return  INDEX_NONE;		
+	}
+	const BoneIndex* Found = Weights.Deformer.BoneNameToIndexMap.Find(BoneName);
+	return Found ? *Found : INDEX_NONE;
 }
 
 void USkinWeightsPaintTool::OnShutdown(EToolShutdownType ShutdownType)
