MeshModelingTools: Create Draw Spline Tool for drawing splines and spline-based blueprints.

diff --git a/Engine/Plugins/Editor/ModelingToolsEditorMode/Source/ModelingToolsEditorMode/Private/ModelingToolsEditorMode.cpp b/Engine/Plugins/Editor/ModelingToolsEditorMode/Source/ModelingToolsEditorMode/Private/ModelingToolsEditorMode.cpp
--- a/Engine/Plugins/Editor/ModelingToolsEditorMode/Source/ModelingToolsEditorMode/Private/ModelingToolsEditorMode.cpp
+++ b/Engine/Plugins/Editor/ModelingToolsEditorMode/Source/ModelingToolsEditorMode/Private/ModelingToolsEditorMode.cpp
@@ -45,6 +45,7 @@
 #include "DrawPolygonTool.h"
 #include "DrawPolyPathTool.h"
 #include "DrawAndRevolveTool.h"
+#include "DrawSplineTool.h"
 #include "RevolveSplineTool.h"
 #include "ShapeSprayTool.h"
 #include "VoxelSolidifyMeshesTool.h"
@@ -602,6 +603,9 @@ void UModelingToolsEditorMode::Enter()
 	auto CubeGridToolBuilder = NewObject<UCubeGridToolBuilder>();
 	RegisterTool(ToolManagerCommands.BeginCubeGridTool, TEXT("BeginCubeGridTool"), CubeGridToolBuilder);
 
+	auto DrawSplineToolBuilder = NewObject<UDrawSplineToolBuilder>();
+	RegisterTool(ToolManagerCommands.BeginDrawSplineTool, TEXT("BeginDrawSplineTool"), DrawSplineToolBuilder);
+
 	auto TriangulateSplinesToolBuilder = NewObject<UTriangulateSplinesToolBuilder>();
 	RegisterTool(ToolManagerCommands.BeginTriangulateSplinesTool, TEXT("BeginTriangulateSplinesTool"), TriangulateSplinesToolBuilder);
 
diff --git a/Engine/Plugins/Editor/ModelingToolsEditorMode/Source/ModelingToolsEditorMode/Private/ModelingToolsEditorModeToolkit.cpp b/Engine/Plugins/Editor/ModelingToolsEditorMode/Source/ModelingToolsEditorMode/Private/ModelingToolsEditorModeToolkit.cpp
--- a/Engine/Plugins/Editor/ModelingToolsEditorMode/Source/ModelingToolsEditorMode/Private/ModelingToolsEditorModeToolkit.cpp
+++ b/Engine/Plugins/Editor/ModelingToolsEditorMode/Source/ModelingToolsEditorMode/Private/ModelingToolsEditorModeToolkit.cpp
@@ -260,6 +260,7 @@ void FModelingToolsEditorModeToolkit::RegisterPalettes()
 		Commands.BeginDrawPolyPathTool,
 		Commands.BeginDrawAndRevolveTool,
 		Commands.BeginRevolveSplineTool,
+		Commands.BeginDrawSplineTool,
 		Commands.BeginTriangulateSplinesTool
 	});
 	ToolkitBuilder->AddPalette(
diff --git a/Engine/Plugins/Editor/ModelingToolsEditorMode/Source/ModelingToolsEditorMode/Private/ModelingToolsManagerActions.cpp b/Engine/Plugins/Editor/ModelingToolsEditorMode/Source/ModelingToolsEditorMode/Private/ModelingToolsManagerActions.cpp
--- a/Engine/Plugins/Editor/ModelingToolsEditorMode/Source/ModelingToolsEditorMode/Private/ModelingToolsManagerActions.cpp
+++ b/Engine/Plugins/Editor/ModelingToolsEditorMode/Source/ModelingToolsEditorMode/Private/ModelingToolsManagerActions.cpp
@@ -83,6 +83,7 @@ void FModelingToolsManagerCommands::RegisterCommands()
 		REGISTER_MODELING_TOOL_COMMAND(BeginRevolveBoundaryTool, "BdryRev", "Revolve Mesh boundary loops to create new objects");
 		REGISTER_MODELING_TOOL_COMMAND(BeginCombineMeshesTool, "Merge", "Merge multiple Meshes to create new objects");
 		REGISTER_MODELING_TOOL_COMMAND(BeginDuplicateMeshesTool, "Dupe", "Duplicate single Meshes to create new objects");
+		REGISTER_MODELING_TOOL_COMMAND(BeginDrawSplineTool, "DrwSpln", "Draw a spline in the level");
 		REGISTER_MODELING_TOOL_COMMAND(BeginTriangulateSplinesTool, "MshSpln", "Triangulate the Spline Components of selected actors to create new objects");
 	}
 	else
@@ -94,6 +95,7 @@ void FModelingToolsManagerCommands::RegisterCommands()
 		REGISTER_MODELING_TOOL_COMMAND(BeginRevolveBoundaryTool, "Revolve Boundary", "Revolve Mesh boundary loops to create new objects");
 		REGISTER_MODELING_TOOL_COMMAND(BeginCombineMeshesTool, "Merge", "Merge multiple Meshes to create new objects");
 		REGISTER_MODELING_TOOL_COMMAND(BeginDuplicateMeshesTool, "Duplicate", "Duplicate single Meshes to create new objects");
+		REGISTER_MODELING_TOOL_COMMAND(BeginDrawSplineTool, "Draw Spline", "Draw a spline in the level");
 		REGISTER_MODELING_TOOL_COMMAND(BeginTriangulateSplinesTool, "Mesh Splines", "Triangulate the Spline Components of selected actors to create new objects");
 	}
 
diff --git a/Engine/Plugins/Editor/ModelingToolsEditorMode/Source/ModelingToolsEditorMode/Public/ModelingToolsManagerActions.h b/Engine/Plugins/Editor/ModelingToolsEditorMode/Source/ModelingToolsEditorMode/Public/ModelingToolsManagerActions.h
--- a/Engine/Plugins/Editor/ModelingToolsEditorMode/Source/ModelingToolsEditorMode/Public/ModelingToolsManagerActions.h
+++ b/Engine/Plugins/Editor/ModelingToolsEditorMode/Source/ModelingToolsEditorMode/Public/ModelingToolsManagerActions.h
@@ -55,6 +55,7 @@ public:
 	TSharedPtr<FUICommandInfo> BeginDrawAndRevolveTool;
 	TSharedPtr<FUICommandInfo> BeginRevolveSplineTool;
 	TSharedPtr<FUICommandInfo> BeginShapeSprayTool;
+	TSharedPtr<FUICommandInfo> BeginDrawSplineTool;
 	TSharedPtr<FUICommandInfo> BeginTriangulateSplinesTool;
 
 	TSharedPtr<FUICommandInfo> BeginSculptMeshTool;
diff --git a/Engine/Plugins/Experimental/MeshModelingToolsetExp/Source/MeshModelingToolsEditorOnlyExp/Private/DrawSplineTool.cpp b/Engine/Plugins/Experimental/MeshModelingToolsetExp/Source/MeshModelingToolsEditorOnlyExp/Private/DrawSplineTool.cpp
new file mode 100644
--- /dev/null
+++ b/Engine/Plugins/Experimental/MeshModelingToolsetExp/Source/MeshModelingToolsEditorOnlyExp/Private/DrawSplineTool.cpp
@@ -0,0 +1,1096 @@
+// Copyright Epic Games, Inc. All Rights Reserved.
+
+#include "DrawSplineTool.h"
+
+#include "ActorFactories/ActorFactoryEmptyActor.h"
+#include "AssetSelection.h" // FActorFactoryAssetProxy
+#include "BaseBehaviors/SingleClickOrDragBehavior.h"
+#include "BaseGizmos/GizmoMath.h"
+#include "Components/SplineComponent.h"
+#include "CoreGlobals.h" // GUndo
+#include "Kismet2/ComponentEditorUtils.h" // GenerateValidVariableName
+#include "Drawing/PreviewGeometryActor.h"
+#include "Editor/UnrealEdEngine.h" // DuplicateActors
+#include "Engine/World.h" 
+#include "GameFramework/Actor.h"
+#include "InputState.h" // FInputDeviceRay
+#include "InteractiveToolManager.h"
+#include "Mechanics/ConstructionPlaneMechanic.h"
+#include "SceneManagement.h" // FPrimitiveDrawInterface
+#include "Selection/ToolSelectionUtil.h"
+#include "SplineUtil.h"
+#include "ToolBuilderUtil.h"
+#include "UnrealEdGlobals.h" // GUnrealEd
+
+#define LOCTEXT_NAMESPACE "UDrawSplineTool"
+
+using namespace UE::Geometry;
+
+namespace DrawSplineToolLocals
+{
+	FText AddPointTransactionName = LOCTEXT("AddPointTransactionName", "Add Point");
+
+	USplineComponent* CreateNewSplineInActor(AActor* Actor, bool bTransact = false, bool bSetAsRoot = false)
+	{
+		if (!ensure(Actor))
+		{
+			return nullptr;
+		}
+
+		if (bTransact)
+		{
+			Actor->Modify();
+		}
+
+		FName NewComponentName = *FComponentEditorUtils::GenerateValidVariableName(USplineComponent::StaticClass(), Actor);
+		// Note that the RF_Transactional is important here for the spline to undo/redo properly in the future
+		USplineComponent* Spline = NewObject<USplineComponent>(Actor, USplineComponent::StaticClass(), NewComponentName, 
+			bTransact ? RF_Transactional : RF_NoFlags);
+
+		if (bSetAsRoot)
+		{
+			Actor->SetRootComponent(Spline);
+		}
+		else
+		{
+			Spline->SetupAttachment(Actor->GetRootComponent());
+		}
+
+		Spline->OnComponentCreated();
+		Actor->AddInstanceComponent(Spline);
+		Spline->RegisterComponent();
+		Spline->ResetRelativeTransform();
+		Actor->PostEditChange();
+
+		return Spline;
+	};
+
+	USplineComponent* GetOrCreateTargetSpline(AActor* Actor, int32 TargetIndex, bool bTransact = false)
+	{
+		if (TargetIndex >= 0)
+		{
+			TInlineComponentArray<USplineComponent*> SplineComponents;
+			Actor->GetComponents<USplineComponent>(SplineComponents);
+			if (TargetIndex < SplineComponents.Num())
+			{
+				return SplineComponents[TargetIndex];
+			}
+		}
+
+		// If we got to here, we didn't have an existing target at that index. Create one.
+		return CreateNewSplineInActor(Actor, bTransact);
+	}
+
+	void CopySplineToSpline(const USplineComponent& Source, USplineComponent& Destination, bool bTransact = false)
+	{
+		if (bTransact)
+		{
+			Destination.Modify();
+		}
+
+		Destination.ClearSplinePoints();
+		Destination.bSplineHasBeenEdited = true;
+
+		// We iterate here (rather than just copying over the SplineCurves data) so that we can transform
+		// the data properly into the coordinate space of the target component.
+		int32 NumSplinePoints = Source.GetNumberOfSplinePoints();
+		for (int32 i = 0; i < NumSplinePoints; ++i)
+		{
+			Destination.AddSplinePoint(Source.GetLocationAtSplinePoint(i, ESplineCoordinateSpace::World),
+				ESplineCoordinateSpace::World, false);
+			Destination.SetUpVectorAtSplinePoint(i, Source.GetUpVectorAtSplinePoint(i, ESplineCoordinateSpace::World),
+				ESplineCoordinateSpace::World, false);
+			Destination.SetTangentsAtSplinePoint(i,
+				Source.GetArriveTangentAtSplinePoint(i, ESplineCoordinateSpace::World),
+				Source.GetLeaveTangentAtSplinePoint(i, ESplineCoordinateSpace::World),
+				ESplineCoordinateSpace::World, false);
+			Destination.SetSplinePointType(i, Source.GetSplinePointType(i), false);
+		}
+
+		Destination.SetClosedLoop(Source.IsClosedLoop());
+
+		Destination.UpdateSpline();
+	};
+
+	// Gives the scale used for tangent visualization (and which therefore needs to be used in raycasting the handles)
+	float GetTangentScale()
+	{
+		return GetDefault<ULevelEditorViewportSettings>()->SplineTangentScale;
+	}
+
+	// Might be useful to have in SplineUtil, but uncertain what the API should be (should it be part of
+	// DrawSpline? Should there be options for selection color?). Also potentially messier to match the tangent
+	// scale with the UI interaction..
+	void DrawTangent(const USplineComponent& SplineComp, int32 PointIndex, IToolsContextRenderAPI& RenderAPI)
+	{
+		if (!ensure(PointIndex >= 0 && PointIndex < SplineComp.GetNumberOfSplinePoints()))
+		{
+			return;
+		}
+
+		FPrimitiveDrawInterface* PDI = RenderAPI.GetPrimitiveDrawInterface();
+
+		const float TangentScale = GetDefault<ULevelEditorViewportSettings>()->SplineTangentScale;
+		const float TangentHandleSize = 8.0f + GetDefault<ULevelEditorViewportSettings>()->SplineTangentHandleSizeAdjustment;
+
+		const FVector Location = SplineComp.GetLocationAtSplinePoint(PointIndex, ESplineCoordinateSpace::World);
+		const FVector LeaveTangent = SplineComp.GetLeaveTangentAtSplinePoint(PointIndex, ESplineCoordinateSpace::World) * TangentScale;
+		const FVector ArriveTangent = SplineComp.bAllowDiscontinuousSpline ?
+			SplineComp.GetArriveTangentAtSplinePoint(PointIndex, ESplineCoordinateSpace::World) * TangentScale : LeaveTangent;
+
+		FColor Color = FColor::White;
+
+		PDI->DrawLine(Location, Location - ArriveTangent, Color, SDPG_Foreground);
+		PDI->DrawLine(Location, Location + LeaveTangent, Color, SDPG_Foreground);
+
+		PDI->DrawPoint(Location + LeaveTangent, Color, TangentHandleSize, SDPG_Foreground);
+		PDI->DrawPoint(Location - ArriveTangent, Color, TangentHandleSize, SDPG_Foreground);
+	}
+
+	// Undoes a point addition with an auto tangent
+	class FSimplePointInsertionChange : public UDrawSplineTool::FSplineChange
+	{
+	public:
+		FSimplePointInsertionChange(const FVector3d& HitLocationIn, const FVector3d& UpVectorIn)
+			: HitLocation(HitLocationIn)
+			, UpVector(UpVectorIn)
+		{
+		}
+
+		virtual void Apply(USplineComponent& Spline) override
+		{
+			Spline.AddSplinePoint(HitLocation, ESplineCoordinateSpace::World, false);
+			int32 PointIndex = Spline.GetNumberOfSplinePoints() - 1;
+			Spline.SetUpVectorAtSplinePoint(PointIndex, UpVector, ESplineCoordinateSpace::World, true);
+		}
+		virtual void Revert(USplineComponent& Spline) override
+		{
+			if (ensure(Spline.GetNumberOfSplinePoints() > 0))
+			{
+				Spline.RemoveSplinePoint(Spline.GetNumberOfSplinePoints() - 1, true);
+			}
+		}
+		virtual FString ToString() const override
+		{
+			return TEXT("FSimplePointInsertionChange");
+		}
+
+	protected:
+		FVector3d HitLocation;
+		FVector3d UpVector;
+	};
+
+	// Undoes a point addition with an explicit tangent
+	class FTangentPointInsertionChange : public UDrawSplineTool::FSplineChange
+	{
+	public:
+		FTangentPointInsertionChange(const FVector3d& HitLocationIn, const FVector3d& UpVectorIn, const FVector3d& TangentIn)
+			: HitLocation(HitLocationIn)
+			, UpVector(UpVectorIn)
+			, Tangent(TangentIn)
+		{
+		}
+
+		virtual void Apply(USplineComponent& Spline) override
+		{
+			Spline.AddSplinePoint(HitLocation, ESplineCoordinateSpace::World, false);
+			int32 PointIndex = Spline.GetNumberOfSplinePoints() - 1;
+			Spline.SetUpVectorAtSplinePoint(PointIndex, UpVector, ESplineCoordinateSpace::World, false);
+			Spline.SetTangentAtSplinePoint(PointIndex, Tangent, ESplineCoordinateSpace::World, true);
+		}
+		virtual void Revert(USplineComponent& Spline) override
+		{
+			if (ensure(Spline.GetNumberOfSplinePoints() > 0))
+			{
+				Spline.RemoveSplinePoint(Spline.GetNumberOfSplinePoints() - 1, true);
+			}
+		}
+		virtual FString ToString() const override
+		{
+			return TEXT("FTangentPointInsertionChange");
+		}
+
+	protected:
+		FVector3d HitLocation;
+		FVector3d UpVector;
+		FVector3d Tangent;
+	};
+
+	// Undoes a free draw stroke (multiple points at once)
+	class FStrokeInsertionChange : public UDrawSplineTool::FSplineChange
+	{
+	public:
+		FStrokeInsertionChange(const TArray<FVector3d>& HitLocationsIn, const TArray<FVector3d>& UpVectorsIn)
+			: HitLocations(HitLocationsIn)
+			, UpVectors(UpVectorsIn)
+		{
+			if (!ensure(HitLocations.Num() == UpVectors.Num()))
+			{
+				int32 Num = FMath::Min(HitLocations.Num(), UpVectors.Num());
+				HitLocations.SetNum(Num);
+				UpVectors.SetNum(Num);
+			}
+		}
+
+		virtual void Apply(USplineComponent& Spline) override
+		{
+			for (int32 i = 0; i < HitLocations.Num(); ++i)
+			{
+				Spline.AddSplinePoint(HitLocations[i], ESplineCoordinateSpace::World, false);
+				int32 PointIndex = Spline.GetNumberOfSplinePoints() - 1;
+				Spline.SetUpVectorAtSplinePoint(PointIndex, UpVectors[i], ESplineCoordinateSpace::World, false);
+			}
+			Spline.UpdateSpline();
+		}
+		virtual void Revert(USplineComponent& Spline) override
+		{
+			for (int32 i = 0; i < HitLocations.Num(); ++i)
+			{
+				if (!ensure(Spline.GetNumberOfSplinePoints() > 0))
+				{
+					break;
+				}
+				Spline.RemoveSplinePoint(Spline.GetNumberOfSplinePoints() - 1, false);
+			}
+			Spline.UpdateSpline();
+		}
+		virtual FString ToString() const override
+		{
+			return TEXT("FStrokeInsertionChange");
+		}
+
+	protected:
+		TArray<FVector3d> HitLocations;
+		TArray<FVector3d> UpVectors;
+	};
+}
+
+void UDrawSplineTool::Setup()
+{
+	UInteractiveTool::Setup();
+
+	Settings = NewObject<UDrawSplineToolProperties>(this);
+	Settings->RestoreProperties(this);
+	AddToolPropertySource(Settings);
+
+	Settings->TargetActor = SelectedActor;
+
+	SetToolDisplayName(LOCTEXT("DrawSplineToolName", "Draw Spline"));
+	GetToolManager()->DisplayMessage(
+		LOCTEXT("DrawSplineToolDescription", "Draw a spline to replace an existing one or add it to an actor."),
+		EToolMessageLevel::UserNotification);
+
+	PlaneMechanic = NewObject<UConstructionPlaneMechanic>(this);
+	PlaneMechanic->Setup(this);
+	PlaneMechanic->Initialize(GetTargetWorld(), FFrame3d(FVector3d::Zero(), FVector3d::UnitX()));
+	PlaneMechanic->bShowGrid = Settings->bHitCustomPlane;
+	PlaneMechanic->CanUpdatePlaneFunc = [this] { return Settings->bHitCustomPlane; };
+	Settings->WatchProperty(Settings->bHitCustomPlane, [this](bool) {
+		PlaneMechanic->bShowGrid = Settings->bHitCustomPlane;
+	});
+
+	ClickOrDragBehavior = NewObject<USingleClickOrDragInputBehavior>();
+	ClickOrDragBehavior->Initialize(this, this);
+	AddInputBehavior(ClickOrDragBehavior);
+
+	Settings->WatchProperty(Settings->bLoop, [this](bool) {
+		if (ensure(WorkingSpline.IsValid()))
+		{
+			WorkingSpline->SetClosedLoop(Settings->bLoop);
+		}
+	});
+
+	TransitionOutputMode();
+
+	Settings->WatchProperty(Settings->OutputMode, [this](EDrawSplineOutputMode) {
+		TransitionOutputMode();
+	});
+	Settings->WatchProperty(Settings->TargetActor, [this](TWeakObjectPtr<AActor>) {
+		TransitionOutputMode();
+	});
+	Settings->WatchProperty(Settings->ExistingSplineIndexToReplace, [this](int32) {
+		TransitionOutputMode();
+	});
+	Settings->WatchProperty(Settings->BlueprintToCreate, [this](TWeakObjectPtr<UBlueprint>) {
+		TransitionOutputMode();
+		});
+	Settings->WatchProperty(Settings->bPreviewUsingActorCopy, [this](bool) {
+		TransitionOutputMode();
+	});
+
+	Settings->SilentUpdateWatched();
+}
+
+// Set things up for a new output mode or destination
+void UDrawSplineTool::TransitionOutputMode()
+{
+	using namespace DrawSplineToolLocals;
+
+	// Setting up the previews seems to be the most error prone part of the tool because editor duplicating, hiding
+	// from outliner, and avoiding emitting undo/redo transactions seems to be quite finnicky...
+
+	// This function is sometimes called from inside transactions (such as tool start, or dragging the "component to replace"
+	// slider). Several calls here would transact in that case (for instance, the Destroy() calls on the previews seem
+	// to do it), which we generally don't want to do. So we disable transacting in this function with the hack below.
+	// Note that we still have to take care that any editor functions we call don't open their own transactions...
+	ITransaction* UndoState = GUndo;
+	GUndo = nullptr; // Pretend we're not in a transaction
+	ON_SCOPE_EXIT{ GUndo = UndoState; }; // Revert later
+
+	// Restore the visibility of the previous target actor and spline, if needed
+	if (PreviousTargetActor)
+	{
+		PreviousTargetActor->SetIsTemporarilyHiddenInEditor(PreviousTargetActorVisibility);
+		PreviousTargetActor = nullptr;
+	}
+	if (HiddenSpline.IsValid())
+	{
+		HiddenSpline->bDrawDebug = bPreviousSplineVisibility;
+		HiddenSpline = nullptr;
+	}
+
+	// Keep the previous spline/previews temporarily so we can transfer over spline data
+	// when we make new previews
+	AActor* PreviousPreviewRoot = PreviewRootActor;
+	AActor* PreviousPreview = PreviewActor;
+	USplineComponent* PreviousSpline = WorkingSpline.Get();
+
+	PreviewRootActor = nullptr;
+	PreviewActor = nullptr;
+	WorkingSpline = nullptr;
+
+	// Create an entirely new preview root. We could probably keep the same one and disconnect/connect,
+	// but it seems cleaner to build from scratch whenever we have to change output mode.
+	FRotator Rotation(0.0f, 0.0f, 0.0f);
+	FActorSpawnParameters SpawnInfo;
+	SpawnInfo.ObjectFlags = RF_Transient;
+	PreviewRootActor = GetTargetWorld()->SpawnActor<APreviewGeometryActor>(FVector::ZeroVector, Rotation, SpawnInfo);
+	USceneComponent* RootComponent = NewObject<USceneComponent>(PreviewRootActor);
+	PreviewRootActor->AddOwnedComponent(RootComponent);
+	PreviewRootActor->SetRootComponent(RootComponent);
+	RootComponent->RegisterComponent();
+	
+	// Used for visualizing the effect of a spline on some special actor
+	auto CreateDuplicatePreviewActor = [this](AActor* Actor)
+	{
+		TArray<AActor*> NewActors;
+		GUnrealEd->DuplicateActors({ Actor }, NewActors, GetWorld()->GetCurrentLevel(), FVector3d::Zero());
+		if (!ensure(NewActors.Num() > 0))
+		{
+			return;
+		}
+
+		PreviewActor = NewActors[0];
+		PreviewActor->ClearFlags(RF_Transactional);
+		PreviewActor->SetFlags(RF_Transient);
+
+		if (!ensure(PreviewActor->GetRootComponent()))
+		{
+			USceneComponent* NewRoot = NewObject<USceneComponent>(PreviewRootActor);
+			PreviewActor->AddOwnedComponent(NewRoot);
+			PreviewActor->SetRootComponent(NewRoot);
+			NewRoot->RegisterComponent();
+		}
+
+		// Attach the preview actor to the non-outliner-visible preview root. The proper way to do this is 
+		// "GEditor->ParentActors(PreviewActor, PreviewRootActor, NAME_None);", but it is hard to prevent
+		// that call from emitting an undo/redo transaction. It may be possible if the actors AND the root
+		// components are all marked as not transactable, but seems simpler to do this by hand.
+		PreviewActor->GetRootComponent()->AttachToComponent(PreviewRootActor->GetRootComponent(), 
+			FAttachmentTransformRules::KeepWorldTransform, NAME_None);
+
+		// Hide the original
+		PreviousTargetActor = Actor;
+		PreviousTargetActorVisibility = PreviousTargetActor->IsHiddenEd();
+		PreviousTargetActor->SetIsTemporarilyHiddenInEditor(true);
+	};
+
+	auto FallbackSplinePlacement = [this]()
+	{
+		WorkingSpline = CreateNewSplineInActor(PreviewRootActor);
+	};
+
+	// Set up the new preview
+	if (!Settings->bPreviewUsingActorCopy)
+	{
+		FallbackSplinePlacement();
+	}
+	else
+	{
+		switch (Settings->OutputMode)
+		{
+		case EDrawSplineOutputMode::EmptyActor:
+			FallbackSplinePlacement();
+			break;
+		case EDrawSplineOutputMode::ExistingActor:
+		{
+			if (!Settings->TargetActor.IsValid())
+			{
+				FallbackSplinePlacement();
+				break;
+			}
+
+			CreateDuplicatePreviewActor(Settings->TargetActor.Get());
+
+			// Hide the spline we're replacing, if we are replacing one.
+			// TODO: This isn't quite perfect because if the spline is selected, the component visualizer in the editor
+			// will still draw it even if the parent actor is hidden and the spline is set to not be drawn...
+			TInlineComponentArray<USplineComponent*> SplineComponents;
+			Settings->TargetActor->GetComponents<USplineComponent>(SplineComponents);
+			if (Settings->ExistingSplineIndexToReplace >= 0 && Settings->ExistingSplineIndexToReplace < SplineComponents.Num())
+			{
+				HiddenSpline = SplineComponents[Settings->ExistingSplineIndexToReplace];
+				bPreviousSplineVisibility = HiddenSpline->bDrawDebug;
+				HiddenSpline->bDrawDebug = false;
+			}
+
+			WorkingSpline = GetOrCreateTargetSpline(PreviewActor, Settings->ExistingSplineIndexToReplace);
+			break;
+		}
+		case EDrawSplineOutputMode::CreateBlueprint:
+		{
+			bool bCanCreateActor = Settings->BlueprintToCreate.IsValid()
+				&& Settings->BlueprintToCreate->GeneratedClass != nullptr
+				&& !Settings->BlueprintToCreate->GeneratedClass->HasAnyClassFlags(CLASS_NotPlaceable | CLASS_Abstract);
+
+			if (!bCanCreateActor)
+			{
+				FallbackSplinePlacement();
+				break;
+			}
+
+			// Instantiate the blueprint
+			PreviewActor = FActorFactoryAssetProxy::AddActorForAsset(
+				Settings->BlueprintToCreate.Get(),
+				/*bSelectActors =*/ false,
+				// Important that we don't use the default (RF_Transactional) here, or else we'll end up
+				// issuing an undo transaction in this call.
+				EObjectFlags::RF_Transient);
+
+			WorkingSpline = GetOrCreateTargetSpline(PreviewActor, Settings->ExistingSplineIndexToReplace);
+			break;
+		}
+		default:
+			ensure(false);
+		}
+	}
+
+	if (ensure(WorkingSpline.IsValid()))
+	{
+		if (PreviousSpline)
+		{
+			CopySplineToSpline(*PreviousSpline, *WorkingSpline);
+		}
+		else
+		{
+			WorkingSpline->ClearSplinePoints();
+		}
+
+		WorkingSpline->SetClosedLoop(Settings->bLoop);
+
+		// This has to be set so that construction script reruns transfer over current spline state.
+		WorkingSpline->bSplineHasBeenEdited = true;
+
+		// Get the index of the spline in the components array for recapturing on construction script reruns.
+		if (PreviewActor)
+		{
+			TInlineComponentArray<USplineComponent*> SplineComponents;
+			PreviewActor->GetComponents<USplineComponent>(SplineComponents);
+			SplineRecaptureIndex = SplineComponents.IndexOfByKey(WorkingSpline.Get());
+			ensure(SplineRecaptureIndex >= 0);
+		}
+	}
+	
+	// Now that we've copied over previous preview data, destroy the old previews
+	if (PreviousPreview)
+	{
+		PreviousPreview->Destroy();
+	}
+	if (PreviousPreviewRoot)
+	{
+		PreviousPreviewRoot->Destroy();
+	}
+}
+
+void UDrawSplineTool::Shutdown(EToolShutdownType ShutdownType)
+{
+	Settings->SaveProperties(this);
+
+	if (PreviousTargetActor)
+	{
+		PreviousTargetActor->SetIsTemporarilyHiddenInEditor(PreviousTargetActorVisibility);
+		PreviousTargetActor = nullptr;
+	}
+	if (HiddenSpline.IsValid())
+	{
+		HiddenSpline->bDrawDebug = bPreviousSplineVisibility;
+		HiddenSpline = nullptr;
+	}
+
+	int32 NumSplinePoints = WorkingSpline->GetNumberOfSplinePoints();
+	if (ShutdownType == EToolShutdownType::Accept && NumSplinePoints > 0)
+	{
+		GenerateAsset();
+	}
+
+	PlaneMechanic->Shutdown();
+
+	if (PreviewActor)
+	{
+		PreviewActor->Destroy();
+	}
+	if (PreviewRootActor)
+	{
+		PreviewRootActor->Destroy();
+	}
+
+	Super::Shutdown(ShutdownType);
+}
+
+void UDrawSplineTool::GenerateAsset()
+{
+	using namespace DrawSplineToolLocals;
+	USplineComponent* OutputSpline = nullptr;
+	auto CreateSplineInEmptyActor = [this, &OutputSpline]()
+	{
+		// Get centroid of spline
+		int32 NumSplinePoints = WorkingSpline->GetNumberOfSplinePoints();
+		FVector3d Center = FVector3d::Zero();
+		for (int32 i = 0; i < NumSplinePoints; ++i)
+		{
+			Center += WorkingSpline->GetLocationAtSplinePoint(i, ESplineCoordinateSpace::World);
+		}
+		Center /= NumSplinePoints;
+
+		// Spawning via a factory is editor-only
+		UActorFactoryEmptyActor* EmptyActorFactory = NewObject<UActorFactoryEmptyActor>();
+		FAssetData AssetData(EmptyActorFactory->GetDefaultActorClass(FAssetData()));
+		FActorSpawnParameters SpawnParams;
+		SpawnParams.Name = TEXT("Spline");
+		SpawnParams.NameMode = FActorSpawnParameters::ESpawnActorNameMode::Requested;
+		AActor* NewActor = EmptyActorFactory->CreateActor(AssetData.GetAsset(),
+			TargetWorld->GetCurrentLevel(),
+			FTransform(Center),
+			SpawnParams);
+
+		// This is also editor-only: it's the label that shows up in the hierarchy
+		FActorLabelUtilities::SetActorLabelUnique(NewActor, TEXT("SplineActor"));
+
+		UActorComponent* OldRoot = NewActor->GetRootComponent();
+		OutputSpline = CreateNewSplineInActor(NewActor, true, true);
+		OutputSpline->SetWorldTransform(FTransform(Center));
+		OldRoot->DestroyComponent();
+
+		CopySplineToSpline(*WorkingSpline, *OutputSpline, true);
+	};
+
+	GetToolManager()->BeginUndoTransaction(LOCTEXT("DrawSplineTransactionName", "Draw Spline"));
+
+	switch (Settings->OutputMode)
+	{
+	case EDrawSplineOutputMode::EmptyActor:
+		CreateSplineInEmptyActor();
+		break;
+	case EDrawSplineOutputMode::ExistingActor:
+	{
+		if (!Settings->TargetActor.IsValid())
+		{
+			CreateSplineInEmptyActor();
+			break;
+		}
+
+		OutputSpline = GetOrCreateTargetSpline(Settings->TargetActor.Get(), Settings->ExistingSplineIndexToReplace, true);
+		CopySplineToSpline(*WorkingSpline, *OutputSpline, true);
+		Settings->TargetActor->RerunConstructionScripts();
+		break;
+	}
+	case EDrawSplineOutputMode::CreateBlueprint:
+	{
+		bool bCanCreateActor = Settings->BlueprintToCreate.IsValid()// != nullptr
+			&& Settings->BlueprintToCreate->GeneratedClass != nullptr
+			&& !Settings->BlueprintToCreate->GeneratedClass->HasAnyClassFlags(CLASS_NotPlaceable | CLASS_Abstract);
+
+		if (!bCanCreateActor)
+		{
+			CreateSplineInEmptyActor();
+			break;
+		}
+
+		// Instantiate the blueprint
+		AActor* NewActor = FActorFactoryAssetProxy::AddActorForAsset(
+			Settings->BlueprintToCreate.Get(),
+			/*bSelectActors =*/ false);
+
+		OutputSpline = GetOrCreateTargetSpline(NewActor, Settings->ExistingSplineIndexToReplace, true);
+		CopySplineToSpline(*WorkingSpline, *OutputSpline, true);
+		NewActor->RerunConstructionScripts();
+		break;
+	}
+	default:
+		ensure(false);
+	}
+
+	// TODO: Someday when we support component selection, we should select OutputSpline directly.
+	ToolSelectionUtil::SetNewActorSelection(GetToolManager(), OutputSpline->GetAttachmentRootActor());
+
+	GetToolManager()->EndUndoTransaction();
+}
+
+// Helper to add a point given a hit location and hit normal
+void UDrawSplineTool::AddSplinePoint(const FVector3d& HitLocation, const FVector3d& HitNormal)
+{
+	using namespace DrawSplineToolLocals;
+
+	int32 NumSplinePoints = WorkingSpline->GetNumberOfSplinePoints();
+	FVector3d UpVectorToUse = GetUpVectorToUse(HitLocation, HitNormal, NumSplinePoints);
+
+	WorkingSpline->AddSplinePoint(HitLocation, ESplineCoordinateSpace::World, 
+		/*bUpdate =*/ false);
+	WorkingSpline->SetUpVectorAtSplinePoint(NumSplinePoints, UpVectorToUse, ESplineCoordinateSpace::World, 
+		/*bUpdate =*/ true);
+}
+
+FVector3d UDrawSplineTool::GetUpVectorToUse(const FVector3d& HitLocation, const FVector3d& HitNormal, int32 NumSplinePointsBeforehand)
+{
+	FVector3d UpVectorToUse = HitNormal;
+	switch (Settings->UpVectorMode)
+	{
+	case EDrawSplineUpVectorMode::AlignToPrevious:
+	{
+		if (NumSplinePointsBeforehand == 0)
+		{
+			// TODO: Maybe add some different options of what normal to start with
+		}
+		else if (NumSplinePointsBeforehand > 1)
+		{
+			UpVectorToUse = WorkingSpline->GetUpVectorAtSplinePoint(NumSplinePointsBeforehand - 1, ESplineCoordinateSpace::World);
+		}
+		else // if NumSplinePointsBeforehand == 1
+		{
+			// If there's only one point, GetUpVectorAtSplinePoint is unreliable because it seeks to build a
+			// quaternion from the tangent and the set up vector, and the tangent is zero. We want to use
+			// the "stored" up vector directly.
+			FVector3d LocalUpVector = WorkingSpline->SplineCurves.Rotation.Points[0].OutVal.RotateVector(WorkingSpline->DefaultUpVector);
+			UpVectorToUse = WorkingSpline->GetComponentTransform().TransformVectorNoScale(LocalUpVector);
+		}
+		break;
+	}
+	case EDrawSplineUpVectorMode::UseHitNormal:
+		break;
+	}
+
+	return UpVectorToUse;
+}
+
+bool UDrawSplineTool::Raycast(const FRay& WorldRay, FVector3d& HitLocationOut, FVector3d& HitNormalOut, double& HitTOut)
+{
+	double BestHitT = TNumericLimits<double>::Max();
+	
+	if (Settings->bHitCustomPlane)
+	{
+		FVector IntersectionPoint;
+		bool bHitPlane = false;
+		GizmoMath::RayPlaneIntersectionPoint(PlaneMechanic->Plane.Origin, PlaneMechanic->Plane.Z(),
+			WorldRay.Origin, WorldRay.Direction, bHitPlane, IntersectionPoint);
+
+		if (bHitPlane)
+		{
+			HitLocationOut = IntersectionPoint;
+			HitNormalOut = PlaneMechanic->Plane.Z();
+			HitTOut = WorldRay.GetParameter(IntersectionPoint);
+			BestHitT = HitTOut;
+		}
+	}
+
+	if (Settings->bHitWorld)
+	{
+		FHitResult GeometryHit;
+		TArray<const UPrimitiveComponent*> ComponentsToIgnore;
+		if (PreviewActor)
+		{
+			PreviewActor->GetComponents<const UPrimitiveComponent>(ComponentsToIgnore);
+		}
+		if (ToolSceneQueriesUtil::FindNearestVisibleObjectHit(this, GeometryHit, WorldRay, &ComponentsToIgnore)
+			&& GeometryHit.Distance < BestHitT)
+		{
+			HitLocationOut = GeometryHit.ImpactPoint;
+			HitNormalOut = GeometryHit.ImpactNormal;
+			HitTOut = GeometryHit.Distance;
+			BestHitT = HitTOut;
+		}
+	}
+
+	// Only raycast the ground plane / ortho background if we didn't hit anything else
+	if (Settings->bHitGroundPlanes && BestHitT == TNumericLimits<double>::Max())
+	{
+		FVector3d PlaneNormal = CameraState.bIsOrthographic ? -WorldRay.Direction : FVector3d::UnitZ();
+		FVector IntersectionPoint;
+		bool bHitPlane = false;
+		GizmoMath::RayPlaneIntersectionPoint(FVector3d::Zero(), PlaneNormal,
+			WorldRay.Origin, WorldRay.Direction, bHitPlane, IntersectionPoint);
+
+		if (bHitPlane)
+		{
+			HitLocationOut = IntersectionPoint;
+			HitNormalOut = PlaneNormal;
+			HitTOut = WorldRay.GetParameter(IntersectionPoint);
+			BestHitT = HitTOut;
+		}
+	}
+	
+	return BestHitT < TNumericLimits<double>::Max();
+}
+
+FInputRayHit UDrawSplineTool::IsHitByClick(const FInputDeviceRay& ClickPos)
+{
+	FVector3d HitLocation, HitNormal;
+	double HitT;
+	if (Raycast(ClickPos.WorldRay, HitLocation, HitNormal, HitT))
+	{
+		return FInputRayHit(HitT);
+	}
+	return FInputRayHit();
+}
+
+void UDrawSplineTool::OnClicked(const FInputDeviceRay& ClickPos)
+{
+	using namespace DrawSplineToolLocals;
+
+	FVector3d HitLocation, HitNormal;
+	double HitT;
+	if (Raycast(ClickPos.WorldRay, HitLocation, HitNormal, HitT))
+	{
+		switch (Settings->DrawMode)
+		{
+		case EDrawSplineDrawMode::ClickAutoTangent:
+		case EDrawSplineDrawMode::FreeDraw:
+		{
+			AddSplinePoint(HitLocation, HitNormal);
+
+			int32 PointIndex = WorkingSpline->GetNumberOfSplinePoints() - 1;
+			GetToolManager()->EmitObjectChange(this,
+				MakeUnique<FSimplePointInsertionChange>(
+					HitLocation, 
+					WorkingSpline->GetUpVectorAtSplinePoint(PointIndex, ESplineCoordinateSpace::World)),
+				AddPointTransactionName);
+			break;
+		}
+		case EDrawSplineDrawMode::TangentDrag:
+		{
+			AddSplinePoint(HitLocation, HitNormal);
+
+			int32 PointIndex = WorkingSpline->GetNumberOfSplinePoints() - 1;
+			WorkingSpline->SetTangentAtSplinePoint(PointIndex,
+				FVector3d::Zero(),
+				ESplineCoordinateSpace::World, true);
+
+			GetToolManager()->EmitObjectChange(this,
+				MakeUnique<FTangentPointInsertionChange>( 
+					HitLocation, 
+					WorkingSpline->GetUpVectorAtSplinePoint(PointIndex, ESplineCoordinateSpace::World), 
+					FVector3d::Zero()),
+				AddPointTransactionName);
+			break;
+		}
+		}
+
+		bNeedToRerunConstructionScript = true;
+	}
+}
+
+FInputRayHit UDrawSplineTool::CanBeginClickDragSequence(const FInputDeviceRay& PressPos)
+{
+	FVector3d HitLocation, HitNormal;
+	double HitT;
+	if (Raycast(PressPos.WorldRay, HitLocation, HitNormal, HitT))
+	{
+		return FInputRayHit(HitT);
+	}
+	return FInputRayHit();
+}
+
+void UDrawSplineTool::OnClickPress(const FInputDeviceRay& PressPos)
+{
+	FVector3d HitLocation, HitNormal;
+	double HitT;
+
+	// Regardless of DrawMode, start by placing a point, though don't emit a transaction until mouse up
+	if (ensure(Raycast(PressPos.WorldRay, HitLocation, HitNormal, HitT)))
+	{
+		AddSplinePoint(HitLocation, HitNormal);
+
+		if (Settings->DrawMode == EDrawSplineDrawMode::FreeDraw)
+		{
+			// Remember which point started this stroke
+			FreeDrawStrokeStartIndex = WorkingSpline->GetNumberOfSplinePoints() - 1;
+		}
+		
+		bNeedToRerunConstructionScript = bNeedToRerunConstructionScript || Settings->bRerunConstructionScriptOnDrag;
+	}
+}
+
+void UDrawSplineTool::OnClickDrag(const FInputDeviceRay& DragPos)
+{
+	using namespace DrawSplineToolLocals;
+
+	int32 NumSplinePoints = WorkingSpline->GetNumberOfSplinePoints();
+	if (!ensure(NumSplinePoints > 0))
+	{
+		return;
+	}
+
+	FVector3d HitLocation, HitNormal;
+	double HitT;
+	if (Raycast(DragPos.WorldRay, HitLocation, HitNormal, HitT))
+	{
+		switch (Settings->DrawMode)
+		{
+		case EDrawSplineDrawMode::ClickAutoTangent:
+		{
+			// Drag the last placed point
+			FVector3d UpVector = GetUpVectorToUse(HitLocation, HitNormal, NumSplinePoints);
+			WorkingSpline->SetLocationAtSplinePoint(NumSplinePoints - 1, HitLocation, ESplineCoordinateSpace::World, false);
+			WorkingSpline->SetUpVectorAtSplinePoint(NumSplinePoints - 1, UpVector, ESplineCoordinateSpace::World, true);
+
+			break;
+		}
+		case EDrawSplineDrawMode::TangentDrag:
+		{
+			// Set the tangent
+			FVector3d LastPoint = WorkingSpline->GetLocationAtSplinePoint(NumSplinePoints - 1, ESplineCoordinateSpace::World);
+			FVector3d Tangent = (HitLocation - LastPoint) / GetTangentScale();
+			WorkingSpline->SetTangentAtSplinePoint(NumSplinePoints - 1, Tangent, ESplineCoordinateSpace::World, true);
+			bDrawTangentForLastPoint = true;
+			break;
+		}
+			
+		case EDrawSplineDrawMode::FreeDraw:
+		{
+			// Instead of dragging the first placed point (which gets placed in OnClickPress), we drag a second "preview" one
+			// until we get far enough from the previous to where we want to place a new control point.
+			if (!bFreeDrawPlacedPreviewPoint)
+			{
+				AddSplinePoint(HitLocation, HitNormal);
+				bFreeDrawPlacedPreviewPoint = true;
+			}
+			else
+			{
+				FVector3d PreviousPoint = WorkingSpline->GetLocationAtSplinePoint(NumSplinePoints - 2, ESplineCoordinateSpace::World);
+				if (FVector3d::DistSquared(HitLocation, PreviousPoint) >= Settings->MinPointSpacing * Settings->MinPointSpacing)
+				{
+					AddSplinePoint(HitLocation, HitNormal);
+				}
+				else
+				{
+					// Drag the preview point
+					FVector3d UpVector = GetUpVectorToUse(HitLocation, HitNormal, NumSplinePoints);
+					WorkingSpline->SetLocationAtSplinePoint(NumSplinePoints - 1, HitLocation, ESplineCoordinateSpace::World, false);
+					WorkingSpline->SetUpVectorAtSplinePoint(NumSplinePoints - 1, UpVector, ESplineCoordinateSpace::World, true);
+				}
+			}
+			break;
+		}
+		}
+	}
+
+	bNeedToRerunConstructionScript = bNeedToRerunConstructionScript || Settings->bRerunConstructionScriptOnDrag;
+}
+
+void UDrawSplineTool::OnClickRelease(const FInputDeviceRay& ReleasePos)
+{
+	OnClickDrag(ReleasePos);
+	OnTerminateDragSequence();
+}
+void UDrawSplineTool::OnTerminateDragSequence()
+{
+	using namespace DrawSplineToolLocals;
+
+	bDrawTangentForLastPoint = false;
+	bFreeDrawPlacedPreviewPoint = false;
+	bNeedToRerunConstructionScript = true;
+
+	int32 NumSplinePoints = WorkingSpline->GetNumberOfSplinePoints();
+
+	// Emit the appropriate undo transaction
+	switch (Settings->DrawMode)
+	{
+	case EDrawSplineDrawMode::ClickAutoTangent:
+	{
+		WorkingSpline->GetLocationAtSplinePoint(NumSplinePoints-1, ESplineCoordinateSpace::World);
+		WorkingSpline->GetUpVectorAtSplinePoint(NumSplinePoints - 1, ESplineCoordinateSpace::World);
+		GetToolManager()->EmitObjectChange(this,
+			MakeUnique<FSimplePointInsertionChange>(
+				WorkingSpline->GetLocationAtSplinePoint(NumSplinePoints - 1, ESplineCoordinateSpace::World),
+				WorkingSpline->GetUpVectorAtSplinePoint(NumSplinePoints - 1, ESplineCoordinateSpace::World)),
+			AddPointTransactionName);
+		break;
+	}
+	case EDrawSplineDrawMode::TangentDrag:
+	{
+		GetToolManager()->EmitObjectChange(this,
+			MakeUnique<FTangentPointInsertionChange>( 
+				WorkingSpline->GetLocationAtSplinePoint(NumSplinePoints - 1, ESplineCoordinateSpace::World),
+				WorkingSpline->GetUpVectorAtSplinePoint(NumSplinePoints - 1, ESplineCoordinateSpace::World),
+				WorkingSpline->GetTangentAtSplinePoint(NumSplinePoints - 1, ESplineCoordinateSpace::World)),
+			AddPointTransactionName);
+		break;
+	}
+
+	case EDrawSplineDrawMode::FreeDraw:
+	{
+		TArray<FVector3d> HitLocations;
+		TArray<FVector3d> UpVectors;
+		for (int32 i = FreeDrawStrokeStartIndex; i < NumSplinePoints; ++i)
+		{
+			HitLocations.Add(WorkingSpline->GetLocationAtSplinePoint(i, ESplineCoordinateSpace::World));
+			UpVectors.Add(WorkingSpline->GetUpVectorAtSplinePoint(i, ESplineCoordinateSpace::World));
+		}
+
+		GetToolManager()->EmitObjectChange(this,
+			MakeUnique<FStrokeInsertionChange>(HitLocations, UpVectors),
+			AddPointTransactionName);
+		break;
+	}
+	}
+}
+
+
+void UDrawSplineTool::OnTick(float DeltaTime)
+{
+	if (PlaneMechanic)
+	{
+		PlaneMechanic->Tick(DeltaTime);
+	}
+
+	if (bNeedToRerunConstructionScript)
+	{
+		bNeedToRerunConstructionScript = false;
+		if (PreviewActor)
+		{
+			PreviewActor->RerunConstructionScripts();
+
+			// Rerunning the construction script can make us lose our reference to the spline, so try to
+			// recapture.
+			// TODO: This might be avoidable with FComponentReference?
+			if (!WorkingSpline.IsValid())
+			{
+				TInlineComponentArray<USplineComponent*> SplineComponents;
+				PreviewActor->GetComponents<USplineComponent>(SplineComponents);
+
+				if (ensure(SplineRecaptureIndex >= 0 && SplineRecaptureIndex < SplineComponents.Num()))
+				{
+					WorkingSpline = SplineComponents[SplineRecaptureIndex];
+				}
+				else
+				{
+					// If we failed to recapture, it's not clear what to do. We can switch to working inside
+					// an empty actor, though we'll lose current spline state.
+					Settings->OutputMode = EDrawSplineOutputMode::EmptyActor;
+					Settings->CheckAndUpdateWatched();
+				}
+			}
+		}
+	}
+}
+
+void UDrawSplineTool::Render(IToolsContextRenderAPI* RenderAPI)
+{
+	using namespace DrawSplineToolLocals;
+
+	Super::Render(RenderAPI);
+
+	GetToolManager()->GetContextQueriesAPI()->GetCurrentViewState(CameraState);
+
+	if (PlaneMechanic)
+	{
+		PlaneMechanic->Render(RenderAPI);
+	}
+
+	if (WorkingSpline.IsValid())
+	{
+		if (bDrawTangentForLastPoint)
+		{
+			DrawTangent(*WorkingSpline, WorkingSpline->GetNumberOfSplinePoints() - 1, *RenderAPI);
+		}
+
+		SplineUtil::FDrawSplineSettings DrawSettings;
+		DrawSettings.ScaleVisualizationWidth = Settings->FrameVisualizationWidth;
+		SplineUtil::DrawSpline(*WorkingSpline, *RenderAPI, DrawSettings);
+	}
+}
+
+void UDrawSplineTool::OnPropertyModified(UObject* PropertySet, FProperty* Property)
+{
+}
+
+bool UDrawSplineTool::CanAccept() const
+{
+	return WorkingSpline.IsValid() && WorkingSpline->GetNumberOfSplinePoints() > 0;
+}
+
+// To be called by builder
+void UDrawSplineTool::SetSelectedActor(AActor* Actor)
+{
+	SelectedActor = Actor;
+}
+void UDrawSplineTool::SetWorld(UWorld* World)
+{
+	TargetWorld = World;
+}
+
+/// Tool builder:
+
+bool UDrawSplineToolBuilder::CanBuildTool(const FToolBuilderState& SceneState) const
+{
+	return true;
+}
+
+UInteractiveTool* UDrawSplineToolBuilder::BuildTool(const FToolBuilderState& SceneState) const
+{
+	UDrawSplineTool* NewTool = NewObject<UDrawSplineTool>(SceneState.ToolManager);
+	NewTool->SetWorld(SceneState.World);
+
+	// May be null
+	NewTool->SetSelectedActor(ToolBuilderUtil::FindFirstActor(SceneState, [](AActor*) { return true; }));
+
+	return NewTool;
+}
+
+
+void UDrawSplineTool::FSplineChange::Apply(UObject* Object)
+{
+	UDrawSplineTool* Tool = Cast<UDrawSplineTool>(Object);
+	if (!ensure(Tool))
+	{
+		return;
+	}
+	TWeakObjectPtr<USplineComponent> Spline = Tool->WorkingSpline;
+	if (!ensure(Spline.IsValid()))
+	{
+		return;
+	}
+
+	Apply(*Spline);
+
+	Tool->bNeedToRerunConstructionScript = true;
+}
+
+void UDrawSplineTool::FSplineChange::Revert(UObject* Object)
+{
+	UDrawSplineTool* Tool = Cast<UDrawSplineTool>(Object);
+	if (!ensure(Tool))
+	{
+		return;
+	}
+	TWeakObjectPtr<USplineComponent> Spline = Tool->WorkingSpline;
+	if (!ensure(Spline.IsValid()))
+	{
+		return;
+	}
+
+	Revert(*Spline);
+
+	Tool->bNeedToRerunConstructionScript = true;
+}
+
+#undef LOCTEXT_NAMESPACE
diff --git a/Engine/Plugins/Experimental/MeshModelingToolsetExp/Source/MeshModelingToolsEditorOnlyExp/Public/DrawSplineTool.h b/Engine/Plugins/Experimental/MeshModelingToolsetExp/Source/MeshModelingToolsEditorOnlyExp/Public/DrawSplineTool.h
new file mode 100644
--- /dev/null
+++ b/Engine/Plugins/Experimental/MeshModelingToolsetExp/Source/MeshModelingToolsEditorOnlyExp/Public/DrawSplineTool.h
@@ -0,0 +1,274 @@
+// Copyright Epic Games, Inc. All Rights Reserved.
+
+#pragma once
+
+#include "CoreMinimal.h"
+
+#include "BaseBehaviors/BehaviorTargetInterfaces.h"
+#include "InteractiveTool.h"
+#include "InteractiveToolBuilder.h"
+#include "InteractiveToolChange.h"
+#include "ToolContextInterfaces.h" // FViewCameraState
+
+#include "DrawSplineTool.generated.h"
+
+class AActor;
+class APreviewGeometryActor;
+class UConstructionPlaneMechanic;
+class USingleClickOrDragInputBehavior;
+class USplineComponent;
+class UWorld;
+
+UENUM()
+enum class EDrawSplineDrawMode : uint8
+{
+	// Click to place a point and then drag to set its tangent. Clicking without
+	// dragging will create sharp corners.
+	TangentDrag,
+	// Click and drag new points, with the tangent set automatically
+	ClickAutoTangent,
+	// Drag to place multiple points, with spacing controlled by Min Point Spacing
+	FreeDraw,
+};
+
+UENUM()
+enum class EDrawSplineOutputMode : uint8
+{
+	// Create a new empty actor with the spline inside it 
+	EmptyActor,
+	// Attach the spline to an existing actor, or replace a spline inside that
+	// actor if Existing Spline Index To Replace is valid.
+	ExistingActor,
+	// Create the blueprint specified by Blueprint To Create, and either attach
+	// the spline to that, or replace an existing spline if Existing Spline Index 
+	// To Replace is valid.
+	CreateBlueprint,
+
+	//~ The original implementation of this option was too unwieldy, but there is some way to
+	//~ implement it using FComponentReference to make a nice spline component picker...
+	//SelectedComponent,
+};
+
+UENUM()
+enum class EDrawSplineUpVectorMode : uint8
+{
+	// Pick the first up vector based on the hit normal, and then align subsequent
+	// up vectors with the previous ones.
+	AlignToPrevious,
+	// Base the up vector off the hit normal.
+	UseHitNormal,
+};
+
+UCLASS()
+class MESHMODELINGTOOLSEDITORONLYEXP_API UDrawSplineToolProperties : public UInteractiveToolPropertySet
+{
+	GENERATED_BODY()
+
+public:
+	//~ This is TransientToolProperty because while it's useful to toggle after drawing a spline,
+	//~ it's confusing when you first start the tool and are drawing a spline.
+	/** 
+	 * Determines whether the created spline is a loop. This can be toggled using "Closed Loop" in
+	 * the detail panel after spline creation.
+	 */
+	UPROPERTY(EditAnywhere, Category = Spline, meta = (TransientToolProperty))
+	bool bLoop = false;
+
+	/**
+	 * Determines how the resulting spline is emitted on tool accept.
+	 */
+	UPROPERTY(EditAnywhere, Category = Spline)
+	EDrawSplineOutputMode OutputMode = EDrawSplineOutputMode::EmptyActor;
+
+	/**
+	 * Actor to attach to when Output Mode is "Existing Actor"
+	 */
+	UPROPERTY(EditAnywhere, Category = Spline, meta = (TransientToolProperty, 
+		EditCondition = "OutputMode == EDrawSplineOutputMode::ExistingActor", EditConditionHides))
+	TWeakObjectPtr<AActor> TargetActor = nullptr;
+
+	/**
+	 * Blueprint to create when Output Mode is "Create Blueprint"
+	 */
+	UPROPERTY(EditAnywhere, Category = Spline, meta = (
+		EditCondition = "OutputMode == EDrawSplineOutputMode::CreateBlueprint", EditConditionHides))
+	TWeakObjectPtr<UBlueprint> BlueprintToCreate;
+
+	/**
+	 * When attaching to an existing actor or creating a blueprint, controls whether the drawn
+	 * spline replaces one of the target's existing components or gets added as a new one (if
+	 * the index is out of bounds).
+	 */
+	UPROPERTY(EditAnywhere, Category = Spline, meta = (ClampMin = -1, UIMax = 10, EditConditionHides,
+		EditCondition = "OutputMode == EDrawSplineOutputMode::ExistingActor || OutputMode == EDrawSplineOutputMode::CreateBlueprint"))
+	int32 ExistingSplineIndexToReplace = 0;
+
+	/**
+	 * How the spline is drawn in the tool.
+	 */
+	UPROPERTY(EditAnywhere, Category = Spline)
+	EDrawSplineDrawMode DrawMode = EDrawSplineDrawMode::TangentDrag;
+	
+	/** Point spacing when Draw Mode is "Free Draw" */
+	UPROPERTY(EditAnywhere, Category = Spline, meta = (ClampMin = 0, UIMax = 1000,
+		EditCondition = "DrawMode == EDrawSplineDrawMode::FreeDraw", EditConditionHides))
+	double MinPointSpacing = 200;
+
+	/**
+	 * When nonzero, allows a visualization of the rotation of the spline. Can be controlled
+	 * in the detail panel after creation via "Scale Visualization Width" property.
+	 */
+	UPROPERTY(EditAnywhere, Category = Spline, AdvancedDisplay, meta = (ClampMin = 0, UIMax = 100))
+	double FrameVisualizationWidth = 0;
+
+	/**
+	 * How the spline rotation is set. It is suggested to use a nonzero FrameVisualizationWidth to
+	 * see the effects.
+	 */
+	UPROPERTY(EditAnywhere, Category = Spline, AdvancedDisplay)
+	EDrawSplineUpVectorMode UpVectorMode = EDrawSplineUpVectorMode::UseHitNormal;
+
+	/**
+	 * When modifying existing actors, whether the result should be previewed using a copy
+	 * of that actor (rather than just drawing the spline by itself). Could be toggled off
+	 * if something about duplicating the given actor is problematic.
+	 */
+	UPROPERTY(EditAnywhere, Category = Spline, AdvancedDisplay, meta=(
+		EditCondition = "OutputMode != EDrawSplineOutputMode::EmptyActor"))
+	bool bPreviewUsingActorCopy = true;
+
+	UPROPERTY(EditAnywhere, Category = RaycastTargets, meta = (DisplayName = "World"))
+	bool bHitWorld = true;
+
+	UPROPERTY(EditAnywhere, Category = RaycastTargets, meta = (DisplayName = "Custom Plane"))
+	bool bHitCustomPlane = false;
+
+	UPROPERTY(EditAnywhere, Category = RaycastTargets, meta = (DisplayName = "Ground Planes"))
+	bool bHitGroundPlanes = true;
+
+	/**
+	 * If modifying a blueprint actor, whether to run the construction script while dragging
+	 * or only at the end of a drag. Can be toggled off for expensive construction scripts.
+	 */
+	UPROPERTY(EditAnywhere, Category = Spline, AdvancedDisplay, meta = (
+		EditCondition = "OutputMode != EDrawSplineOutputMode::EmptyActor"))
+	bool bRerunConstructionScriptOnDrag = true;
+};
+
+UCLASS()
+class MESHMODELINGTOOLSEDITORONLYEXP_API UDrawSplineTool : public UInteractiveTool
+	, public IClickBehaviorTarget
+	, public IClickDragBehaviorTarget
+{
+	GENERATED_BODY()
+
+public:
+
+	virtual void SetSelectedActor(AActor* Actor);
+
+	virtual void SetWorld(UWorld* World);
+	virtual UWorld* GetTargetWorld() { return TargetWorld.Get(); }
+
+	// UInteractiveTool
+	virtual void Setup() override;
+	virtual void Shutdown(EToolShutdownType ShutdownType) override;
+	virtual void OnTick(float DeltaTime) override;
+	virtual bool HasCancel() const override { return true; }
+	virtual bool HasAccept() const override { return true; }
+	virtual bool CanAccept() const override;
+	virtual void Render(IToolsContextRenderAPI* RenderAPI) override;
+	virtual void OnPropertyModified(UObject* PropertySet, FProperty* Property) override;
+
+	// IClickBehaviorTarget
+	virtual FInputRayHit IsHitByClick(const FInputDeviceRay& ClickPos);
+	virtual void OnClicked(const FInputDeviceRay& ClickPos);
+
+	// IClickDragBehaviorTarget
+	virtual FInputRayHit CanBeginClickDragSequence(const FInputDeviceRay& PressPos);
+	virtual void OnClickPress(const FInputDeviceRay& PressPos);
+	virtual void OnClickDrag(const FInputDeviceRay& DragPos);
+	virtual void OnClickRelease(const FInputDeviceRay& ReleasePos);
+	virtual void OnTerminateDragSequence();
+protected:
+
+	UPROPERTY()
+	TObjectPtr<UDrawSplineToolProperties> Settings = nullptr;
+
+	UPROPERTY()
+	TObjectPtr<USingleClickOrDragInputBehavior> ClickOrDragBehavior = nullptr;
+
+	UPROPERTY()
+	TObjectPtr<UConstructionPlaneMechanic> PlaneMechanic = nullptr;
+
+	TWeakObjectPtr<UWorld> TargetWorld = nullptr;
+
+	// This is only used to initialize TargetActor in the settings object
+	TWeakObjectPtr<AActor> SelectedActor = nullptr;
+
+	// PreviewRootActor either holds WorkingSpline inside it directly, or has some preview actor
+	// attached to it (so that the preview actor is hidden from outliner, like APreviewGeometryActor is).
+	UPROPERTY()
+	TObjectPtr<APreviewGeometryActor> PreviewRootActor = nullptr;
+
+	// The preview actor may be a duplicate of some target blueprint actor so that we can see the
+	// effects of the drawn spline immediately
+	UPROPERTY()
+	TObjectPtr<AActor> PreviewActor = nullptr;
+
+	// Used for recapturing the spline when rerunning construction scripts
+	int32 SplineRecaptureIndex = 0;
+
+	// This is the spline we add points to. It points to a component nested somewhere under 
+	// PreviewRootActor.
+	TWeakObjectPtr<USplineComponent> WorkingSpline = nullptr;
+
+	bool bDrawTangentForLastPoint = false;
+	bool bFreeDrawPlacedPreviewPoint = false;
+	int32 FreeDrawStrokeStartIndex = 0;
+
+	bool Raycast(const FRay& WorldRay, FVector3d& HitLocationOut, FVector3d& HitNormalOut, double& HitTOut);
+	void AddSplinePoint(const FVector3d& HitLocation, const FVector3d& UpVector);
+	FVector3d GetUpVectorToUse(const FVector3d& HitLocation, const FVector3d& HitNormal, int32 NumSplinePointsBeforehand);
+
+	void TransitionOutputMode();
+	void GenerateAsset();
+	
+	// Used to restore visibility of previous actor when switching to a different one
+	UPROPERTY()
+	TObjectPtr<AActor> PreviousTargetActor = nullptr;
+	bool PreviousTargetActorVisibility = true;
+	// Used to restore visibility of previous spline when switching to a different one
+	TWeakObjectPtr<USplineComponent> HiddenSpline = nullptr;
+	bool bPreviousSplineVisibility = true;
+
+	bool bNeedToRerunConstructionScript = false;
+
+	FViewCameraState CameraState;
+
+public:
+	// Helper class for making undo/redo transactions, to avoid friending all the variations.
+	class FSplineChange : public FToolCommandChange
+	{
+	public:
+		// These pass the working spline to the overloads below
+		virtual void Apply(UObject* Object) override;
+		virtual void Revert(UObject* Object) override;
+
+	protected:
+		virtual void Apply(USplineComponent& Spline) = 0;
+		virtual void Revert(USplineComponent& Spline) = 0;
+	};
+};
+
+
+
+UCLASS(Transient)
+class MESHMODELINGTOOLSEDITORONLYEXP_API UDrawSplineToolBuilder : public UInteractiveToolBuilder
+{
+	GENERATED_BODY()
+
+public:
+	virtual bool CanBuildTool(const FToolBuilderState& SceneState) const override;
+
+	virtual UInteractiveTool* BuildTool(const FToolBuilderState& SceneState) const override;
+};
diff --git a/Engine/Plugins/Runtime/MeshModelingToolset/Source/ModelingComponents/ModelingComponents.Build.cs b/Engine/Plugins/Runtime/MeshModelingToolset/Source/ModelingComponents/ModelingComponents.Build.cs
--- a/Engine/Plugins/Runtime/MeshModelingToolset/Source/ModelingComponents/ModelingComponents.Build.cs
+++ b/Engine/Plugins/Runtime/MeshModelingToolset/Source/ModelingComponents/ModelingComponents.Build.cs
@@ -34,6 +34,7 @@ public class ModelingComponents : ModuleRules
 				"RHI",
 				"ImageWriteQueue",
 				"SkeletalMeshDescription", // FSkeletalMeshAttributes::DefaultSkinWeightProfileName
+				"SlateCore",
 				// ... add private dependencies that you statically link with here ...
 			}
 		);
diff --git a/Engine/Plugins/Runtime/MeshModelingToolset/Source/ModelingComponents/Private/SplineUtil.cpp b/Engine/Plugins/Runtime/MeshModelingToolset/Source/ModelingComponents/Private/SplineUtil.cpp
new file mode 100644
--- /dev/null
+++ b/Engine/Plugins/Runtime/MeshModelingToolset/Source/ModelingComponents/Private/SplineUtil.cpp
@@ -0,0 +1,175 @@
+// Copyright Epic Games, Inc. All Rights Reserved.
+
+#include "SplineUtil.h"
+
+#include "Components/SplineComponent.h"
+#include "SceneManagement.h" // FPrimitiveDrawInterface
+#include "Styling/StyleColors.h"
+#include "ToolContextInterfaces.h" // IToolsContextRenderAPI
+
+UE::Geometry::SplineUtil::FDrawSplineSettings::FDrawSplineSettings()
+	: RegularColor(FStyleColors::White.GetSpecifiedColor().ToFColor(true))
+	, SelectedColor(FStyleColors::AccentOrange.GetSpecifiedColor().ToFColor(true))
+{
+}
+
+// Mostly copied from the editor-only FSplineComponentVisualizer
+void UE::Geometry::SplineUtil::DrawSpline(const USplineComponent& SplineComp, IToolsContextRenderAPI& RenderAPI, const FDrawSplineSettings& Settings)
+{
+	const FSceneView* View = RenderAPI.GetSceneView();
+	auto GetDashSize = [View](const FVector& Start, const FVector& End, float Scale) -> double
+	{
+		const double StartW = View->WorldToScreen(Start).W;
+		const double EndW = View->WorldToScreen(End).W;
+
+		const double WLimit = 10.0f;
+		if (StartW > WLimit || EndW > WLimit)
+		{
+			return FMath::Max(StartW, EndW) * Scale;
+		}
+
+		return 0;
+	};
+
+	FPrimitiveDrawInterface* PDI = RenderAPI.GetPrimitiveDrawInterface();
+
+	const FInterpCurveVector& SplineInfo = SplineComp.GetSplinePointsPosition();
+
+	const float GrabHandleSize = 10.0f;
+
+	const bool bShouldVisualizeScale = Settings.ScaleVisualizationWidth > 0;
+	const float DefaultScale = Settings.ScaleVisualizationWidth;
+
+	FVector OldKeyPos(0);
+	FVector OldKeyRightVector(0);
+	FVector OldKeyScale(0);
+
+	const int32 NumPoints = SplineInfo.Points.Num();
+	const int32 NumSegments = SplineInfo.bIsLooped ? NumPoints : NumPoints - 1;
+	for (int32 KeyIdx = 0; KeyIdx < NumSegments + 1; KeyIdx++)
+	{
+		const FVector NewKeyPos = SplineComp.GetLocationAtSplinePoint(KeyIdx, ESplineCoordinateSpace::World);
+		const FVector NewKeyRightVector = SplineComp.GetRightVectorAtSplinePoint(KeyIdx, ESplineCoordinateSpace::World);
+		const FVector NewKeyUpVector = SplineComp.GetUpVectorAtSplinePoint(KeyIdx, ESplineCoordinateSpace::World);
+		const FVector NewKeyScale = SplineComp.GetScaleAtSplinePoint(KeyIdx) * DefaultScale;
+
+		const FColor KeyColor = (Settings.SelectedKeys && Settings.SelectedKeys->Contains(KeyIdx)) ? Settings.SelectedColor
+			: Settings.RegularColor;
+
+		// Draw the keypoint and up/right vectors
+		if (KeyIdx < NumPoints)
+		{
+			if (bShouldVisualizeScale)
+			{
+				PDI->DrawLine(NewKeyPos, NewKeyPos - NewKeyRightVector * NewKeyScale.Y, KeyColor, SDPG_Foreground);
+				PDI->DrawLine(NewKeyPos, NewKeyPos + NewKeyRightVector * NewKeyScale.Y, KeyColor, SDPG_Foreground);
+				PDI->DrawLine(NewKeyPos, NewKeyPos + NewKeyUpVector * NewKeyScale.Z, KeyColor, SDPG_Foreground);
+
+				const int32 ArcPoints = 20;
+				FVector OldArcPos = NewKeyPos + NewKeyRightVector * NewKeyScale.Y;
+				for (int32 ArcIndex = 1; ArcIndex <= ArcPoints; ArcIndex++)
+				{
+					float Sin;
+					float Cos;
+					FMath::SinCos(&Sin, &Cos, ArcIndex * PI / ArcPoints);
+					const FVector NewArcPos = NewKeyPos + Cos * NewKeyRightVector * NewKeyScale.Y + Sin * NewKeyUpVector * NewKeyScale.Z;
+					PDI->DrawLine(OldArcPos, NewArcPos, KeyColor, SDPG_Foreground);
+					OldArcPos = NewArcPos;
+				}
+			}
+
+			PDI->DrawPoint(NewKeyPos, KeyColor, GrabHandleSize, SDPG_Foreground);
+		}
+
+		// If not the first keypoint, draw a line to the previous keypoint.
+		if (KeyIdx > 0)
+		{
+			const FColor LineColor = Settings.RegularColor;
+
+			// For constant interpolation - don't draw ticks - just draw dotted line.
+			if (SplineInfo.Points[KeyIdx - 1].InterpMode == CIM_Constant)
+			{
+				const double DashSize = GetDashSize(OldKeyPos, NewKeyPos, 0.03f);
+				if (DashSize > 0.0f)
+				{
+					DrawDashedLine(PDI, OldKeyPos, NewKeyPos, LineColor, DashSize, SDPG_World);
+				}
+			}
+			else
+			{
+				// Determine the colors to use
+				const bool bKeyIdxLooped = (SplineInfo.bIsLooped && KeyIdx == NumPoints);
+				const int32 BeginIdx = bKeyIdxLooped ? 0 : KeyIdx;
+				const int32 EndIdx = KeyIdx - 1;
+				const bool bBeginSelected = Settings.SelectedKeys && Settings.SelectedKeys->Contains(BeginIdx);
+				const bool bEndSelected = Settings.SelectedKeys && Settings.SelectedKeys->Contains(BeginIdx);
+				const FColor BeginColor = (bBeginSelected) ? Settings.SelectedColor : Settings.RegularColor;
+				const FColor EndColor = (bEndSelected) ? Settings.SelectedColor : Settings.RegularColor;
+
+				// Find position on first keyframe.
+				FVector OldPos = OldKeyPos;
+				FVector OldRightVector = OldKeyRightVector;
+				FVector OldScale = OldKeyScale;
+
+				// Then draw a line for each substep.
+				constexpr int32 NumSteps = 20;
+				constexpr float PartialGradientProportion = 0.75f;
+				constexpr int32 PartialNumSteps = (int32)(NumSteps * PartialGradientProportion);
+				const float SegmentLineThickness = 0;
+
+				for (int32 StepIdx = 1; StepIdx <= NumSteps; StepIdx++)
+				{
+					const float StepRatio = StepIdx / static_cast<float>(NumSteps);
+					const float Key = EndIdx + StepRatio;
+					const FVector NewPos = SplineComp.GetLocationAtSplineInputKey(Key, ESplineCoordinateSpace::World);
+					const FVector NewRightVector = SplineComp.GetRightVectorAtSplineInputKey(Key, ESplineCoordinateSpace::World);
+					const FVector NewScale = SplineComp.GetScaleAtSplineInputKey(Key) * DefaultScale;
+
+					// creates a gradient that starts partway through the selection
+					FColor StepColor;
+					if (bBeginSelected == bEndSelected)
+					{
+						StepColor = BeginColor;
+					}
+					else if (bBeginSelected && StepIdx > (NumSteps - PartialNumSteps))
+					{
+						const float LerpRatio = (1.0f - StepRatio) / PartialGradientProportion;
+						StepColor = FMath::Lerp(BeginColor.ReinterpretAsLinear(), EndColor.ReinterpretAsLinear(), LerpRatio).ToFColor(false);
+					}
+					else if (bEndSelected && StepIdx <= PartialNumSteps)
+					{
+						const float LerpRatio = 1.0f - (StepRatio / PartialGradientProportion);
+						StepColor = FMath::Lerp(BeginColor.ReinterpretAsLinear(), EndColor.ReinterpretAsLinear(), LerpRatio).ToFColor(false);
+					}
+					else
+					{
+						StepColor = Settings.RegularColor; // unselected
+					}
+
+					PDI->DrawLine(OldPos, NewPos, StepColor, SDPG_Foreground, SegmentLineThickness);
+					if (bShouldVisualizeScale)
+					{
+						PDI->DrawLine(OldPos - OldRightVector * OldScale.Y, NewPos - NewRightVector * NewScale.Y, LineColor, SDPG_Foreground);
+						PDI->DrawLine(OldPos + OldRightVector * OldScale.Y, NewPos + NewRightVector * NewScale.Y, LineColor, SDPG_Foreground);
+
+						constexpr bool bVisualizeSplineInterpolatedVectors = false;
+						if (bVisualizeSplineInterpolatedVectors)
+						{
+							const FVector NewUpVector = SplineComp.GetUpVectorAtSplineInputKey(Key, ESplineCoordinateSpace::World);
+							PDI->DrawLine(NewPos, NewPos + NewUpVector * Settings.ScaleVisualizationWidth * 0.5f, LineColor, SDPG_Foreground);
+							PDI->DrawLine(NewPos, NewPos + NewRightVector * Settings.ScaleVisualizationWidth * 0.5f, LineColor, SDPG_Foreground);
+						}
+					}
+
+					OldPos = NewPos;
+					OldRightVector = NewRightVector;
+					OldScale = NewScale;
+				}
+			}
+		}
+
+		OldKeyPos = NewKeyPos;
+		OldKeyRightVector = NewKeyRightVector;
+		OldKeyScale = NewKeyScale;
+	}
+}
diff --git a/Engine/Plugins/Runtime/MeshModelingToolset/Source/ModelingComponents/Public/SplineUtil.h b/Engine/Plugins/Runtime/MeshModelingToolset/Source/ModelingComponents/Public/SplineUtil.h
new file mode 100644
--- /dev/null
+++ b/Engine/Plugins/Runtime/MeshModelingToolset/Source/ModelingComponents/Public/SplineUtil.h
@@ -0,0 +1,36 @@
+// Copyright Epic Games, Inc. All Rights Reserved.
+
+#pragma once
+
+#include "CoreMinimal.h"
+
+class USplineComponent;
+class IToolsContextRenderAPI;
+
+namespace UE {
+namespace Geometry {
+
+namespace SplineUtil {
+
+	// Used for DrawSpline()
+	struct MODELINGCOMPONENTS_API FDrawSplineSettings
+	{
+		FDrawSplineSettings();
+
+		// Defaults to FStyleColors::White
+		FColor RegularColor;
+		// Defaults to FStyleColors::AccentOrange
+		FColor SelectedColor;
+		// If non-positive, the spline is drawn just as points and curves. If positive, the
+		// orientation and scale are visualized with the given base width.
+		double ScaleVisualizationWidth = 0;
+
+		// Keys to use SelectedColor for
+		TSet<int32>* SelectedKeys = nullptr;
+	};
+
+	MODELINGCOMPONENTS_API void DrawSpline(const USplineComponent& SplineComp, IToolsContextRenderAPI& RenderAPI, const FDrawSplineSettings& Settings);
+
+}
+}
+}
\ No newline at end of file
