MeshModelingTools: Fix nonunity build errors from missing includes for Draw Spline Tool.

diff --git a/Engine/Plugins/Experimental/MeshModelingToolsetExp/Source/MeshModelingToolsEditorOnlyExp/Private/DrawSplineTool.cpp b/Engine/Plugins/Experimental/MeshModelingToolsetExp/Source/MeshModelingToolsEditorOnlyExp/Private/DrawSplineTool.cpp
--- a/Engine/Plugins/Experimental/MeshModelingToolsetExp/Source/MeshModelingToolsEditorOnlyExp/Private/DrawSplineTool.cpp
+++ b/Engine/Plugins/Experimental/MeshModelingToolsetExp/Source/MeshModelingToolsEditorOnlyExp/Private/DrawSplineTool.cpp
@@ -8,16 +8,18 @@
 #include "BaseGizmos/GizmoMath.h"
 #include "Components/SplineComponent.h"
 #include "CoreGlobals.h" // GUndo
-#include "Kismet2/ComponentEditorUtils.h" // GenerateValidVariableName
 #include "Drawing/PreviewGeometryActor.h"
 #include "Editor/UnrealEdEngine.h" // DuplicateActors
+#include "Engine/Blueprint.h"
 #include "Engine/World.h" 
 #include "GameFramework/Actor.h"
 #include "InputState.h" // FInputDeviceRay
 #include "InteractiveToolManager.h"
+#include "Kismet2/ComponentEditorUtils.h" // GenerateValidVariableName
 #include "Mechanics/ConstructionPlaneMechanic.h"
 #include "SceneManagement.h" // FPrimitiveDrawInterface
 #include "Selection/ToolSelectionUtil.h"
+#include "Settings/LevelEditorViewportSettings.h"
 #include "SplineUtil.h"
 #include "ToolBuilderUtil.h"
 #include "UnrealEdGlobals.h" // GUnrealEd
diff --git a/Engine/Plugins/Runtime/MeshModelingToolset/Source/ModelingComponents/Private/SplineUtil.cpp b/Engine/Plugins/Runtime/MeshModelingToolset/Source/ModelingComponents/Private/SplineUtil.cpp
--- a/Engine/Plugins/Runtime/MeshModelingToolset/Source/ModelingComponents/Private/SplineUtil.cpp
+++ b/Engine/Plugins/Runtime/MeshModelingToolset/Source/ModelingComponents/Private/SplineUtil.cpp
@@ -4,6 +4,7 @@
 
 #include "Components/SplineComponent.h"
 #include "SceneManagement.h" // FPrimitiveDrawInterface
+#include "SceneView.h"
 #include "Styling/StyleColors.h"
 #include "ToolContextInterfaces.h" // IToolsContextRenderAPI
 
