MeshModelingTools: Fix incorrect edit condition in Mesh To Collision tool introduced in 25180815

diff --git a/Engine/Plugins/Experimental/MeshModelingToolsetExp/Source/MeshModelingToolsExp/Public/Physics/SetCollisionGeometryTool.h b/Engine/Plugins/Experimental/MeshModelingToolsetExp/Source/MeshModelingToolsExp/Public/Physics/SetCollisionGeometryTool.h
--- a/Engine/Plugins/Experimental/MeshModelingToolsetExp/Source/MeshModelingToolsExp/Public/Physics/SetCollisionGeometryTool.h
+++ b/Engine/Plugins/Experimental/MeshModelingToolsetExp/Source/MeshModelingToolsExp/Public/Physics/SetCollisionGeometryTool.h
@@ -100,7 +100,7 @@ public:
 	bool bUseWorldSpace = false;
 
 	UPROPERTY(EditAnywhere, Category = Options, meta = (
-		EditCondition = "bUsingMultipleInputs && GeometryType != ECollisionGeometryType::Empty && GeometryType != ECollisionGeometryType::CopyFromInputs"))
+		EditCondition = "GeometryType != ECollisionGeometryType::Empty && GeometryType != ECollisionGeometryType::CopyFromInputs"))
 	ESetCollisionGeometryInputMode InputMode = ESetCollisionGeometryInputMode::PerInputObject;
 
 	UPROPERTY(EditAnywhere, Category = Options)
