Refactor collision geometry visualization

This change was motivated by future work to implement a minimal boilerplate system for visualizing geometry selections in tools, which will work similarly to this

diff --git a/Engine/Plugins/Experimental/MeshLODToolset/Source/MeshLODToolset/Private/Tools/GenerateStaticMeshLODAssetTool.cpp b/Engine/Plugins/Experimental/MeshLODToolset/Source/MeshLODToolset/Private/Tools/GenerateStaticMeshLODAssetTool.cpp
--- a/Engine/Plugins/Experimental/MeshLODToolset/Source/MeshLODToolset/Private/Tools/GenerateStaticMeshLODAssetTool.cpp
+++ b/Engine/Plugins/Experimental/MeshLODToolset/Source/MeshLODToolset/Private/Tools/GenerateStaticMeshLODAssetTool.cpp
@@ -408,9 +408,7 @@ void UGenerateStaticMeshLODAssetTool::Setup()
 		FPhysicsDataCollection PhysicsData;
 		PhysicsData.Geometry = GenerateLODOp->ResultCollision;
 		PhysicsData.CopyGeometryToAggregate();
-		UE::PhysicsTools::InitializePreviewGeometryLines(PhysicsData,
-														 CollisionPreview,
-														 CollisionVizSettings->Color, CollisionVizSettings->LineThickness, 0.0f, 16, CollisionVizSettings->bRandomColors);
+		UE::PhysicsTools::InitializeCollisionGeometryVisualization(CollisionPreview, CollisionVizSettings, PhysicsData, 0.f, 16);
 
 		// Must happen on main thread, and GenerateProcess might be in use by an Op somewhere else
 		GenerateProcess->GraphEvalCriticalSection.Lock();
@@ -436,10 +434,7 @@ void UGenerateStaticMeshLODAssetTool::Setup()
 	CollisionVizSettings = NewObject<UCollisionGeometryVisualizationProperties>(this);
 	CollisionVizSettings->RestoreProperties(this);
 	AddToolPropertySource(CollisionVizSettings);
-	CollisionVizSettings->WatchProperty(CollisionVizSettings->LineThickness, [this](float NewValue) { bCollisionVisualizationDirty = true; });
-	CollisionVizSettings->WatchProperty(CollisionVizSettings->Color, [this](FColor NewValue) { bCollisionVisualizationDirty = true; });
-	CollisionVizSettings->WatchProperty(CollisionVizSettings->bRandomColors, [this](bool bNewValue) { bCollisionVisualizationDirty = true; });
-	CollisionVizSettings->WatchProperty(CollisionVizSettings->bShowHidden, [this](bool bNewValue) { bCollisionVisualizationDirty = true; });
+	CollisionVizSettings->Initialize(this);
 
 	CollisionPreview = NewObject<UPreviewGeometry>(this);
 	CollisionPreview->CreateInWorld(GetTargetWorld(), (FTransform)PreviewTransform);
@@ -609,31 +604,11 @@ void UGenerateStaticMeshLODAssetTool::OnTick(float DeltaTime)
 		PreviewWithBackgroundCompute->Tick(DeltaTime);
 	}
 
-	if (bCollisionVisualizationDirty)
-	{
-		UpdateCollisionVisualization();
-		bCollisionVisualizationDirty = false;
-	}
+	UE::PhysicsTools::UpdateCollisionGeometryVisualization(CollisionPreview, CollisionVizSettings);
 }
 
 
 
-void UGenerateStaticMeshLODAssetTool::UpdateCollisionVisualization()
-{
-	float UseThickness = CollisionVizSettings->LineThickness;
-	FColor UseColor = CollisionVizSettings->Color;
-	LineMaterial = ToolSetupUtil::GetDefaultLineComponentMaterial(GetToolManager(), !CollisionVizSettings->bShowHidden);
-
-	int32 ColorIdx = 0;
-	CollisionPreview->UpdateAllLineSets([&](ULineSetComponent* LineSet)
-	{
-		LineSet->SetAllLinesThickness(UseThickness);
-		LineSet->SetAllLinesColor(CollisionVizSettings->bRandomColors ? LinearColors::SelectFColor(ColorIdx++) : UseColor);
-	});
-	CollisionPreview->SetAllLineSetsMaterial(LineMaterial);
-}
-
-
 void UGenerateStaticMeshLODAssetTool::CreateNewAsset()
 {
 	check(PreviewWithBackgroundCompute->HaveValidResult());
diff --git a/Engine/Plugins/Experimental/MeshLODToolset/Source/MeshLODToolset/Public/Tools/GenerateStaticMeshLODAssetTool.h b/Engine/Plugins/Experimental/MeshLODToolset/Source/MeshLODToolset/Public/Tools/GenerateStaticMeshLODAssetTool.h
--- a/Engine/Plugins/Experimental/MeshLODToolset/Source/MeshLODToolset/Public/Tools/GenerateStaticMeshLODAssetTool.h
+++ b/Engine/Plugins/Experimental/MeshLODToolset/Source/MeshLODToolset/Public/Tools/GenerateStaticMeshLODAssetTool.h
@@ -320,9 +320,6 @@ protected:
 	UPROPERTY()
 	TObjectPtr<UPhysicsObjectToolPropertySet> ObjectData = nullptr;
 
-	UPROPERTY()
-	TObjectPtr<UMaterialInterface> LineMaterial = nullptr;
-
 	UPROPERTY()
 	TObjectPtr<UPreviewGeometry> CollisionPreview;
 
@@ -337,10 +334,6 @@ protected:
 
 	void OnSettingsModified();
 
-	bool bCollisionVisualizationDirty = false;
-	void UpdateCollisionVisualization();
-
-
 	void CreateNewAsset();
 	void UpdateExistingAsset();
 
diff --git a/Engine/Plugins/Experimental/MeshModelingToolsetExp/Source/MeshModelingToolsExp/Private/Physics/CollisionGeometryVisualization.cpp b/Engine/Plugins/Experimental/MeshModelingToolsetExp/Source/MeshModelingToolsExp/Private/Physics/CollisionGeometryVisualization.cpp
--- a/Engine/Plugins/Experimental/MeshModelingToolsetExp/Source/MeshModelingToolsExp/Private/Physics/CollisionGeometryVisualization.cpp
+++ b/Engine/Plugins/Experimental/MeshModelingToolsetExp/Source/MeshModelingToolsExp/Private/Physics/CollisionGeometryVisualization.cpp
@@ -1,31 +1,31 @@
 // Copyright Epic Games, Inc. All Rights Reserved.
 
 #include "Physics/CollisionGeometryVisualization.h"
+#include "Physics/PhysicsDataCollection.h"
+#include "Physics/CollisionPropertySets.h"
 #include "Generators/LineSegmentGenerators.h"
-#include "Util/ColorConstants.h"
-#include "Chaos/Levelset.h"
+#include "Drawing/PreviewGeometryActor.h"
 
 using namespace UE::Geometry;
 
-void UE::PhysicsTools::InitializePreviewGeometryLines(const FPhysicsDataCollection& PhysicsData, UPreviewGeometry* PreviewGeom,
-	const FColor& LineColor, float LineThickness, float DepthBias, int32 CircleStepResolution, bool bRandomColors)
+namespace
 {
+
+void InitializePreviewGeometryLines(
+	const FPhysicsDataCollection& PhysicsData,
+	UPreviewGeometry* PreviewGeom,
+	UMaterialInterface* LineMaterial,
+	TFunctionRef<FColor(int32 LineSetIndex)> LineSetIndexToColorFunc,
+	float LineThickness,
+	float DepthBias,
+	int32 CircleStepResolution,
+	int32 FirstLineSetIndex)
+{
+	check(PreviewGeom);
+	check(LineMaterial);
+
 	int32 CircleSteps = FMath::Max(4, CircleStepResolution);
-	FColor SphereColor = LineColor;
-	FColor BoxColor = LineColor;
-	FColor ConvexColor = LineColor;
-	FColor LevelSetColor = LineColor;
-	FColor CapsuleColor = LineColor;
-	int32 OverrideColorIdx = 0;
-
-	auto GetNextColor = [&OverrideColorIdx, bRandomColors](FColor DefaultColor)
-	{
-		if (bRandomColors)
-		{
-			return LinearColors::SelectFColor(OverrideColorIdx++);
-		}
-		return DefaultColor;
-	};
+	int32 LineSetIndex = FirstLineSetIndex;
 
 	const FKAggregateGeom& AggGeom = PhysicsData.AggGeom;
 
@@ -34,7 +34,7 @@ void UE::PhysicsTools::InitializePreviewGeometryLines(const FPhysicsDataCollecti
 	{
 		PreviewGeom->CreateOrUpdateLineSet(FString::Printf(TEXT("Spheres %d"), Index), 1, [&](int32 UnusedIndex, TArray<FRenderableLine>& LinesOut)
 		{
-			FColor Color = GetNextColor(SphereColor);
+			FColor Color = LineSetIndexToColorFunc(LineSetIndex++);
 
 			const FKSphereElem& Sphere = AggGeom.SphereElems[Index];
 			FTransform ElemTransform = Sphere.GetTransform();
@@ -56,7 +56,7 @@ void UE::PhysicsTools::InitializePreviewGeometryLines(const FPhysicsDataCollecti
 	{
 		PreviewGeom->CreateOrUpdateLineSet(FString::Printf(TEXT("Boxes %d"), Index), 1, [&](int32 UnusedIndex, TArray<FRenderableLine>& LinesOut)
 		{
-			FColor Color = GetNextColor(BoxColor);
+			FColor Color = LineSetIndexToColorFunc(LineSetIndex++);
 
 			const FKBoxElem& Box = AggGeom.BoxElems[Index];
 			FTransform ElemTransform = Box.GetTransform();
@@ -77,7 +77,7 @@ void UE::PhysicsTools::InitializePreviewGeometryLines(const FPhysicsDataCollecti
 	{
 		PreviewGeom->CreateOrUpdateLineSet(FString::Printf(TEXT("Capsules %d"), Index), 1, [&](int32 UnusedIndex, TArray<FRenderableLine>& LinesOut)
 		{
-			FColor Color = GetNextColor(CapsuleColor);
+			FColor Color = LineSetIndexToColorFunc(LineSetIndex++);
 
 			const FKSphylElem& Capsule = AggGeom.SphylElems[Index];
 			FTransform ElemTransform = Capsule.GetTransform();
@@ -125,7 +125,7 @@ void UE::PhysicsTools::InitializePreviewGeometryLines(const FPhysicsDataCollecti
 	{
 		PreviewGeom->CreateOrUpdateLineSet(FString::Printf(TEXT("Convex %d"), Index), 1, [&](int32 UnusedIndex, TArray<FRenderableLine>& LinesOut)
 		{
-			FColor Color = GetNextColor(ConvexColor);
+			FColor Color = LineSetIndexToColorFunc(LineSetIndex++);
 
 			const FKConvexElem& Convex = AggGeom.ConvexElems[Index];
 			FTransform ElemTransform = Convex.GetTransform();
@@ -143,13 +143,13 @@ void UE::PhysicsTools::InitializePreviewGeometryLines(const FPhysicsDataCollecti
 			}
 		});
 	}
-	
+
 	// for Level Sets draw the grid cells where phi < 0
 	for (int32 Index = 0; Index < AggGeom.LevelSetElems.Num(); Index++)
 	{
 		PreviewGeom->CreateOrUpdateLineSet(FString::Printf(TEXT("Level Set %d"), Index), 1, [&](int32 UnusedIndex, TArray<FRenderableLine>& LinesOut)
 		{
-			const FColor Color = GetNextColor(LevelSetColor);
+			FColor Color = LineSetIndexToColorFunc(LineSetIndex++);
 			const FKLevelSetElem& LevelSet = AggGeom.LevelSetElems[Index];
 			
 			FTransform ElemTransform = LevelSet.GetTransform();
@@ -183,4 +183,107 @@ void UE::PhysicsTools::InitializePreviewGeometryLines(const FPhysicsDataCollecti
 	// Unclear whether we actually use these in the Engine, for UBodySetup? Does not appear to be supported by UxX import system,
 	// and online documentation suggests they may only be supported for cloth?
 	ensure(AggGeom.TaperedCapsuleElems.Num() == 0);
-}
\ No newline at end of file
+
+	PreviewGeom->SetAllLineSetsMaterial(LineMaterial);
+}
+
+void UpdatePreviewGeometryLines(
+	UPreviewGeometry* PartialPreviewGeom,
+	UCollisionGeometryVisualizationProperties* Settings,
+	int32 FirstLineSetIndex = 0)
+{
+	check(PartialPreviewGeom);
+	check(Settings);
+
+	int32 LineSetIndex = FirstLineSetIndex;
+	PartialPreviewGeom->UpdateAllLineSets([&](ULineSetComponent* LineSet)
+	{
+		FColor LineColor = Settings->GetLineSetColor(LineSetIndex++);
+		LineSet->SetAllLinesThickness(Settings->LineThickness);
+		LineSet->SetAllLinesColor(LineColor);
+	});
+	PartialPreviewGeom->SetAllLineSetsMaterial(Settings->GetLineMaterial());
+}
+
+} // end namespace
+
+
+
+
+
+
+void UE::PhysicsTools::InitializeCollisionGeometryVisualization(
+	UPreviewGeometry* PreviewGeom,
+	UCollisionGeometryVisualizationProperties* Settings,
+	const FPhysicsDataCollection& PhysicsData,
+	float DepthBias,
+	int32 CircleStepResolution)
+{
+	check(PreviewGeom);
+	check(Settings);
+
+	InitializePreviewGeometryLines(
+		PhysicsData,
+		PreviewGeom,
+		Settings->GetLineMaterial(),
+		[&Settings](int LineSetIndex) { return Settings->GetLineSetColor(LineSetIndex); },
+		Settings->LineThickness,
+		DepthBias,
+		CircleStepResolution,
+		0);
+
+	Settings->bVisualizationDirty = false;
+}
+
+void UE::PhysicsTools::UpdateCollisionGeometryVisualization(
+	UPreviewGeometry* PreviewGeom,
+	UCollisionGeometryVisualizationProperties* Settings)
+{
+	check(PreviewGeom);
+	check(Settings);
+
+	if (Settings->bVisualizationDirty)
+	{
+		UpdatePreviewGeometryLines(PreviewGeom, Settings);
+		Settings->bVisualizationDirty = false;
+	}
+}
+
+
+
+
+
+
+void UE::PhysicsTools::PartiallyInitializeCollisionGeometryVisualization(
+	UPreviewGeometry* PreviewGeom,
+	UCollisionGeometryVisualizationProperties* Settings,
+	const FPhysicsDataCollection& PhysicsData,
+	int32 FirstLineSetIndex,
+	float DepthBias,
+	int32 CircleStepResolution)
+{
+	check(PreviewGeom);
+	check(Settings);
+
+	InitializePreviewGeometryLines(
+		PhysicsData,
+		PreviewGeom,
+		Settings->GetLineMaterial(),
+		[&Settings](int LineSetIndex) { return Settings->GetLineSetColor(LineSetIndex); },
+		Settings->LineThickness,
+		DepthBias,
+		CircleStepResolution,
+		FirstLineSetIndex);
+}
+
+void UE::PhysicsTools::PartiallyUpdateCollisionGeometryVisualization(
+	UPreviewGeometry* PartialPreviewGeom,
+	UCollisionGeometryVisualizationProperties* Settings,
+	int32 FirstLineSetIndex)
+{
+	check(PartialPreviewGeom);
+	check(Settings);
+
+	UpdatePreviewGeometryLines(PartialPreviewGeom, Settings, FirstLineSetIndex);
+}
+
diff --git a/Engine/Plugins/Experimental/MeshModelingToolsetExp/Source/MeshModelingToolsExp/Private/Physics/CollisionPropertySets.cpp b/Engine/Plugins/Experimental/MeshModelingToolsetExp/Source/MeshModelingToolsExp/Private/Physics/CollisionPropertySets.cpp
--- a/Engine/Plugins/Experimental/MeshModelingToolsetExp/Source/MeshModelingToolsExp/Private/Physics/CollisionPropertySets.cpp
+++ b/Engine/Plugins/Experimental/MeshModelingToolsetExp/Source/MeshModelingToolsExp/Private/Physics/CollisionPropertySets.cpp
@@ -6,9 +6,27 @@
 #include "Engine/World.h"
 #include "GameFramework/Actor.h"
 #include "Physics/PhysicsDataCollection.h"
+#include "Util/ColorConstants.h"
+#include "ToolSetupUtil.h"
 
 #include UE_INLINE_GENERATED_CPP_BY_NAME(CollisionPropertySets)
 
+void UCollisionGeometryVisualizationProperties::Initialize(UInteractiveTool* Tool)
+{
+	WatchProperty(LineThickness, [this](float NewValue) { bVisualizationDirty = true; });
+	WatchProperty(Color, [this](FColor NewValue) { bVisualizationDirty = true; });
+	WatchProperty(bRandomColors, [this](bool bNewValue) { bVisualizationDirty = true; });
+	WatchProperty(bShowHidden, [this](bool bNewValue) { bVisualizationDirty = true; });
+
+	LineMaterial = ToolSetupUtil::GetDefaultLineComponentMaterial(Tool->GetToolManager(), true);
+	LineMaterialShowingHidden = ToolSetupUtil::GetDefaultLineComponentMaterial(Tool->GetToolManager(), false);
+}
+
+FColor UCollisionGeometryVisualizationProperties::GetLineSetColor(int32 LineSetIndex) const
+{
+	return bRandomColors ? UE::Geometry::LinearColors::SelectFColor(LineSetIndex) : Color;
+}
+
 void UE::PhysicsTools::InitializePhysicsToolObjectPropertySet(const FPhysicsDataCollection* PhysicsData, UPhysicsObjectToolPropertySet* PropSet)
 {
 	check(PhysicsData);
diff --git a/Engine/Plugins/Experimental/MeshModelingToolsetExp/Source/MeshModelingToolsExp/Private/Physics/ExtractCollisionGeometryTool.cpp b/Engine/Plugins/Experimental/MeshModelingToolsetExp/Source/MeshModelingToolsExp/Private/Physics/ExtractCollisionGeometryTool.cpp
--- a/Engine/Plugins/Experimental/MeshModelingToolsetExp/Source/MeshModelingToolsExp/Private/Physics/ExtractCollisionGeometryTool.cpp
+++ b/Engine/Plugins/Experimental/MeshModelingToolsetExp/Source/MeshModelingToolsExp/Private/Physics/ExtractCollisionGeometryTool.cpp
@@ -92,10 +92,7 @@ void UExtractCollisionGeometryTool::Setup()
 	VizSettings = NewObject<UCollisionGeometryVisualizationProperties>(this);
 	VizSettings->RestoreProperties(this);
 	AddToolPropertySource(VizSettings);
-	VizSettings->WatchProperty(VizSettings->LineThickness, [this](float NewValue) { bVisualizationDirty = true; });
-	VizSettings->WatchProperty(VizSettings->Color, [this](FColor NewValue) { bVisualizationDirty = true; });
-	VizSettings->WatchProperty(VizSettings->bRandomColors, [this](bool bNewValue) { bVisualizationDirty = true; });
-	VizSettings->WatchProperty(VizSettings->bShowHidden, [this](bool bNewValue) { bVisualizationDirty = true; });
+	VizSettings->Initialize(this);
 
 	UBodySetup* BodySetup = UE::ToolTarget::GetPhysicsBodySetup(Target);
 	if (BodySetup)
@@ -108,9 +105,8 @@ void UExtractCollisionGeometryTool::Setup()
 		//PhysicsInfo->ExternalScale3D = TargetTransform.GetScale3D();
 		//TargetTransform.SetScale3D(FVector::OneVector);
 		PreviewElements->CreateInWorld(UE::ToolTarget::GetTargetActor(Target)->GetWorld(), TargetTransform);
-		
-		UE::PhysicsTools::InitializePreviewGeometryLines(*PhysicsInfo, PreviewElements,
-			VizSettings->Color, VizSettings->LineThickness, 0.0f, 16, VizSettings->bRandomColors);
+
+		UE::PhysicsTools::InitializeCollisionGeometryVisualization(PreviewElements, VizSettings, *PhysicsInfo);
 
 		ObjectProps = NewObject<UPhysicsObjectToolPropertySet>(this);
 		UE::PhysicsTools::InitializePhysicsToolObjectPropertySet(PhysicsInfo.Get(), ObjectProps);
@@ -216,28 +212,7 @@ void UExtractCollisionGeometryTool::OnTick(float DeltaTime)
 		bResultValid = true;
 	}
 
-	if (bVisualizationDirty)
-	{
-		UpdateVisualization();
-		bVisualizationDirty = false;
-	}
-}
-
-
-
-void UExtractCollisionGeometryTool::UpdateVisualization()
-{
-	float UseThickness = VizSettings->LineThickness;
-	FColor UseColor = VizSettings->Color;
-	int32 ColorIdx = 0;
-	PreviewElements->UpdateAllLineSets([&](ULineSetComponent* LineSet)
-		{
-			LineSet->SetAllLinesThickness(UseThickness);
-			LineSet->SetAllLinesColor(VizSettings->bRandomColors ? LinearColors::SelectFColor(ColorIdx++) : UseColor);
-		});
-	
-	UMaterialInterface* LineMaterial = ToolSetupUtil::GetDefaultLineComponentMaterial(GetToolManager(), !VizSettings->bShowHidden);
-	PreviewElements->SetAllLineSetsMaterial(LineMaterial);
+	UE::PhysicsTools::UpdateCollisionGeometryVisualization(PreviewElements, VizSettings);
 }
 
 
diff --git a/Engine/Plugins/Experimental/MeshModelingToolsetExp/Source/MeshModelingToolsExp/Private/Physics/PhysicsInspectorTool.cpp b/Engine/Plugins/Experimental/MeshModelingToolsetExp/Source/MeshModelingToolsExp/Private/Physics/PhysicsInspectorTool.cpp
--- a/Engine/Plugins/Experimental/MeshModelingToolsetExp/Source/MeshModelingToolsExp/Private/Physics/PhysicsInspectorTool.cpp
+++ b/Engine/Plugins/Experimental/MeshModelingToolsetExp/Source/MeshModelingToolsExp/Private/Physics/PhysicsInspectorTool.cpp
@@ -49,34 +49,37 @@ void UPhysicsInspectorTool::Setup()
 	VizSettings = NewObject<UCollisionGeometryVisualizationProperties>(this);
 	VizSettings->RestoreProperties(this);
 	AddToolPropertySource(VizSettings);
-	VizSettings->WatchProperty(VizSettings->LineThickness, [this](float NewValue) { bVisualizationDirty = true; });
-	VizSettings->WatchProperty(VizSettings->Color, [this](FColor NewValue) { bVisualizationDirty = true; });
-	VizSettings->WatchProperty(VizSettings->bRandomColors, [this](bool bNewValue) { bVisualizationDirty = true; });
-	VizSettings->WatchProperty(VizSettings->bShowHidden, [this](bool bNewValue) { bVisualizationDirty = true; });
 
-	for (int32 ComponentIdx = 0; ComponentIdx < Targets.Num(); ComponentIdx++)
+	// Intitialize the collision geometry visualization
 	{
-		UBodySetup* BodySetup = UE::ToolTarget::GetPhysicsBodySetup(Targets[ComponentIdx]);
-		if (BodySetup)
+		VizSettings->Initialize(this);
+		int32 FirstLineSetIndex = 0;
+		for (int32 ComponentIdx = 0; ComponentIdx < Targets.Num(); ComponentIdx++)
 		{
-			TSharedPtr<FPhysicsDataCollection> PhysicsData = MakeShared<FPhysicsDataCollection>();
-			PhysicsData->InitializeFromComponent( UE::ToolTarget::GetTargetComponent(Targets[ComponentIdx]), true);
-
-			PhysicsInfos.Add(PhysicsData);
-
-			UPreviewGeometry* PreviewGeom = NewObject<UPreviewGeometry>(this);
-			FTransform TargetTransform = (FTransform) UE::ToolTarget::GetLocalToWorldTransform(Targets[ComponentIdx]);
-			PhysicsData->ExternalScale3D = TargetTransform.GetScale3D();
-			TargetTransform.SetScale3D(FVector::OneVector);
-			PreviewGeom->CreateInWorld(UE::ToolTarget::GetTargetActor(Targets[ComponentIdx])->GetWorld(), TargetTransform);
-			PreviewElements.Add(PreviewGeom);
-
-			InitializeGeometry(*PhysicsData, PreviewGeom);
-
-			UPhysicsObjectToolPropertySet* ObjectProps = NewObject<UPhysicsObjectToolPropertySet>(this);
-			InitializeObjectProperties(*PhysicsData, ObjectProps);
-			AddToolPropertySource(ObjectProps);
+			UBodySetup* BodySetup = UE::ToolTarget::GetPhysicsBodySetup(Targets[ComponentIdx]);
+			if (BodySetup)
+			{
+				TSharedPtr<FPhysicsDataCollection> PhysicsData = MakeShared<FPhysicsDataCollection>();
+				PhysicsData->InitializeFromComponent( UE::ToolTarget::GetTargetComponent(Targets[ComponentIdx]), true);
+
+				PhysicsInfos.Add(PhysicsData);
+
+				UPreviewGeometry* PreviewGeom = NewObject<UPreviewGeometry>(this);
+				FTransform TargetTransform = (FTransform) UE::ToolTarget::GetLocalToWorldTransform(Targets[ComponentIdx]);
+				PhysicsData->ExternalScale3D = TargetTransform.GetScale3D();
+				TargetTransform.SetScale3D(FVector::OneVector);
+				PreviewGeom->CreateInWorld(UE::ToolTarget::GetTargetActor(Targets[ComponentIdx])->GetWorld(), TargetTransform);
+				PreviewElements.Add(PreviewGeom);
+
+				UE::PhysicsTools::PartiallyInitializeCollisionGeometryVisualization(PreviewGeom, VizSettings, *PhysicsData, FirstLineSetIndex);
+				FirstLineSetIndex += PreviewGeom->LineSets.Num();
+
+				UPhysicsObjectToolPropertySet* ObjectProps = NewObject<UPhysicsObjectToolPropertySet>(this);
+				UE::PhysicsTools::InitializePhysicsToolObjectPropertySet(PhysicsData.Get(), ObjectProps);
+				AddToolPropertySource(ObjectProps);
+			}
 		}
+		VizSettings->bVisualizationDirty = false;
 	}
 
 	SetToolDisplayName(LOCTEXT("ToolName", "Physics Inspector"));
@@ -102,46 +105,17 @@ void UPhysicsInspectorTool::OnShutdown(EToolShutdownType ShutdownType)
 
 void UPhysicsInspectorTool::OnTick(float DeltaTime)
 {
-	if (bVisualizationDirty)
-	{
-		UpdateVisualization();
-		bVisualizationDirty = false;
-	}
-}
-
-
-
-void UPhysicsInspectorTool::UpdateVisualization()
-{
-	float UseThickness = VizSettings->LineThickness;
-	FColor UseColor = VizSettings->Color;
-	LineMaterial = ToolSetupUtil::GetDefaultLineComponentMaterial(GetToolManager(), !VizSettings->bShowHidden);
-
-	int32 ColorIdx = 0;
-	for (UPreviewGeometry* Preview : PreviewElements)
+	if (VizSettings->bVisualizationDirty)
 	{
-		Preview->UpdateAllLineSets([&](ULineSetComponent* LineSet)
+		int32 FirstLineSetIndex= 0;
+		for (UPreviewGeometry* Preview : PreviewElements)
 		{
-			LineSet->SetAllLinesThickness(UseThickness);
-			LineSet->SetAllLinesColor(VizSettings->bRandomColors ? LinearColors::SelectFColor(ColorIdx++) : UseColor);
-		});
-		Preview->SetAllLineSetsMaterial(LineMaterial);
-	}
-}
-
-
-
-void UPhysicsInspectorTool::InitializeObjectProperties(const FPhysicsDataCollection& PhysicsData, UPhysicsObjectToolPropertySet* PropSet)
-{
-	UE::PhysicsTools::InitializePhysicsToolObjectPropertySet(&PhysicsData, PropSet);
-}
-
-
+			UE::PhysicsTools::PartiallyUpdateCollisionGeometryVisualization(Preview, VizSettings, FirstLineSetIndex);
+			FirstLineSetIndex += Preview->LineSets.Num();
+		}
 
-void UPhysicsInspectorTool::InitializeGeometry(const FPhysicsDataCollection& PhysicsData, UPreviewGeometry* PreviewGeom)
-{
-	UE::PhysicsTools::InitializePreviewGeometryLines(PhysicsData, PreviewGeom,
-		VizSettings->Color, VizSettings->LineThickness, 0.0f, 16, VizSettings->bRandomColors);
+		VizSettings->bVisualizationDirty = false;
+	}
 }
 
 
diff --git a/Engine/Plugins/Experimental/MeshModelingToolsetExp/Source/MeshModelingToolsExp/Private/Physics/SetCollisionGeometryTool.cpp b/Engine/Plugins/Experimental/MeshModelingToolsetExp/Source/MeshModelingToolsExp/Private/Physics/SetCollisionGeometryTool.cpp
--- a/Engine/Plugins/Experimental/MeshModelingToolsetExp/Source/MeshModelingToolsExp/Private/Physics/SetCollisionGeometryTool.cpp
+++ b/Engine/Plugins/Experimental/MeshModelingToolsetExp/Source/MeshModelingToolsExp/Private/Physics/SetCollisionGeometryTool.cpp
@@ -28,6 +28,7 @@
 #include "ToolTargetManager.h"
 #include "UObject/UObjectIterator.h"
 #include "ModelingToolTargetUtil.h"
+#include "Drawing/PreviewGeometryActor.h"
 
 #include UE_INLINE_GENERATED_CPP_BY_NAME(SetCollisionGeometryTool)
 
@@ -302,10 +303,7 @@ void USetCollisionGeometryTool::Setup()
 	VizSettings = NewObject<UCollisionGeometryVisualizationProperties>(this);
 	VizSettings->RestoreProperties(this);
 	AddToolPropertySource(VizSettings);
-	VizSettings->WatchProperty(VizSettings->LineThickness, [this](float NewValue) { bVisualizationDirty = true; });
-	VizSettings->WatchProperty(VizSettings->Color, [this](FColor NewValue) { bVisualizationDirty = true; });
-	VizSettings->WatchProperty(VizSettings->bRandomColors, [this](bool bNewValue) { bVisualizationDirty = true; });
-	VizSettings->WatchProperty(VizSettings->bShowHidden, [this](bool bNewValue) { bVisualizationDirty = true; });
+	VizSettings->Initialize(this);
 
 	// add option for collision properties
 	CollisionProps = NewObject<UPhysicsObjectToolPropertySet>(this);
@@ -521,12 +519,12 @@ void USetCollisionGeometryTool::OnTick(float DeltaTime)
 			{
 				GeneratedCollision = MakeShareable<FPhysicsDataCollection>(Result.Release());
 
-				bVisualizationDirty = true;
+				VizSettings->bVisualizationDirty = true;
 
 				// update visualization
 				PreviewGeom->RemoveAllLineSets();
-				UE::PhysicsTools::InitializePreviewGeometryLines(*GeneratedCollision, PreviewGeom,
-					VizSettings->Color, VizSettings->LineThickness, 0.0f, 16, VizSettings->bRandomColors);
+
+				UE::PhysicsTools::InitializeCollisionGeometryVisualization(PreviewGeom, VizSettings, *GeneratedCollision);
 
 				// update property set
 				CollisionProps->Reset();
@@ -535,11 +533,7 @@ void USetCollisionGeometryTool::OnTick(float DeltaTime)
 		}
 	}
 
-	if (bVisualizationDirty)
-	{
-		UpdateVisualization();
-		bVisualizationDirty = false;
-	}
+	UE::PhysicsTools::UpdateCollisionGeometryVisualization(PreviewGeom, VizSettings);
 }
 
 
@@ -607,24 +601,6 @@ void USetCollisionGeometryTool::UpdateActiveGroupLayer()
 
 
 
-void USetCollisionGeometryTool::UpdateVisualization()
-{
-	float UseThickness = VizSettings->LineThickness;
-	FColor UseColor = VizSettings->Color;
-	int32 ColorIdx = 0;
-	PreviewGeom->UpdateAllLineSets([&](ULineSetComponent* LineSet)
-	{
-		LineSet->SetAllLinesThickness(UseThickness);
-		LineSet->SetAllLinesColor(VizSettings->bRandomColors ? LinearColors::SelectFColor(ColorIdx++) : UseColor);
-	});
-
-	LineMaterial = ToolSetupUtil::GetDefaultLineComponentMaterial(GetToolManager(), !VizSettings->bShowHidden);
-	PreviewGeom->SetAllLineSetsMaterial(LineMaterial);
-}
-
-
-
-
 
 
 
diff --git a/Engine/Plugins/Experimental/MeshModelingToolsetExp/Source/MeshModelingToolsExp/Public/Physics/CollisionGeometryVisualization.h b/Engine/Plugins/Experimental/MeshModelingToolsetExp/Source/MeshModelingToolsExp/Public/Physics/CollisionGeometryVisualization.h
--- a/Engine/Plugins/Experimental/MeshModelingToolsetExp/Source/MeshModelingToolsExp/Public/Physics/CollisionGeometryVisualization.h
+++ b/Engine/Plugins/Experimental/MeshModelingToolsetExp/Source/MeshModelingToolsExp/Public/Physics/CollisionGeometryVisualization.h
@@ -2,20 +2,62 @@
 
 #pragma once
 
-#include "Physics/PhysicsDataCollection.h"
-#include "Drawing/PreviewGeometryActor.h"
+#include "HAL/Platform.h" // Required by MESHMODELINGTOOLSEXP_API
 
-namespace UE
+class FPhysicsDataCollection;
+class UCollisionGeometryVisualizationProperties;
+class UPreviewGeometry;
+class UMaterialInterface;
+
+namespace UE::PhysicsTools
 {
-	namespace PhysicsTools
-	{
-		/**
-		 * Create line sets in a UPreviewGeometry for all the elements in a Physics Data Collection.
-		 * Spheres and Capsules are drawn as 3-axis wireframes. Convexes are added as wireframes.
-		 */
-		void MESHMODELINGTOOLSEXP_API InitializePreviewGeometryLines(const FPhysicsDataCollection& PhysicsData, UPreviewGeometry* PreviewGeom,
-			const FColor& LineColor, float LineThickness, float DepthBias = 0.0, int32 CircleStepResolution = 16, bool bRandomColors = true );
-	}
+	/**
+	 * Create line sets in a UPreviewGeometry for all the elements in a Physics Data Collection.
+	 * Spheres and Capsules are drawn as 3-axis wireframes. Convexes are added as wireframes.
+	 */
+	void MESHMODELINGTOOLSEXP_API InitializeCollisionGeometryVisualization(
+		UPreviewGeometry* PreviewGeom,
+		UCollisionGeometryVisualizationProperties* Settings,
+		const FPhysicsDataCollection& PhysicsData,
+		float DepthBias = 0.f,
+		int32 CircleStepResolution = 16);
+
+	/**
+	 * Update line sets in a UPreviewGeometry.
+	 */
+	void MESHMODELINGTOOLSEXP_API UpdateCollisionGeometryVisualization(
+		UPreviewGeometry* PreviewGeom,
+		UCollisionGeometryVisualizationProperties* Settings);
+
+	// The following functions are like InitializeCollisionGeometryVisualization/UpdateCollisionGeometryVisualization
+	// but assume that the given UPreviewGeometry only represents part of the complete collision geometry visualization
+
+	/**
+	 * Like InitializeCollisionGeometryVisualization but only initializes part of the collision geometry visualization.
+	 * 
+	 * The FirstLineSetIndex parameter is used implement Settings->bRandomColors since random colors are determined
+	 * using unique integers for each line set, which may be stored in separate UPreviewGeometry objects. Since the
+	 * initialization is partial this function doesn't set Settings->bVisualizationDirty, the caller should do that.
+	 */
+	void MESHMODELINGTOOLSEXP_API PartiallyInitializeCollisionGeometryVisualization(
+		UPreviewGeometry* PartialPreviewGeom,
+		UCollisionGeometryVisualizationProperties* Settings,
+		const FPhysicsDataCollection& PhysicsData,
+		int32 FirstLineSetIndex,
+		float DepthBias = 0.f,
+		int32 CircleStepResolution = 16);
+
+	/**
+	 * Like UpdateCollisionGeometryVisualization but only updates part of the collision geometry visualization.
+	 * 
+	 * The FirstLineSetIndex parameter is used implement Settings->bRandomColors since random colors are determined
+	 * using unique integers for each line set, which may be stored in separate UPreviewGeometry objects. Since the
+	 * update is partial this function doesn't check or set Settings->bVisualizationDirty, the caller should do that. 
+	 */
+	void MESHMODELINGTOOLSEXP_API PartiallyUpdateCollisionGeometryVisualization(
+		UPreviewGeometry* PartialPreviewGeom,
+		UCollisionGeometryVisualizationProperties* Settings,
+		int32 FirstLineSetIndex);
 }
 
 
diff --git a/Engine/Plugins/Experimental/MeshModelingToolsetExp/Source/MeshModelingToolsExp/Public/Physics/CollisionPropertySets.h b/Engine/Plugins/Experimental/MeshModelingToolsetExp/Source/MeshModelingToolsExp/Public/Physics/CollisionPropertySets.h
--- a/Engine/Plugins/Experimental/MeshModelingToolsetExp/Source/MeshModelingToolsExp/Public/Physics/CollisionPropertySets.h
+++ b/Engine/Plugins/Experimental/MeshModelingToolsetExp/Source/MeshModelingToolsExp/Public/Physics/CollisionPropertySets.h
@@ -144,6 +144,23 @@ class MESHMODELINGTOOLSEXP_API UCollisionGeometryVisualizationProperties : publi
 	GENERATED_BODY()
 public:
 
+	/** Sets materials and registers property watchers to set bVisualizationDirty to true when any property changes */
+	void Initialize(UInteractiveTool* Tool);
+
+	/**
+	 * Returns the line set color.
+	 * If bRandomColors == true  LineSetIndex is used to return a random color
+	 * If bRandomColors == false LineSetIndex is ignored and this->Color is returned
+	 */
+	FColor GetLineSetColor(int32 LineSetIndex = 0) const;
+
+	UMaterialInterface* GetLineMaterial() const
+	{
+		ensure(LineMaterial);
+		ensure(LineMaterialShowingHidden);
+		return bShowHidden ? LineMaterialShowingHidden : LineMaterial;
+	}
+
 	/** Thickness of lines used to visualize collision shapes */
 	UPROPERTY(EditAnywhere, Category = "Collision Visualization")
 	float LineThickness = 3.0f;
@@ -159,6 +176,16 @@ public:
 	/** The color to use for all collision visualizations, if random colors are not used */
 	UPROPERTY(EditAnywhere, Category = "Collision Visualization", meta = (EditCondition = "!bRandomColors"))
 	FColor Color = FColor::Red;
+
+	//~Used if bShowHidden is false
+	UPROPERTY(Transient, meta=(TransientToolProperty))
+	TObjectPtr<UMaterialInterface> LineMaterial = nullptr;
+
+	//~Used if bShowHidden is true
+	UPROPERTY(Transient, meta=(TransientToolProperty))
+	TObjectPtr<UMaterialInterface> LineMaterialShowingHidden = nullptr;
+
+	bool bVisualizationDirty = false;
 };
 
 
diff --git a/Engine/Plugins/Experimental/MeshModelingToolsetExp/Source/MeshModelingToolsExp/Public/Physics/ExtractCollisionGeometryTool.h b/Engine/Plugins/Experimental/MeshModelingToolsetExp/Source/MeshModelingToolsExp/Public/Physics/ExtractCollisionGeometryTool.h
--- a/Engine/Plugins/Experimental/MeshModelingToolsetExp/Source/MeshModelingToolsExp/Public/Physics/ExtractCollisionGeometryTool.h
+++ b/Engine/Plugins/Experimental/MeshModelingToolsetExp/Source/MeshModelingToolsExp/Public/Physics/ExtractCollisionGeometryTool.h
@@ -110,7 +110,4 @@ protected:
 	bool bResultValid = false;
 	void RecalculateMesh_Simple();
 	void RecalculateMesh_Complex();
-
-	bool bVisualizationDirty = false;
-	void UpdateVisualization();
 };
diff --git a/Engine/Plugins/Experimental/MeshModelingToolsetExp/Source/MeshModelingToolsExp/Public/Physics/PhysicsInspectorTool.h b/Engine/Plugins/Experimental/MeshModelingToolsetExp/Source/MeshModelingToolsExp/Public/Physics/PhysicsInspectorTool.h
--- a/Engine/Plugins/Experimental/MeshModelingToolsetExp/Source/MeshModelingToolsExp/Public/Physics/PhysicsInspectorTool.h
+++ b/Engine/Plugins/Experimental/MeshModelingToolsetExp/Source/MeshModelingToolsExp/Public/Physics/PhysicsInspectorTool.h
@@ -50,19 +50,10 @@ protected:
 	UPROPERTY()
 	TArray<TObjectPtr<UPhysicsObjectToolPropertySet>> ObjectData;
 
-	UPROPERTY()
-	TObjectPtr<UMaterialInterface> LineMaterial = nullptr;
-
 protected:
 	UPROPERTY()
 	TArray<TObjectPtr<UPreviewGeometry>> PreviewElements;
 
 	// these are TSharedPtr because TPimplPtr cannot currently be added to a TArray?
 	TArray<TSharedPtr<FPhysicsDataCollection>> PhysicsInfos;
-
-	void InitializeGeometry(const FPhysicsDataCollection& PhysicsData, UPreviewGeometry* PreviewGeom);
-	void InitializeObjectProperties(const FPhysicsDataCollection& PhysicsData, UPhysicsObjectToolPropertySet* PropSet);
-
-	bool bVisualizationDirty = false;
-	void UpdateVisualization();
 };
diff --git a/Engine/Plugins/Experimental/MeshModelingToolsetExp/Source/MeshModelingToolsExp/Public/Physics/SetCollisionGeometryTool.h b/Engine/Plugins/Experimental/MeshModelingToolsetExp/Source/MeshModelingToolsExp/Public/Physics/SetCollisionGeometryTool.h
--- a/Engine/Plugins/Experimental/MeshModelingToolsetExp/Source/MeshModelingToolsExp/Public/Physics/SetCollisionGeometryTool.h
+++ b/Engine/Plugins/Experimental/MeshModelingToolsetExp/Source/MeshModelingToolsExp/Public/Physics/SetCollisionGeometryTool.h
@@ -199,7 +199,7 @@ public:
 	virtual bool CanAccept() const override
 	{
 		// allow accept when we're showing the current, valid result
-		return Super::CanAccept() && bInputMeshesValid && Compute && Compute->HaveValidResult() && !bVisualizationDirty;
+		return Super::CanAccept() && bInputMeshesValid && Compute && Compute->HaveValidResult() && !VizSettings->bVisualizationDirty;
 	}
 
 	// Begin IGenericDataOperatorFactory interface
@@ -218,10 +218,7 @@ protected:
 	TObjectPtr<UCollisionGeometryVisualizationProperties> VizSettings = nullptr;
 
 	UPROPERTY()
-	TObjectPtr<UPhysicsObjectToolPropertySet> CollisionProps;
-
-	UPROPERTY()
-	TObjectPtr<UMaterialInterface> LineMaterial = nullptr;
+	TObjectPtr<UPhysicsObjectToolPropertySet> CollisionProps = nullptr;
 
 	//
 	// Background compute
@@ -230,7 +227,7 @@ protected:
 
 protected:
 	UPROPERTY()
-	TObjectPtr<UPreviewGeometry> PreviewGeom;
+	TObjectPtr<UPreviewGeometry> PreviewGeom = nullptr;
 
 	TArray<int32> SourceObjectIndices;
 	bool bSourcesHidden = false;
@@ -293,7 +290,4 @@ protected:
 	TSharedPtr<FPhysicsDataCollection, ESPMode::ThreadSafe> GeneratedCollision;
 	TSharedPtr<TArray<FPhysicsDataCollection>, ESPMode::ThreadSafe> OtherInputsCollision;
 	TSharedPtr<TArray<FTransform3d>, ESPMode::ThreadSafe> OtherInputsTransforms;
-
-	bool bVisualizationDirty = false;
-	void UpdateVisualization();
 };
