[Backout] - CL26007087
Original CL Desc
-----------------------------------------------------------------
Refactor collision geometry visualization

This change was motivated by future work to implement a minimal boilerplate system for visualizing geometry selections in tools, which will work similarly to this

diff --git a/Engine/Plugins/Experimental/MeshLODToolset/Source/MeshLODToolset/Private/Tools/GenerateStaticMeshLODAssetTool.cpp b/Engine/Plugins/Experimental/MeshLODToolset/Source/MeshLODToolset/Private/Tools/GenerateStaticMeshLODAssetTool.cpp
--- a/Engine/Plugins/Experimental/MeshLODToolset/Source/MeshLODToolset/Private/Tools/GenerateStaticMeshLODAssetTool.cpp
+++ b/Engine/Plugins/Experimental/MeshLODToolset/Source/MeshLODToolset/Private/Tools/GenerateStaticMeshLODAssetTool.cpp
@@ -408,7 +408,9 @@ void UGenerateStaticMeshLODAssetTool::Setup()
 		FPhysicsDataCollection PhysicsData;
 		PhysicsData.Geometry = GenerateLODOp->ResultCollision;
 		PhysicsData.CopyGeometryToAggregate();
-		UE::PhysicsTools::InitializeCollisionGeometryVisualization(CollisionPreview, CollisionVizSettings, PhysicsData, 0.f, 16);
+		UE::PhysicsTools::InitializePreviewGeometryLines(PhysicsData,
+														 CollisionPreview,
+														 CollisionVizSettings->Color, CollisionVizSettings->LineThickness, 0.0f, 16, CollisionVizSettings->bRandomColors);
 
 		// Must happen on main thread, and GenerateProcess might be in use by an Op somewhere else
 		GenerateProcess->GraphEvalCriticalSection.Lock();
@@ -434,7 +436,10 @@ void UGenerateStaticMeshLODAssetTool::Setup()
 	CollisionVizSettings = NewObject<UCollisionGeometryVisualizationProperties>(this);
 	CollisionVizSettings->RestoreProperties(this);
 	AddToolPropertySource(CollisionVizSettings);
-	CollisionVizSettings->Initialize(this);
+	CollisionVizSettings->WatchProperty(CollisionVizSettings->LineThickness, [this](float NewValue) { bCollisionVisualizationDirty = true; });
+	CollisionVizSettings->WatchProperty(CollisionVizSettings->Color, [this](FColor NewValue) { bCollisionVisualizationDirty = true; });
+	CollisionVizSettings->WatchProperty(CollisionVizSettings->bRandomColors, [this](bool bNewValue) { bCollisionVisualizationDirty = true; });
+	CollisionVizSettings->WatchProperty(CollisionVizSettings->bShowHidden, [this](bool bNewValue) { bCollisionVisualizationDirty = true; });
 
 	CollisionPreview = NewObject<UPreviewGeometry>(this);
 	CollisionPreview->CreateInWorld(GetTargetWorld(), (FTransform)PreviewTransform);
@@ -604,11 +609,31 @@ void UGenerateStaticMeshLODAssetTool::OnTick(float DeltaTime)
 		PreviewWithBackgroundCompute->Tick(DeltaTime);
 	}
 
-	UE::PhysicsTools::UpdateCollisionGeometryVisualization(CollisionPreview, CollisionVizSettings);
+	if (bCollisionVisualizationDirty)
+	{
+		UpdateCollisionVisualization();
+		bCollisionVisualizationDirty = false;
+	}
 }
 
 
 
+void UGenerateStaticMeshLODAssetTool::UpdateCollisionVisualization()
+{
+	float UseThickness = CollisionVizSettings->LineThickness;
+	FColor UseColor = CollisionVizSettings->Color;
+	LineMaterial = ToolSetupUtil::GetDefaultLineComponentMaterial(GetToolManager(), !CollisionVizSettings->bShowHidden);
+
+	int32 ColorIdx = 0;
+	CollisionPreview->UpdateAllLineSets([&](ULineSetComponent* LineSet)
+	{
+		LineSet->SetAllLinesThickness(UseThickness);
+		LineSet->SetAllLinesColor(CollisionVizSettings->bRandomColors ? LinearColors::SelectFColor(ColorIdx++) : UseColor);
+	});
+	CollisionPreview->SetAllLineSetsMaterial(LineMaterial);
+}
+
+
 void UGenerateStaticMeshLODAssetTool::CreateNewAsset()
 {
 	check(PreviewWithBackgroundCompute->HaveValidResult());
diff --git a/Engine/Plugins/Experimental/MeshLODToolset/Source/MeshLODToolset/Public/Tools/GenerateStaticMeshLODAssetTool.h b/Engine/Plugins/Experimental/MeshLODToolset/Source/MeshLODToolset/Public/Tools/GenerateStaticMeshLODAssetTool.h
--- a/Engine/Plugins/Experimental/MeshLODToolset/Source/MeshLODToolset/Public/Tools/GenerateStaticMeshLODAssetTool.h
+++ b/Engine/Plugins/Experimental/MeshLODToolset/Source/MeshLODToolset/Public/Tools/GenerateStaticMeshLODAssetTool.h
@@ -320,6 +320,9 @@ protected:
 	UPROPERTY()
 	TObjectPtr<UPhysicsObjectToolPropertySet> ObjectData = nullptr;
 
+	UPROPERTY()
+	TObjectPtr<UMaterialInterface> LineMaterial = nullptr;
+
 	UPROPERTY()
 	TObjectPtr<UPreviewGeometry> CollisionPreview;
 
@@ -334,6 +337,10 @@ protected:
 
 	void OnSettingsModified();
 
+	bool bCollisionVisualizationDirty = false;
+	void UpdateCollisionVisualization();
+
+
 	void CreateNewAsset();
 	void UpdateExistingAsset();
 
diff --git a/Engine/Plugins/Experimental/MeshModelingToolsetExp/Source/MeshModelingToolsExp/Private/Physics/CollisionGeometryVisualization.cpp b/Engine/Plugins/Experimental/MeshModelingToolsetExp/Source/MeshModelingToolsExp/Private/Physics/CollisionGeometryVisualization.cpp
--- a/Engine/Plugins/Experimental/MeshModelingToolsetExp/Source/MeshModelingToolsExp/Private/Physics/CollisionGeometryVisualization.cpp
+++ b/Engine/Plugins/Experimental/MeshModelingToolsetExp/Source/MeshModelingToolsExp/Private/Physics/CollisionGeometryVisualization.cpp
@@ -1,31 +1,31 @@
 // Copyright Epic Games, Inc. All Rights Reserved.
 
 #include "Physics/CollisionGeometryVisualization.h"
-#include "Physics/PhysicsDataCollection.h"
-#include "Physics/CollisionPropertySets.h"
 #include "Generators/LineSegmentGenerators.h"
-#include "Drawing/PreviewGeometryActor.h"
+#include "Util/ColorConstants.h"
+#include "Chaos/Levelset.h"
 
 using namespace UE::Geometry;
 
-namespace
+void UE::PhysicsTools::InitializePreviewGeometryLines(const FPhysicsDataCollection& PhysicsData, UPreviewGeometry* PreviewGeom,
+	const FColor& LineColor, float LineThickness, float DepthBias, int32 CircleStepResolution, bool bRandomColors)
 {
-
-void InitializePreviewGeometryLines(
-	const FPhysicsDataCollection& PhysicsData,
-	UPreviewGeometry* PreviewGeom,
-	UMaterialInterface* LineMaterial,
-	TFunctionRef<FColor(int32 LineSetIndex)> LineSetIndexToColorFunc,
-	float LineThickness,
-	float DepthBias,
-	int32 CircleStepResolution,
-	int32 FirstLineSetIndex)
-{
-	check(PreviewGeom);
-	check(LineMaterial);
-
 	int32 CircleSteps = FMath::Max(4, CircleStepResolution);
-	int32 LineSetIndex = FirstLineSetIndex;
+	FColor SphereColor = LineColor;
+	FColor BoxColor = LineColor;
+	FColor ConvexColor = LineColor;
+	FColor LevelSetColor = LineColor;
+	FColor CapsuleColor = LineColor;
+	int32 OverrideColorIdx = 0;
+
+	auto GetNextColor = [&OverrideColorIdx, bRandomColors](FColor DefaultColor)
+	{
+		if (bRandomColors)
+		{
+			return LinearColors::SelectFColor(OverrideColorIdx++);
+		}
+		return DefaultColor;
+	};
 
 	const FKAggregateGeom& AggGeom = PhysicsData.AggGeom;
 
@@ -34,7 +34,7 @@ void InitializePreviewGeometryLines(
 	{
 		PreviewGeom->CreateOrUpdateLineSet(FString::Printf(TEXT("Spheres %d"), Index), 1, [&](int32 UnusedIndex, TArray<FRenderableLine>& LinesOut)
 		{
-			FColor Color = LineSetIndexToColorFunc(LineSetIndex++);
+			FColor Color = GetNextColor(SphereColor);
 
 			const FKSphereElem& Sphere = AggGeom.SphereElems[Index];
 			FTransform ElemTransform = Sphere.GetTransform();
@@ -56,7 +56,7 @@ void InitializePreviewGeometryLines(
 	{
 		PreviewGeom->CreateOrUpdateLineSet(FString::Printf(TEXT("Boxes %d"), Index), 1, [&](int32 UnusedIndex, TArray<FRenderableLine>& LinesOut)
 		{
-			FColor Color = LineSetIndexToColorFunc(LineSetIndex++);
+			FColor Color = GetNextColor(BoxColor);
 
 			const FKBoxElem& Box = AggGeom.BoxElems[Index];
 			FTransform ElemTransform = Box.GetTransform();
@@ -77,7 +77,7 @@ void InitializePreviewGeometryLines(
 	{
 		PreviewGeom->CreateOrUpdateLineSet(FString::Printf(TEXT("Capsules %d"), Index), 1, [&](int32 UnusedIndex, TArray<FRenderableLine>& LinesOut)
 		{
-			FColor Color = LineSetIndexToColorFunc(LineSetIndex++);
+			FColor Color = GetNextColor(CapsuleColor);
 
 			const FKSphylElem& Capsule = AggGeom.SphylElems[Index];
 			FTransform ElemTransform = Capsule.GetTransform();
@@ -125,7 +125,7 @@ void InitializePreviewGeometryLines(
 	{
 		PreviewGeom->CreateOrUpdateLineSet(FString::Printf(TEXT("Convex %d"), Index), 1, [&](int32 UnusedIndex, TArray<FRenderableLine>& LinesOut)
 		{
-			FColor Color = LineSetIndexToColorFunc(LineSetIndex++);
+			FColor Color = GetNextColor(ConvexColor);
 
 			const FKConvexElem& Convex = AggGeom.ConvexElems[Index];
 			FTransform ElemTransform = Convex.GetTransform();
@@ -143,13 +143,13 @@ void InitializePreviewGeometryLines(
 			}
 		});
 	}
-
+	
 	// for Level Sets draw the grid cells where phi < 0
 	for (int32 Index = 0; Index < AggGeom.LevelSetElems.Num(); Index++)
 	{
 		PreviewGeom->CreateOrUpdateLineSet(FString::Printf(TEXT("Level Set %d"), Index), 1, [&](int32 UnusedIndex, TArray<FRenderableLine>& LinesOut)
 		{
-			FColor Color = LineSetIndexToColorFunc(LineSetIndex++);
+			const FColor Color = GetNextColor(LevelSetColor);
 			const FKLevelSetElem& LevelSet = AggGeom.LevelSetElems[Index];
 			
 			FTransform ElemTransform = LevelSet.GetTransform();
@@ -183,107 +183,4 @@ void InitializePreviewGeometryLines(
 	// Unclear whether we actually use these in the Engine, for UBodySetup? Does not appear to be supported by UxX import system,
 	// and online documentation suggests they may only be supported for cloth?
 	ensure(AggGeom.TaperedCapsuleElems.Num() == 0);
-
-	PreviewGeom->SetAllLineSetsMaterial(LineMaterial);
-}
-
-void UpdatePreviewGeometryLines(
-	UPreviewGeometry* PartialPreviewGeom,
-	UCollisionGeometryVisualizationProperties* Settings,
-	int32 FirstLineSetIndex = 0)
-{
-	check(PartialPreviewGeom);
-	check(Settings);
-
-	int32 LineSetIndex = FirstLineSetIndex;
-	PartialPreviewGeom->UpdateAllLineSets([&](ULineSetComponent* LineSet)
-	{
-		FColor LineColor = Settings->GetLineSetColor(LineSetIndex++);
-		LineSet->SetAllLinesThickness(Settings->LineThickness);
-		LineSet->SetAllLinesColor(LineColor);
-	});
-	PartialPreviewGeom->SetAllLineSetsMaterial(Settings->GetLineMaterial());
-}
-
-} // end namespace
-
-
-
-
-
-
-void UE::PhysicsTools::InitializeCollisionGeometryVisualization(
-	UPreviewGeometry* PreviewGeom,
-	UCollisionGeometryVisualizationProperties* Settings,
-	const FPhysicsDataCollection& PhysicsData,
-	float DepthBias,
-	int32 CircleStepResolution)
-{
-	check(PreviewGeom);
-	check(Settings);
-
-	InitializePreviewGeometryLines(
-		PhysicsData,
-		PreviewGeom,
-		Settings->GetLineMaterial(),
-		[&Settings](int LineSetIndex) { return Settings->GetLineSetColor(LineSetIndex); },
-		Settings->LineThickness,
-		DepthBias,
-		CircleStepResolution,
-		0);
-
-	Settings->bVisualizationDirty = false;
-}
-
-void UE::PhysicsTools::UpdateCollisionGeometryVisualization(
-	UPreviewGeometry* PreviewGeom,
-	UCollisionGeometryVisualizationProperties* Settings)
-{
-	check(PreviewGeom);
-	check(Settings);
-
-	if (Settings->bVisualizationDirty)
-	{
-		UpdatePreviewGeometryLines(PreviewGeom, Settings);
-		Settings->bVisualizationDirty = false;
-	}
-}
-
-
-
-
-
-
-void UE::PhysicsTools::PartiallyInitializeCollisionGeometryVisualization(
-	UPreviewGeometry* PreviewGeom,
-	UCollisionGeometryVisualizationProperties* Settings,
-	const FPhysicsDataCollection& PhysicsData,
-	int32 FirstLineSetIndex,
-	float DepthBias,
-	int32 CircleStepResolution)
-{
-	check(PreviewGeom);
-	check(Settings);
-
-	InitializePreviewGeometryLines(
-		PhysicsData,
-		PreviewGeom,
-		Settings->GetLineMaterial(),
-		[&Settings](int LineSetIndex) { return Settings->GetLineSetColor(LineSetIndex); },
-		Settings->LineThickness,
-		DepthBias,
-		CircleStepResolution,
-		FirstLineSetIndex);
-}
-
-void UE::PhysicsTools::PartiallyUpdateCollisionGeometryVisualization(
-	UPreviewGeometry* PartialPreviewGeom,
-	UCollisionGeometryVisualizationProperties* Settings,
-	int32 FirstLineSetIndex)
-{
-	check(PartialPreviewGeom);
-	check(Settings);
-
-	UpdatePreviewGeometryLines(PartialPreviewGeom, Settings, FirstLineSetIndex);
-}
-
+}
\ No newline at end of file
diff --git a/Engine/Plugins/Experimental/MeshModelingToolsetExp/Source/MeshModelingToolsExp/Private/Physics/CollisionPropertySets.cpp b/Engine/Plugins/Experimental/MeshModelingToolsetExp/Source/MeshModelingToolsExp/Private/Physics/CollisionPropertySets.cpp
--- a/Engine/Plugins/Experimental/MeshModelingToolsetExp/Source/MeshModelingToolsExp/Private/Physics/CollisionPropertySets.cpp
+++ b/Engine/Plugins/Experimental/MeshModelingToolsetExp/Source/MeshModelingToolsExp/Private/Physics/CollisionPropertySets.cpp
@@ -6,27 +6,9 @@
 #include "Engine/World.h"
 #include "GameFramework/Actor.h"
 #include "Physics/PhysicsDataCollection.h"
-#include "Util/ColorConstants.h"
-#include "ToolSetupUtil.h"
 
 #include UE_INLINE_GENERATED_CPP_BY_NAME(CollisionPropertySets)
 
-void UCollisionGeometryVisualizationProperties::Initialize(UInteractiveTool* Tool)
-{
-	WatchProperty(LineThickness, [this](float NewValue) { bVisualizationDirty = true; });
-	WatchProperty(Color, [this](FColor NewValue) { bVisualizationDirty = true; });
-	WatchProperty(bRandomColors, [this](bool bNewValue) { bVisualizationDirty = true; });
-	WatchProperty(bShowHidden, [this](bool bNewValue) { bVisualizationDirty = true; });
-
-	LineMaterial = ToolSetupUtil::GetDefaultLineComponentMaterial(Tool->GetToolManager(), true);
-	LineMaterialShowingHidden = ToolSetupUtil::GetDefaultLineComponentMaterial(Tool->GetToolManager(), false);
-}
-
-FColor UCollisionGeometryVisualizationProperties::GetLineSetColor(int32 LineSetIndex) const
-{
-	return bRandomColors ? UE::Geometry::LinearColors::SelectFColor(LineSetIndex) : Color;
-}
-
 void UE::PhysicsTools::InitializePhysicsToolObjectPropertySet(const FPhysicsDataCollection* PhysicsData, UPhysicsObjectToolPropertySet* PropSet)
 {
 	check(PhysicsData);
diff --git a/Engine/Plugins/Experimental/MeshModelingToolsetExp/Source/MeshModelingToolsExp/Private/Physics/ExtractCollisionGeometryTool.cpp b/Engine/Plugins/Experimental/MeshModelingToolsetExp/Source/MeshModelingToolsExp/Private/Physics/ExtractCollisionGeometryTool.cpp
--- a/Engine/Plugins/Experimental/MeshModelingToolsetExp/Source/MeshModelingToolsExp/Private/Physics/ExtractCollisionGeometryTool.cpp
+++ b/Engine/Plugins/Experimental/MeshModelingToolsetExp/Source/MeshModelingToolsExp/Private/Physics/ExtractCollisionGeometryTool.cpp
@@ -92,7 +92,10 @@ void UExtractCollisionGeometryTool::Setup()
 	VizSettings = NewObject<UCollisionGeometryVisualizationProperties>(this);
 	VizSettings->RestoreProperties(this);
 	AddToolPropertySource(VizSettings);
-	VizSettings->Initialize(this);
+	VizSettings->WatchProperty(VizSettings->LineThickness, [this](float NewValue) { bVisualizationDirty = true; });
+	VizSettings->WatchProperty(VizSettings->Color, [this](FColor NewValue) { bVisualizationDirty = true; });
+	VizSettings->WatchProperty(VizSettings->bRandomColors, [this](bool bNewValue) { bVisualizationDirty = true; });
+	VizSettings->WatchProperty(VizSettings->bShowHidden, [this](bool bNewValue) { bVisualizationDirty = true; });
 
 	UBodySetup* BodySetup = UE::ToolTarget::GetPhysicsBodySetup(Target);
 	if (BodySetup)
@@ -105,8 +108,9 @@ void UExtractCollisionGeometryTool::Setup()
 		//PhysicsInfo->ExternalScale3D = TargetTransform.GetScale3D();
 		//TargetTransform.SetScale3D(FVector::OneVector);
 		PreviewElements->CreateInWorld(UE::ToolTarget::GetTargetActor(Target)->GetWorld(), TargetTransform);
-
-		UE::PhysicsTools::InitializeCollisionGeometryVisualization(PreviewElements, VizSettings, *PhysicsInfo);
+		
+		UE::PhysicsTools::InitializePreviewGeometryLines(*PhysicsInfo, PreviewElements,
+			VizSettings->Color, VizSettings->LineThickness, 0.0f, 16, VizSettings->bRandomColors);
 
 		ObjectProps = NewObject<UPhysicsObjectToolPropertySet>(this);
 		UE::PhysicsTools::InitializePhysicsToolObjectPropertySet(PhysicsInfo.Get(), ObjectProps);
@@ -212,7 +216,28 @@ void UExtractCollisionGeometryTool::OnTick(float DeltaTime)
 		bResultValid = true;
 	}
 
-	UE::PhysicsTools::UpdateCollisionGeometryVisualization(PreviewElements, VizSettings);
+	if (bVisualizationDirty)
+	{
+		UpdateVisualization();
+		bVisualizationDirty = false;
+	}
+}
+
+
+
+void UExtractCollisionGeometryTool::UpdateVisualization()
+{
+	float UseThickness = VizSettings->LineThickness;
+	FColor UseColor = VizSettings->Color;
+	int32 ColorIdx = 0;
+	PreviewElements->UpdateAllLineSets([&](ULineSetComponent* LineSet)
+		{
+			LineSet->SetAllLinesThickness(UseThickness);
+			LineSet->SetAllLinesColor(VizSettings->bRandomColors ? LinearColors::SelectFColor(ColorIdx++) : UseColor);
+		});
+	
+	UMaterialInterface* LineMaterial = ToolSetupUtil::GetDefaultLineComponentMaterial(GetToolManager(), !VizSettings->bShowHidden);
+	PreviewElements->SetAllLineSetsMaterial(LineMaterial);
 }
 
 
diff --git a/Engine/Plugins/Experimental/MeshModelingToolsetExp/Source/MeshModelingToolsExp/Private/Physics/PhysicsInspectorTool.cpp b/Engine/Plugins/Experimental/MeshModelingToolsetExp/Source/MeshModelingToolsExp/Private/Physics/PhysicsInspectorTool.cpp
--- a/Engine/Plugins/Experimental/MeshModelingToolsetExp/Source/MeshModelingToolsExp/Private/Physics/PhysicsInspectorTool.cpp
+++ b/Engine/Plugins/Experimental/MeshModelingToolsetExp/Source/MeshModelingToolsExp/Private/Physics/PhysicsInspectorTool.cpp
@@ -49,37 +49,34 @@ void UPhysicsInspectorTool::Setup()
 	VizSettings = NewObject<UCollisionGeometryVisualizationProperties>(this);
 	VizSettings->RestoreProperties(this);
 	AddToolPropertySource(VizSettings);
+	VizSettings->WatchProperty(VizSettings->LineThickness, [this](float NewValue) { bVisualizationDirty = true; });
+	VizSettings->WatchProperty(VizSettings->Color, [this](FColor NewValue) { bVisualizationDirty = true; });
+	VizSettings->WatchProperty(VizSettings->bRandomColors, [this](bool bNewValue) { bVisualizationDirty = true; });
+	VizSettings->WatchProperty(VizSettings->bShowHidden, [this](bool bNewValue) { bVisualizationDirty = true; });
 
-	// Intitialize the collision geometry visualization
+	for (int32 ComponentIdx = 0; ComponentIdx < Targets.Num(); ComponentIdx++)
 	{
-		VizSettings->Initialize(this);
-		int32 FirstLineSetIndex = 0;
-		for (int32 ComponentIdx = 0; ComponentIdx < Targets.Num(); ComponentIdx++)
+		UBodySetup* BodySetup = UE::ToolTarget::GetPhysicsBodySetup(Targets[ComponentIdx]);
+		if (BodySetup)
 		{
-			UBodySetup* BodySetup = UE::ToolTarget::GetPhysicsBodySetup(Targets[ComponentIdx]);
-			if (BodySetup)
-			{
-				TSharedPtr<FPhysicsDataCollection> PhysicsData = MakeShared<FPhysicsDataCollection>();
-				PhysicsData->InitializeFromComponent( UE::ToolTarget::GetTargetComponent(Targets[ComponentIdx]), true);
-
-				PhysicsInfos.Add(PhysicsData);
-
-				UPreviewGeometry* PreviewGeom = NewObject<UPreviewGeometry>(this);
-				FTransform TargetTransform = (FTransform) UE::ToolTarget::GetLocalToWorldTransform(Targets[ComponentIdx]);
-				PhysicsData->ExternalScale3D = TargetTransform.GetScale3D();
-				TargetTransform.SetScale3D(FVector::OneVector);
-				PreviewGeom->CreateInWorld(UE::ToolTarget::GetTargetActor(Targets[ComponentIdx])->GetWorld(), TargetTransform);
-				PreviewElements.Add(PreviewGeom);
-
-				UE::PhysicsTools::PartiallyInitializeCollisionGeometryVisualization(PreviewGeom, VizSettings, *PhysicsData, FirstLineSetIndex);
-				FirstLineSetIndex += PreviewGeom->LineSets.Num();
-
-				UPhysicsObjectToolPropertySet* ObjectProps = NewObject<UPhysicsObjectToolPropertySet>(this);
-				UE::PhysicsTools::InitializePhysicsToolObjectPropertySet(PhysicsData.Get(), ObjectProps);
-				AddToolPropertySource(ObjectProps);
-			}
+			TSharedPtr<FPhysicsDataCollection> PhysicsData = MakeShared<FPhysicsDataCollection>();
+			PhysicsData->InitializeFromComponent( UE::ToolTarget::GetTargetComponent(Targets[ComponentIdx]), true);
+
+			PhysicsInfos.Add(PhysicsData);
+
+			UPreviewGeometry* PreviewGeom = NewObject<UPreviewGeometry>(this);
+			FTransform TargetTransform = (FTransform) UE::ToolTarget::GetLocalToWorldTransform(Targets[ComponentIdx]);
+			PhysicsData->ExternalScale3D = TargetTransform.GetScale3D();
+			TargetTransform.SetScale3D(FVector::OneVector);
+			PreviewGeom->CreateInWorld(UE::ToolTarget::GetTargetActor(Targets[ComponentIdx])->GetWorld(), TargetTransform);
+			PreviewElements.Add(PreviewGeom);
+
+			InitializeGeometry(*PhysicsData, PreviewGeom);
+
+			UPhysicsObjectToolPropertySet* ObjectProps = NewObject<UPhysicsObjectToolPropertySet>(this);
+			InitializeObjectProperties(*PhysicsData, ObjectProps);
+			AddToolPropertySource(ObjectProps);
 		}
-		VizSettings->bVisualizationDirty = false;
 	}
 
 	SetToolDisplayName(LOCTEXT("ToolName", "Physics Inspector"));
@@ -105,19 +102,48 @@ void UPhysicsInspectorTool::OnShutdown(EToolShutdownType ShutdownType)
 
 void UPhysicsInspectorTool::OnTick(float DeltaTime)
 {
-	if (VizSettings->bVisualizationDirty)
+	if (bVisualizationDirty)
 	{
-		int32 FirstLineSetIndex= 0;
-		for (UPreviewGeometry* Preview : PreviewElements)
-		{
-			UE::PhysicsTools::PartiallyUpdateCollisionGeometryVisualization(Preview, VizSettings, FirstLineSetIndex);
-			FirstLineSetIndex += Preview->LineSets.Num();
-		}
+		UpdateVisualization();
+		bVisualizationDirty = false;
+	}
+}
+
+
 
-		VizSettings->bVisualizationDirty = false;
+void UPhysicsInspectorTool::UpdateVisualization()
+{
+	float UseThickness = VizSettings->LineThickness;
+	FColor UseColor = VizSettings->Color;
+	LineMaterial = ToolSetupUtil::GetDefaultLineComponentMaterial(GetToolManager(), !VizSettings->bShowHidden);
+
+	int32 ColorIdx = 0;
+	for (UPreviewGeometry* Preview : PreviewElements)
+	{
+		Preview->UpdateAllLineSets([&](ULineSetComponent* LineSet)
+		{
+			LineSet->SetAllLinesThickness(UseThickness);
+			LineSet->SetAllLinesColor(VizSettings->bRandomColors ? LinearColors::SelectFColor(ColorIdx++) : UseColor);
+		});
+		Preview->SetAllLineSetsMaterial(LineMaterial);
 	}
 }
 
 
 
+void UPhysicsInspectorTool::InitializeObjectProperties(const FPhysicsDataCollection& PhysicsData, UPhysicsObjectToolPropertySet* PropSet)
+{
+	UE::PhysicsTools::InitializePhysicsToolObjectPropertySet(&PhysicsData, PropSet);
+}
+
+
+
+void UPhysicsInspectorTool::InitializeGeometry(const FPhysicsDataCollection& PhysicsData, UPreviewGeometry* PreviewGeom)
+{
+	UE::PhysicsTools::InitializePreviewGeometryLines(PhysicsData, PreviewGeom,
+		VizSettings->Color, VizSettings->LineThickness, 0.0f, 16, VizSettings->bRandomColors);
+}
+
+
+
 #undef LOCTEXT_NAMESPACE
diff --git a/Engine/Plugins/Experimental/MeshModelingToolsetExp/Source/MeshModelingToolsExp/Private/Physics/SetCollisionGeometryTool.cpp b/Engine/Plugins/Experimental/MeshModelingToolsetExp/Source/MeshModelingToolsExp/Private/Physics/SetCollisionGeometryTool.cpp
--- a/Engine/Plugins/Experimental/MeshModelingToolsetExp/Source/MeshModelingToolsExp/Private/Physics/SetCollisionGeometryTool.cpp
+++ b/Engine/Plugins/Experimental/MeshModelingToolsetExp/Source/MeshModelingToolsExp/Private/Physics/SetCollisionGeometryTool.cpp
@@ -28,7 +28,6 @@
 #include "ToolTargetManager.h"
 #include "UObject/UObjectIterator.h"
 #include "ModelingToolTargetUtil.h"
-#include "Drawing/PreviewGeometryActor.h"
 
 #include UE_INLINE_GENERATED_CPP_BY_NAME(SetCollisionGeometryTool)
 
@@ -303,7 +302,10 @@ void USetCollisionGeometryTool::Setup()
 	VizSettings = NewObject<UCollisionGeometryVisualizationProperties>(this);
 	VizSettings->RestoreProperties(this);
 	AddToolPropertySource(VizSettings);
-	VizSettings->Initialize(this);
+	VizSettings->WatchProperty(VizSettings->LineThickness, [this](float NewValue) { bVisualizationDirty = true; });
+	VizSettings->WatchProperty(VizSettings->Color, [this](FColor NewValue) { bVisualizationDirty = true; });
+	VizSettings->WatchProperty(VizSettings->bRandomColors, [this](bool bNewValue) { bVisualizationDirty = true; });
+	VizSettings->WatchProperty(VizSettings->bShowHidden, [this](bool bNewValue) { bVisualizationDirty = true; });
 
 	// add option for collision properties
 	CollisionProps = NewObject<UPhysicsObjectToolPropertySet>(this);
@@ -519,12 +521,12 @@ void USetCollisionGeometryTool::OnTick(float DeltaTime)
 			{
 				GeneratedCollision = MakeShareable<FPhysicsDataCollection>(Result.Release());
 
-				VizSettings->bVisualizationDirty = true;
+				bVisualizationDirty = true;
 
 				// update visualization
 				PreviewGeom->RemoveAllLineSets();
-
-				UE::PhysicsTools::InitializeCollisionGeometryVisualization(PreviewGeom, VizSettings, *GeneratedCollision);
+				UE::PhysicsTools::InitializePreviewGeometryLines(*GeneratedCollision, PreviewGeom,
+					VizSettings->Color, VizSettings->LineThickness, 0.0f, 16, VizSettings->bRandomColors);
 
 				// update property set
 				CollisionProps->Reset();
@@ -533,7 +535,11 @@ void USetCollisionGeometryTool::OnTick(float DeltaTime)
 		}
 	}
 
-	UE::PhysicsTools::UpdateCollisionGeometryVisualization(PreviewGeom, VizSettings);
+	if (bVisualizationDirty)
+	{
+		UpdateVisualization();
+		bVisualizationDirty = false;
+	}
 }
 
 
@@ -601,6 +607,24 @@ void USetCollisionGeometryTool::UpdateActiveGroupLayer()
 
 
 
+void USetCollisionGeometryTool::UpdateVisualization()
+{
+	float UseThickness = VizSettings->LineThickness;
+	FColor UseColor = VizSettings->Color;
+	int32 ColorIdx = 0;
+	PreviewGeom->UpdateAllLineSets([&](ULineSetComponent* LineSet)
+	{
+		LineSet->SetAllLinesThickness(UseThickness);
+		LineSet->SetAllLinesColor(VizSettings->bRandomColors ? LinearColors::SelectFColor(ColorIdx++) : UseColor);
+	});
+
+	LineMaterial = ToolSetupUtil::GetDefaultLineComponentMaterial(GetToolManager(), !VizSettings->bShowHidden);
+	PreviewGeom->SetAllLineSetsMaterial(LineMaterial);
+}
+
+
+
+
 
 
 
diff --git a/Engine/Plugins/Experimental/MeshModelingToolsetExp/Source/MeshModelingToolsExp/Public/Physics/CollisionGeometryVisualization.h b/Engine/Plugins/Experimental/MeshModelingToolsetExp/Source/MeshModelingToolsExp/Public/Physics/CollisionGeometryVisualization.h
--- a/Engine/Plugins/Experimental/MeshModelingToolsetExp/Source/MeshModelingToolsExp/Public/Physics/CollisionGeometryVisualization.h
+++ b/Engine/Plugins/Experimental/MeshModelingToolsetExp/Source/MeshModelingToolsExp/Public/Physics/CollisionGeometryVisualization.h
@@ -2,62 +2,20 @@
 
 #pragma once
 
-#include "HAL/Platform.h" // Required by MESHMODELINGTOOLSEXP_API
+#include "Physics/PhysicsDataCollection.h"
+#include "Drawing/PreviewGeometryActor.h"
 
-class FPhysicsDataCollection;
-class UCollisionGeometryVisualizationProperties;
-class UPreviewGeometry;
-class UMaterialInterface;
-
-namespace UE::PhysicsTools
+namespace UE
 {
-	/**
-	 * Create line sets in a UPreviewGeometry for all the elements in a Physics Data Collection.
-	 * Spheres and Capsules are drawn as 3-axis wireframes. Convexes are added as wireframes.
-	 */
-	void MESHMODELINGTOOLSEXP_API InitializeCollisionGeometryVisualization(
-		UPreviewGeometry* PreviewGeom,
-		UCollisionGeometryVisualizationProperties* Settings,
-		const FPhysicsDataCollection& PhysicsData,
-		float DepthBias = 0.f,
-		int32 CircleStepResolution = 16);
-
-	/**
-	 * Update line sets in a UPreviewGeometry.
-	 */
-	void MESHMODELINGTOOLSEXP_API UpdateCollisionGeometryVisualization(
-		UPreviewGeometry* PreviewGeom,
-		UCollisionGeometryVisualizationProperties* Settings);
-
-	// The following functions are like InitializeCollisionGeometryVisualization/UpdateCollisionGeometryVisualization
-	// but assume that the given UPreviewGeometry only represents part of the complete collision geometry visualization
-
-	/**
-	 * Like InitializeCollisionGeometryVisualization but only initializes part of the collision geometry visualization.
-	 * 
-	 * The FirstLineSetIndex parameter is used implement Settings->bRandomColors since random colors are determined
-	 * using unique integers for each line set, which may be stored in separate UPreviewGeometry objects. Since the
-	 * initialization is partial this function doesn't set Settings->bVisualizationDirty, the caller should do that.
-	 */
-	void MESHMODELINGTOOLSEXP_API PartiallyInitializeCollisionGeometryVisualization(
-		UPreviewGeometry* PartialPreviewGeom,
-		UCollisionGeometryVisualizationProperties* Settings,
-		const FPhysicsDataCollection& PhysicsData,
-		int32 FirstLineSetIndex,
-		float DepthBias = 0.f,
-		int32 CircleStepResolution = 16);
-
-	/**
-	 * Like UpdateCollisionGeometryVisualization but only updates part of the collision geometry visualization.
-	 * 
-	 * The FirstLineSetIndex parameter is used implement Settings->bRandomColors since random colors are determined
-	 * using unique integers for each line set, which may be stored in separate UPreviewGeometry objects. Since the
-	 * update is partial this function doesn't check or set Settings->bVisualizationDirty, the caller should do that. 
-	 */
-	void MESHMODELINGTOOLSEXP_API PartiallyUpdateCollisionGeometryVisualization(
-		UPreviewGeometry* PartialPreviewGeom,
-		UCollisionGeometryVisualizationProperties* Settings,
-		int32 FirstLineSetIndex);
+	namespace PhysicsTools
+	{
+		/**
+		 * Create line sets in a UPreviewGeometry for all the elements in a Physics Data Collection.
+		 * Spheres and Capsules are drawn as 3-axis wireframes. Convexes are added as wireframes.
+		 */
+		void MESHMODELINGTOOLSEXP_API InitializePreviewGeometryLines(const FPhysicsDataCollection& PhysicsData, UPreviewGeometry* PreviewGeom,
+			const FColor& LineColor, float LineThickness, float DepthBias = 0.0, int32 CircleStepResolution = 16, bool bRandomColors = true );
+	}
 }
 
 
diff --git a/Engine/Plugins/Experimental/MeshModelingToolsetExp/Source/MeshModelingToolsExp/Public/Physics/CollisionPropertySets.h b/Engine/Plugins/Experimental/MeshModelingToolsetExp/Source/MeshModelingToolsExp/Public/Physics/CollisionPropertySets.h
--- a/Engine/Plugins/Experimental/MeshModelingToolsetExp/Source/MeshModelingToolsExp/Public/Physics/CollisionPropertySets.h
+++ b/Engine/Plugins/Experimental/MeshModelingToolsetExp/Source/MeshModelingToolsExp/Public/Physics/CollisionPropertySets.h
@@ -144,23 +144,6 @@ class MESHMODELINGTOOLSEXP_API UCollisionGeometryVisualizationProperties : publi
 	GENERATED_BODY()
 public:
 
-	/** Sets materials and registers property watchers to set bVisualizationDirty to true when any property changes */
-	void Initialize(UInteractiveTool* Tool);
-
-	/**
-	 * Returns the line set color.
-	 * If bRandomColors == true  LineSetIndex is used to return a random color
-	 * If bRandomColors == false LineSetIndex is ignored and this->Color is returned
-	 */
-	FColor GetLineSetColor(int32 LineSetIndex = 0) const;
-
-	UMaterialInterface* GetLineMaterial() const
-	{
-		ensure(LineMaterial);
-		ensure(LineMaterialShowingHidden);
-		return bShowHidden ? LineMaterialShowingHidden : LineMaterial;
-	}
-
 	/** Thickness of lines used to visualize collision shapes */
 	UPROPERTY(EditAnywhere, Category = "Collision Visualization")
 	float LineThickness = 3.0f;
@@ -176,16 +159,6 @@ public:
 	/** The color to use for all collision visualizations, if random colors are not used */
 	UPROPERTY(EditAnywhere, Category = "Collision Visualization", meta = (EditCondition = "!bRandomColors"))
 	FColor Color = FColor::Red;
-
-	//~Used if bShowHidden is false
-	UPROPERTY(Transient, meta=(TransientToolProperty))
-	TObjectPtr<UMaterialInterface> LineMaterial = nullptr;
-
-	//~Used if bShowHidden is true
-	UPROPERTY(Transient, meta=(TransientToolProperty))
-	TObjectPtr<UMaterialInterface> LineMaterialShowingHidden = nullptr;
-
-	bool bVisualizationDirty = false;
 };
 
 
diff --git a/Engine/Plugins/Experimental/MeshModelingToolsetExp/Source/MeshModelingToolsExp/Public/Physics/ExtractCollisionGeometryTool.h b/Engine/Plugins/Experimental/MeshModelingToolsetExp/Source/MeshModelingToolsExp/Public/Physics/ExtractCollisionGeometryTool.h
--- a/Engine/Plugins/Experimental/MeshModelingToolsetExp/Source/MeshModelingToolsExp/Public/Physics/ExtractCollisionGeometryTool.h
+++ b/Engine/Plugins/Experimental/MeshModelingToolsetExp/Source/MeshModelingToolsExp/Public/Physics/ExtractCollisionGeometryTool.h
@@ -110,4 +110,7 @@ protected:
 	bool bResultValid = false;
 	void RecalculateMesh_Simple();
 	void RecalculateMesh_Complex();
+
+	bool bVisualizationDirty = false;
+	void UpdateVisualization();
 };
diff --git a/Engine/Plugins/Experimental/MeshModelingToolsetExp/Source/MeshModelingToolsExp/Public/Physics/PhysicsInspectorTool.h b/Engine/Plugins/Experimental/MeshModelingToolsetExp/Source/MeshModelingToolsExp/Public/Physics/PhysicsInspectorTool.h
--- a/Engine/Plugins/Experimental/MeshModelingToolsetExp/Source/MeshModelingToolsExp/Public/Physics/PhysicsInspectorTool.h
+++ b/Engine/Plugins/Experimental/MeshModelingToolsetExp/Source/MeshModelingToolsExp/Public/Physics/PhysicsInspectorTool.h
@@ -50,10 +50,19 @@ protected:
 	UPROPERTY()
 	TArray<TObjectPtr<UPhysicsObjectToolPropertySet>> ObjectData;
 
+	UPROPERTY()
+	TObjectPtr<UMaterialInterface> LineMaterial = nullptr;
+
 protected:
 	UPROPERTY()
 	TArray<TObjectPtr<UPreviewGeometry>> PreviewElements;
 
 	// these are TSharedPtr because TPimplPtr cannot currently be added to a TArray?
 	TArray<TSharedPtr<FPhysicsDataCollection>> PhysicsInfos;
+
+	void InitializeGeometry(const FPhysicsDataCollection& PhysicsData, UPreviewGeometry* PreviewGeom);
+	void InitializeObjectProperties(const FPhysicsDataCollection& PhysicsData, UPhysicsObjectToolPropertySet* PropSet);
+
+	bool bVisualizationDirty = false;
+	void UpdateVisualization();
 };
diff --git a/Engine/Plugins/Experimental/MeshModelingToolsetExp/Source/MeshModelingToolsExp/Public/Physics/SetCollisionGeometryTool.h b/Engine/Plugins/Experimental/MeshModelingToolsetExp/Source/MeshModelingToolsExp/Public/Physics/SetCollisionGeometryTool.h
--- a/Engine/Plugins/Experimental/MeshModelingToolsetExp/Source/MeshModelingToolsExp/Public/Physics/SetCollisionGeometryTool.h
+++ b/Engine/Plugins/Experimental/MeshModelingToolsetExp/Source/MeshModelingToolsExp/Public/Physics/SetCollisionGeometryTool.h
@@ -199,7 +199,7 @@ public:
 	virtual bool CanAccept() const override
 	{
 		// allow accept when we're showing the current, valid result
-		return Super::CanAccept() && bInputMeshesValid && Compute && Compute->HaveValidResult() && !VizSettings->bVisualizationDirty;
+		return Super::CanAccept() && bInputMeshesValid && Compute && Compute->HaveValidResult() && !bVisualizationDirty;
 	}
 
 	// Begin IGenericDataOperatorFactory interface
@@ -218,7 +218,10 @@ protected:
 	TObjectPtr<UCollisionGeometryVisualizationProperties> VizSettings = nullptr;
 
 	UPROPERTY()
-	TObjectPtr<UPhysicsObjectToolPropertySet> CollisionProps = nullptr;
+	TObjectPtr<UPhysicsObjectToolPropertySet> CollisionProps;
+
+	UPROPERTY()
+	TObjectPtr<UMaterialInterface> LineMaterial = nullptr;
 
 	//
 	// Background compute
@@ -227,7 +230,7 @@ protected:
 
 protected:
 	UPROPERTY()
-	TObjectPtr<UPreviewGeometry> PreviewGeom = nullptr;
+	TObjectPtr<UPreviewGeometry> PreviewGeom;
 
 	TArray<int32> SourceObjectIndices;
 	bool bSourcesHidden = false;
@@ -290,4 +293,7 @@ protected:
 	TSharedPtr<FPhysicsDataCollection, ESPMode::ThreadSafe> GeneratedCollision;
 	TSharedPtr<TArray<FPhysicsDataCollection>, ESPMode::ThreadSafe> OtherInputsCollision;
 	TSharedPtr<TArray<FTransform3d>, ESPMode::ThreadSafe> OtherInputsTransforms;
+
+	bool bVisualizationDirty = false;
+	void UpdateVisualization();
 };
