MeshModelingTools: Subdivide tool improvements dealing with meshes with open boundaries.
- Factor out "extra corner" finding code from PolyEd and add support for it in subdivide tool.
- Add support for rounded corners for mesh boundaries.
- Override subdivision scheme during tool usage rather than just at tool start, to allow extra corners to fix some issues.
- Add some clarifying comments about schemes.

diff --git a/Engine/Plugins/Experimental/MeshModelingToolsetExp/Source/MeshModelingToolsEditorOnlyExp/Private/SubdividePolyTool.cpp b/Engine/Plugins/Experimental/MeshModelingToolsetExp/Source/MeshModelingToolsEditorOnlyExp/Private/SubdividePolyTool.cpp
--- a/Engine/Plugins/Experimental/MeshModelingToolsetExp/Source/MeshModelingToolsEditorOnlyExp/Private/SubdividePolyTool.cpp
+++ b/Engine/Plugins/Experimental/MeshModelingToolsetExp/Source/MeshModelingToolsEditorOnlyExp/Private/SubdividePolyTool.cpp
@@ -28,21 +28,29 @@ public:
 
 	SubdivPostProcessor(int InSubdivisionLevel,
 						ESubdivisionScheme InSubdivisionScheme,
+						ESubdivisionBoundaryScheme InBoundaryScheme,
 						ESubdivisionOutputNormals InNormalComputationMethod,
 						ESubdivisionOutputUVs InUVComputationMethod,
-						bool bInNewPolyGroups) :
+						bool bInNewPolyGroups,
+						TFunction<bool(const FGroupTopology& GroupTopology, 
+							int32 Vid, const FIndex2i& AttachedGroupEdgeEids)> ShouldAddExtraCornerAtVertIn) :
 		SubdivisionLevel(InSubdivisionLevel),
 		SubdivisionScheme(InSubdivisionScheme),
+		BoundaryScheme(InBoundaryScheme),
 		NormalComputationMethod(InNormalComputationMethod),
 		UVComputationMethod(InUVComputationMethod),
-		bNewPolyGroups(bInNewPolyGroups)
+		bNewPolyGroups(bInNewPolyGroups),
+		ShouldAddExtraCornerAtVert(ShouldAddExtraCornerAtVertIn)
 	{}
 
 	int SubdivisionLevel = 3;
 	ESubdivisionScheme SubdivisionScheme = ESubdivisionScheme::CatmullClark;
+	ESubdivisionBoundaryScheme BoundaryScheme = ESubdivisionBoundaryScheme::SmoothCorners;
 	ESubdivisionOutputNormals NormalComputationMethod = ESubdivisionOutputNormals::Generated;
 	ESubdivisionOutputUVs UVComputationMethod = ESubdivisionOutputUVs::Interpolated;
 	bool bNewPolyGroups = false;
+	// Function passed to the group topology to add extra corners
+	TFunction<bool(const FGroupTopology& GroupTopology, int32 Vid, const FIndex2i& AttachedGroupEdgeEids)> ShouldAddExtraCornerAtVert;
 
 	void ProcessMesh(const FDynamicMesh3& Mesh, FDynamicMesh3& OutRenderMesh) final
 	{
@@ -56,10 +64,14 @@ public:
 			return;
 		}
 
-		constexpr bool bAutoCompute = true;
-		FGroupTopology Topo(&Mesh, bAutoCompute);
+		constexpr bool bBuildTopologyImmediately = false; // need to attach extra corner function first
+		FGroupTopology Topo(&Mesh, bBuildTopologyImmediately);
+		Topo.ShouldAddExtraCornerAtVert = ShouldAddExtraCornerAtVert;
+		Topo.RebuildTopology();
+
 		FSubdividePoly Subd(Topo, Mesh, SubdivisionLevel);
 		Subd.SubdivisionScheme = SubdivisionScheme;
+		Subd.BoundaryScheme = BoundaryScheme;
 		Subd.NormalComputationMethod = NormalComputationMethod;
 		Subd.UVComputationMethod = UVComputationMethod;
 		Subd.bNewPolyGroups = bNewPolyGroups;
@@ -83,42 +95,48 @@ bool USubdividePolyTool::CheckGroupTopology(FText& Message)
 {
 	Message = FText();
 
-	FGroupTopology Topo(OriginalMesh.Get(), true);
-	FSubdividePoly TempSubD(Topo, *OriginalMesh, 1);
+	if (!ensure(Topology.IsValid()))
+	{
+		return false;
+	}
+
+	FSubdividePoly TempSubD(*Topology, *OriginalMesh, 1);
+	TempSubD.BoundaryScheme = Properties->BoundaryScheme; // will only matter here if we someday support NoBoundaryFaces
 	FSubdividePoly::ETopologyCheckResult CheckResult = TempSubD.ValidateTopology();
 
 	if (CheckResult == FSubdividePoly::ETopologyCheckResult::NoGroups)
 	{
 		Message = LOCTEXT("NoGroupsWarning",
-						  "This object has no PolyGroups.\nUse the PolyGroups or Select Tool to assign PolyGroups.\nTool will be limited to Loop subdivision scheme.");
+						  "This object has no PolyGroups.\nTool can only use Loop subdivision scheme.\nUse the PolyGroups or Select Tool to assign PolyGroups.");
 		return false;
 	}
 
 	if (CheckResult == FSubdividePoly::ETopologyCheckResult::InsufficientGroups)
 	{
 		Message = LOCTEXT("SingleGroupsWarning",
-						  "This object has only one PolyGroup.\nUse the PolyGroups or Select Tool to assign PolyGroups.\nTool will be limited to Loop subdivision scheme.");
+						  "This object has only one PolyGroup.\nTool can only use Loop subdivision scheme.\nUse the PolyGroups or Select Tool to assign PolyGroups.");
 		return false;
 	}
 
 	if (CheckResult == FSubdividePoly::ETopologyCheckResult::UnboundedPolygroup)
 	{
 		Message = LOCTEXT("NoGroupBoundaryWarning",
-			"Found a PolyGroup with no boundaries.\nUse the PolyGroups or Select Tool to assign PolyGroups.\nTool will be limited to Loop subdivision scheme.");
+			"Found a PolyGroup with no boundaries.\nTool can only use Loop subdivision scheme.\nUse the PolyGroups or Select Tool to assign PolyGroups.");
 		return false;
 	}
 
 	if (CheckResult == FSubdividePoly::ETopologyCheckResult::MultiBoundaryPolygroup)
 	{
 		Message = LOCTEXT("MultipleGroupBoundaryWarning",
-			"Found a PolyGroup with multiple boundaries, which is not supported.\nUse the PolyGroups or Select Tool to assign PolyGroups.\nTool will be limited to Loop subdivision scheme.");
+			"Found a PolyGroup with multiple boundaries, which is not supported.\nTool can only use Loop subdivision scheme.\nUse the PolyGroups or Select Tool to assign PolyGroups.");
 		return false;
 	}
 
 	if (CheckResult == FSubdividePoly::ETopologyCheckResult::DegeneratePolygroup)
 	{
 		Message = LOCTEXT("DegenerateGroupPolygon",
-			"One PolyGroup has fewer than three boundary edges.\nUse the PolyGroups or Select Tool to assign/fix PolyGroups.\nTool will be limited to Loop subdivision scheme.");
+			"One PolyGroup has fewer than three boundary edges.\nTool can only use Loop subdivision scheme.\n"
+			"Use the PolyGroups or Select Tool to assign/fix PolyGroups, or use the \"Extra Corner\" settings to break up edges.");
 		return false;
 	}
 
@@ -138,16 +156,14 @@ void USubdividePolyTool::CapSubdivisionLevel(ESubdivisionScheme Scheme, int Desi
 	}
 	else
 	{
-		constexpr bool bAutoCompute = true;
-		FGroupTopology Topo(OriginalMesh.Get(), bAutoCompute);
-		NumOriginalFaces = Topo.Groups.Num();
+		NumOriginalFaces = Topology->Groups.Num();
 	}
 	int MaxLevel = (int)floor(log2(MaxFaces / (NumOriginalFaces+1)) / 2.0);
 
-	FText WarningText;
+	CappedSubdivisionMessage = FText();
 	if (DesiredLevel > MaxLevel)
 	{
-		WarningText = FText::Format(LOCTEXT("SubdivisionLevelTooHigh", "Subdivision level clamped: desired subdivision level ({0}) exceeds maximum level ({1}) for a mesh with this number of faces."),
+		CappedSubdivisionMessage = FText::Format(LOCTEXT("SubdivisionLevelTooHigh", "Subdivision level clamped: desired subdivision level ({0}) exceeds maximum level ({1}) for a mesh with this number of faces."),
 										  FText::AsNumber(DesiredLevel),
 										  FText::AsNumber(MaxLevel));
 		
@@ -155,14 +171,28 @@ void USubdividePolyTool::CapSubdivisionLevel(ESubdivisionScheme Scheme, int Desi
 		Properties->SubdivisionLevel = MaxLevel;
 		Properties->SilentUpdateWatched();		// Don't trigger this function again due to setting SubdivisionLevel above
 	}
-	
-	if (!PersistentErrorMessage.IsEmpty())
+}
+
+ESubdivisionScheme USubdividePolyTool::GetSubdivisionSchemeToUse()
+{
+	return Properties->bOverriddenSubdivisionScheme ? ESubdivisionScheme::Loop : Properties->SubdivisionScheme;
+}
+
+void USubdividePolyTool::UpdateDisplayedMessage()
+{
+	FText Message;
+	if (Properties->bOverriddenSubdivisionScheme && Properties->SubdivisionScheme != ESubdivisionScheme::Loop)
 	{
-		FText Delimiter = FText::FromString("\n");
-		WarningText = WarningText.IsEmpty() ?  PersistentErrorMessage : FText::Join(Delimiter, PersistentErrorMessage, WarningText);
+		Message = OverriddenSchemeMessage;
 	}
-	// if WarningText is empty this will clear possible lingering warning message 
-	GetToolManager()->DisplayMessage(WarningText, EToolMessageLevel::UserWarning);
+
+	FText Delimiter = FText::FromString("\n");
+	if (!CappedSubdivisionMessage.IsEmpty())
+	{
+		Message = FText::Join(Delimiter, Message, CappedSubdivisionMessage);
+	}
+
+	GetToolManager()->DisplayMessage(Message, EToolMessageLevel::UserWarning);
 }
 
 
@@ -178,23 +208,29 @@ void USubdividePolyTool::Setup()
 
 	GetToolManager()->DisplayMessage(LOCTEXT("SubdividePolyToolMessage", "Set the subdivision level and hit Accept to create a new subdivided mesh"), EToolMessageLevel::UserNotification);
 
+	Properties = NewObject<USubdividePolyToolProperties>(this, TEXT("Subdivide Mesh Tool Settings"));
+	Properties->RestoreProperties(this);
+
 	bool bWantVertexNormals = false;
 	OriginalMesh = MakeShared<FDynamicMesh3, ESPMode::ThreadSafe>(bWantVertexNormals, false, false, false);
 	FMeshDescriptionToDynamicMesh Converter;
 	Converter.Convert(UE::ToolTarget::GetMeshDescription(Target), *OriginalMesh);
 
-	const bool bCatmullClarkOK = CheckGroupTopology(PersistentErrorMessage);
-
-	if (!bCatmullClarkOK)
+	Topology = MakeShared<FGroupTopology>(OriginalMesh.Get(), false);
+	auto ShouldAddExtraCornerAtVert = [this](const FGroupTopology& GroupTopology, int32 Vid, const FIndex2i& AttachedGroupEdgeEids)
 	{
-		GetToolManager()->DisplayMessage(PersistentErrorMessage, EToolMessageLevel::UserWarning);
-	}
-	
-	Properties = NewObject<USubdividePolyToolProperties>(this, TEXT("Subdivide Mesh Tool Settings"));
-	Properties->RestoreProperties(this);
+		return Properties->bAddExtraCorners && FGroupTopology::IsEdgeAngleSharp(GroupTopology.GetMesh(), Vid, AttachedGroupEdgeEids, ExtraCornerDotProductThreshold);
+	};
+	Topology->ShouldAddExtraCornerAtVert = ShouldAddExtraCornerAtVert;
+
+	auto UpdateExtraCornerThreshold = [this]() { ExtraCornerDotProductThreshold = FMathd::Cos(Properties->ExtraCornerAngleThresholdDegrees * FMathd::DegToRad); };
+	UpdateExtraCornerThreshold();
 
-	Properties->bCatmullClarkOK = bCatmullClarkOK;
-	if (!bCatmullClarkOK)
+	Topology->RebuildTopology();
+
+	Properties->bOverriddenSubdivisionScheme = !CheckGroupTopology(OverriddenSchemeMessage);
+
+	if (Properties->bOverriddenSubdivisionScheme)
 	{
 		Properties->SubdivisionScheme = ESubdivisionScheme::Loop;
 	}
@@ -225,12 +261,7 @@ void USubdividePolyTool::Setup()
 		Properties->SubdivisionLevel = 1;
 	}
 	
-	CapSubdivisionLevel(Properties->SubdivisionScheme, Properties->SubdivisionLevel);
-	PreviewDynamicMeshComponent->SetRenderMeshPostProcessor(MakeUnique<SubdivPostProcessor>(Properties->SubdivisionLevel,
-																							Properties->SubdivisionScheme,
-																							Properties->NormalComputationMethod,
-																							Properties->UVComputationMethod,
-																							Properties->bNewPolyGroups));
+	CapSubdivisionLevel(GetSubdivisionSchemeToUse(), Properties->SubdivisionLevel);
 
 	// Use the input mesh's material on the preview
 	FComponentMaterialSet MaterialSet;
@@ -248,29 +279,35 @@ void USubdividePolyTool::Setup()
 	}
 
 	// dynamic mesh configuration settings
-	auto RebuildMeshPostProcessor = [this]()
+	auto RebuildMeshPostProcessor = [this, ShouldAddExtraCornerAtVert]()
 	{
 		UDynamicMeshComponent* PreviewDynamicMeshComponent = (UDynamicMeshComponent*)PreviewMesh->GetRootComponent();
-		PreviewDynamicMeshComponent->SetRenderMeshPostProcessor(MakeUnique<SubdivPostProcessor>(Properties->SubdivisionLevel,
-																								Properties->SubdivisionScheme,
-																								Properties->NormalComputationMethod,
-																								Properties->UVComputationMethod,
-																								Properties->bNewPolyGroups));
+		PreviewDynamicMeshComponent->SetRenderMeshPostProcessor(MakeUnique<SubdivPostProcessor>(
+			Properties->SubdivisionLevel,
+			GetSubdivisionSchemeToUse(),
+			Properties->BoundaryScheme,
+			Properties->NormalComputationMethod,
+			Properties->UVComputationMethod,
+			Properties->bNewPolyGroups,
+			ShouldAddExtraCornerAtVert));
 		PreviewDynamicMeshComponent->NotifyMeshUpdated();
 	};
+	RebuildMeshPostProcessor();
 
 	// Watch for property changes
 	Properties->WatchProperty(Properties->SubdivisionLevel, [this, RebuildMeshPostProcessor](int NewSubdLevel)
 	{
-		CapSubdivisionLevel(Properties->SubdivisionScheme, NewSubdLevel);
+		CapSubdivisionLevel(GetSubdivisionSchemeToUse(), NewSubdLevel);
+		UpdateDisplayedMessage();
 		RebuildMeshPostProcessor();
 	});
 
 	Properties->WatchProperty(Properties->SubdivisionScheme, [this, RebuildMeshPostProcessor](ESubdivisionScheme NewScheme)
 	{
-		CapSubdivisionLevel(NewScheme, Properties->SubdivisionLevel);
+		CapSubdivisionLevel(GetSubdivisionSchemeToUse(), Properties->SubdivisionLevel);
+		UpdateDisplayedMessage();
 		RebuildMeshPostProcessor();
-		bPreviewGeometryNeedsUpdate = true;		// Switch from rendering poly cage to all triangle edges
+		bPreviewGeometryNeedsUpdate = true;
 	});
 
 	Properties->WatchProperty(Properties->NormalComputationMethod, [this, RebuildMeshPostProcessor](ESubdivisionOutputNormals)
@@ -285,7 +322,28 @@ void USubdividePolyTool::Setup()
 	{
 		RebuildMeshPostProcessor();
 	});
+	Properties->WatchProperty(Properties->BoundaryScheme, [this, RebuildMeshPostProcessor](ESubdivisionBoundaryScheme)
+	{
+		RebuildMeshPostProcessor();
+	});
+
+	auto OnCornerChange = [this, UpdateExtraCornerThreshold, RebuildMeshPostProcessor]() {
+		UpdateExtraCornerThreshold();
+		Topology->RebuildTopology();
+
+		Properties->bOverriddenSubdivisionScheme = !CheckGroupTopology(OverriddenSchemeMessage);
+		UpdateDisplayedMessage();
+		NotifyOfPropertyChangeByTool(Properties);
 
+		bPreviewGeometryNeedsUpdate = true;
+		RebuildMeshPostProcessor();
+	};
+	Properties->WatchProperty(Properties->bAddExtraCorners, [this, OnCornerChange](bool) {
+		OnCornerChange();
+	});
+	Properties->WatchProperty(Properties->ExtraCornerAngleThresholdDegrees, [this, OnCornerChange](double) {
+		OnCornerChange();
+	});
 
 	auto RenderGroupsChanged = [this](bool bNewRenderGroups)
 	{
@@ -323,6 +381,9 @@ void USubdividePolyTool::Setup()
 
 	TargetComponent->SetOwnerVisibility(false);
 	PreviewMesh->SetVisible(true);
+
+	Properties->SilentUpdateWatched();
+	UpdateDisplayedMessage();
 }
 
 void USubdividePolyTool::CreateOrUpdatePreviewGeometry()
@@ -334,7 +395,7 @@ void USubdividePolyTool::CreateOrUpdatePreviewGeometry()
 		return;
 	}
 
-	if (Properties->SubdivisionScheme == ESubdivisionScheme::Loop)
+	if (GetSubdivisionSchemeToUse() == ESubdivisionScheme::Loop)
 	{
 		int NumEdges = OriginalMesh->EdgeCount();
 
@@ -359,16 +420,15 @@ void USubdividePolyTool::CreateOrUpdatePreviewGeometry()
 	}
 	else
 	{
-		FGroupTopology Topology(OriginalMesh.Get(), true);
-		int NumEdges = Topology.Edges.Num();
+		int NumEdges = Topology->Edges.Num();
 
 		PreviewGeometry->RemoveLineSet(TEXT("AllEdges"));
 
 		PreviewGeometry->CreateOrUpdateLineSet(TEXT("TopologyEdges"),
 											   NumEdges,
-											   [&Topology, this](int32 Index, TArray<FRenderableLine>& LinesOut)
+											   [this](int32 Index, TArray<FRenderableLine>& LinesOut)
 		{
-			const FGroupTopology::FGroupEdge& Edge = Topology.Edges[Index];
+			const FGroupTopology::FGroupEdge& Edge = Topology->Edges[Index];
 			FIndex2i EdgeCorners = Edge.EndpointCorners;
 
 			if (EdgeCorners[0] == FDynamicMesh3::InvalidID || EdgeCorners[1] == FDynamicMesh3::InvalidID)
@@ -376,8 +436,8 @@ void USubdividePolyTool::CreateOrUpdatePreviewGeometry()
 				return;
 			}
 
-			FIndex2i EdgeVertices{ Topology.Corners[EdgeCorners[0]].VertexID,
-									Topology.Corners[EdgeCorners[1]].VertexID };
+			FIndex2i EdgeVertices{ Topology->Corners[EdgeCorners[0]].VertexID,
+									Topology->Corners[EdgeCorners[1]].VertexID };
 			FVector A = (FVector)OriginalMesh->GetVertex(EdgeVertices[0]);
 			FVector B = (FVector)OriginalMesh->GetVertex(EdgeVertices[1]);
 
diff --git a/Engine/Plugins/Experimental/MeshModelingToolsetExp/Source/MeshModelingToolsEditorOnlyExp/Public/SubdividePolyTool.h b/Engine/Plugins/Experimental/MeshModelingToolsetExp/Source/MeshModelingToolsEditorOnlyExp/Public/SubdividePolyTool.h
--- a/Engine/Plugins/Experimental/MeshModelingToolsetExp/Source/MeshModelingToolsEditorOnlyExp/Public/SubdividePolyTool.h
+++ b/Engine/Plugins/Experimental/MeshModelingToolsetExp/Source/MeshModelingToolsEditorOnlyExp/Public/SubdividePolyTool.h
@@ -42,13 +42,13 @@ public:
 	UPROPERTY(EditAnywhere, Category=Settings, meta = (UIMin = "1", ClampMin = "1", Delta = 1, LinearDeltaSensitivity = 50))
 	int SubdivisionLevel = 3;
 
-	// Controls whether the user can select Catmull-Clark or is forced to use Loop
-	UPROPERTY(meta = (TransientToolProperty))
-	bool bCatmullClarkOK = true;
-
-	UPROPERTY(EditAnywhere, Category = Settings, meta = (EditCondition = "bCatmullClarkOK", HideEditConditionToggle))
+	UPROPERTY(EditAnywhere, Category = Settings)
 	ESubdivisionScheme SubdivisionScheme = ESubdivisionScheme::CatmullClark;
 
+	// How to treat mesh boundaries
+	UPROPERTY(EditAnywhere, Category = Settings, meta = (EditCondition = "SubdivisionScheme != ESubdivisionScheme::Bilinear || bOverriddenSubdivisionScheme"))
+	ESubdivisionBoundaryScheme BoundaryScheme = ESubdivisionBoundaryScheme::SmoothCorners;
+
 	UPROPERTY(EditAnywhere, Category=Settings)
 	ESubdivisionOutputNormals NormalComputationMethod = ESubdivisionOutputNormals::Generated;
 
@@ -67,6 +67,17 @@ public:
 	UPROPERTY(EditAnywhere, Category = Settings)
 	bool bRenderCage = true;
 
+	/** When using the group topology for subdivision, whether to add extra corners at sharp group edge bends. */
+	UPROPERTY(EditAnywhere, Category = TopologyOptions, meta = (EditCondition = "SubdivisionScheme != ESubdivisionScheme::Loop"))
+	bool bAddExtraCorners = true;
+
+	/** How acute an angle between two edges needs to be to add an extra corner there when Add Extra Corners is true. */
+	UPROPERTY(EditAnywhere, Category = TopologyOptions, meta = (ClampMin = "0", ClampMax = "180", EditCondition = "SubdivisionScheme != ESubdivisionScheme::Loop && bAddExtraCorners"))
+	double ExtraCornerAngleThresholdDegrees = 135;
+
+	/** Shows whether the current subdivision scheme is overridden to be "Loop" because the group topology is unsuitable. */
+	UPROPERTY(VisibleAnywhere, Category = Settings, AdvancedDisplay, meta = (TransientToolProperty))
+	bool bOverriddenSubdivisionScheme = false;
 };
 
 
@@ -102,6 +113,8 @@ protected:
 	// Input mesh
 	TSharedPtr<UE::Geometry::FDynamicMesh3, ESPMode::ThreadSafe> OriginalMesh;
 
+	TSharedPtr<UE::Geometry::FGroupTopology> Topology;
+
 	UPROPERTY()
 	TObjectPtr<UPreviewGeometry> PreviewGeometry = nullptr;
 
@@ -114,7 +127,10 @@ protected:
 
 	void CapSubdivisionLevel(ESubdivisionScheme Scheme, int DesiredLevel);
 
-	// Tool message when no explicit parameter-related warnings are being pushed.  
-	FText PersistentErrorMessage;
+	double ExtraCornerDotProductThreshold = 2;
+	ESubdivisionScheme GetSubdivisionSchemeToUse();
+	FText OverriddenSchemeMessage;
+	FText CappedSubdivisionMessage;
+	void UpdateDisplayedMessage();
 };
 
diff --git a/Engine/Plugins/Runtime/GeometryProcessing/Source/DynamicMesh/Private/GroupTopology.cpp b/Engine/Plugins/Runtime/GeometryProcessing/Source/DynamicMesh/Private/GroupTopology.cpp
--- a/Engine/Plugins/Runtime/GeometryProcessing/Source/DynamicMesh/Private/GroupTopology.cpp
+++ b/Engine/Plugins/Runtime/GeometryProcessing/Source/DynamicMesh/Private/GroupTopology.cpp
@@ -936,6 +936,36 @@ void FGroupTopology::GetSelectedTriangles(const FGroupTopologySelection& Selecti
 	}
 }
 
+bool UE::Geometry::FGroupTopology::IsEdgeAngleSharp(const FDynamicMesh3* Mesh, int32 SharedVid, const FIndex2i& AttachedEids, double DotThreshold)
+{
+	if (!ensure(Mesh && Mesh->IsEdge(AttachedEids.A) && Mesh->IsEdge(AttachedEids.B)))
+	{
+		return false;
+	}
+
+	// Gets vector pointing from the Vid along the edge.
+	auto GetEdgeUnitVector = [Mesh, SharedVid](int32 Eid, FVector3d& VectorOut)->bool
+	{
+		FIndex2i EdgeVids = Mesh->GetEdgeV(Eid);
+		// Make sure that the Vid is at EdgeVids.A
+		if (EdgeVids.B == SharedVid)
+		{
+			Swap(EdgeVids.A, EdgeVids.B);
+		}
+		VectorOut = Mesh->GetVertex(EdgeVids.B) - Mesh->GetVertex(EdgeVids.A);
+		return VectorOut.Normalize(KINDA_SMALL_NUMBER);
+	};
+
+	FVector Edge1, Edge2;
+	if (!GetEdgeUnitVector(AttachedEids.A, Edge1) || !GetEdgeUnitVector(AttachedEids.B, Edge2))
+	{
+		// If either edge was degenerate, we don't want to consider the angle "sharp"
+		return false;
+	}
+
+	return Edge1.Dot(Edge2) >= DotThreshold;
+}
+
 
 
 void FGroupTopology::GetAllVertexGroups(int32 VertexID, TArray<int32>& GroupsOut) const
diff --git a/Engine/Plugins/Runtime/GeometryProcessing/Source/DynamicMesh/Public/GroupTopology.h b/Engine/Plugins/Runtime/GeometryProcessing/Source/DynamicMesh/Public/GroupTopology.h
--- a/Engine/Plugins/Runtime/GeometryProcessing/Source/DynamicMesh/Public/GroupTopology.h
+++ b/Engine/Plugins/Runtime/GeometryProcessing/Source/DynamicMesh/Public/GroupTopology.h
@@ -437,6 +437,16 @@ public:
 	 */
 	void GetSelectedTriangles(const FGroupTopologySelection& Selection, TArray<int32>& Triangles) const;
 
+	/**
+	 * Helper function for use with ShouldAddExtraCornerAtVert. Checks that the dot product of the two mesh edge
+	 * unit vectors starting from shared vertex is greater than or equal to the given threshold. Assumes that SharedVid 
+	 * actually is shared by the two edges, and returns false if either edge is degenerate.
+	 * 
+	 * @param DotThreshold Threshold usually set equal to FMathd::Cos(ThresholdDegrees * FMathd::DegToRad). The
+	 *   parameter is not degrees just to avoid that calculation on each invocation.
+	 */ 
+	static bool IsEdgeAngleSharp(const FDynamicMesh3* Mesh, int32 SharedVid, const FIndex2i& AttachedGroupEdgeEids, double DotThreshold);
+
 protected:
 	const FDynamicMesh3* Mesh = nullptr;
 	const FDynamicMeshPolygroupAttribute* GroupLayer = nullptr;
diff --git a/Engine/Plugins/Runtime/MeshModelingToolset/Source/MeshModelingTools/Private/EditMeshPolygonsTool.cpp b/Engine/Plugins/Runtime/MeshModelingToolset/Source/MeshModelingTools/Private/EditMeshPolygonsTool.cpp
--- a/Engine/Plugins/Runtime/MeshModelingToolset/Source/MeshModelingTools/Private/EditMeshPolygonsTool.cpp
+++ b/Engine/Plugins/Runtime/MeshModelingToolset/Source/MeshModelingTools/Private/EditMeshPolygonsTool.cpp
@@ -358,43 +358,11 @@ void UEditMeshPolygonsTool::Setup()
 		Topology->ShouldAddExtraCornerAtVert = 
 			[this](const FGroupTopology& GroupTopology, int32 Vid, const FIndex2i& AttachedGroupEdgeEids)
 		{
-			if (!TopologyProperties->bAddExtraCorners)
-			{
-				return false;
-			}
-
 			// Note: it's important that we don't use CurrentMesh here. It's possible that an activity might create a copy of 
 			// the topology that uses the same corner forcing function but points to a different mesh, so we want to use
 			// whatever mesh the passed-in topology uses.
-			const FDynamicMesh3* Mesh = GroupTopology.GetMesh();
-			
-			if (!ensure(Mesh->IsEdge(AttachedGroupEdgeEids.A) && Mesh->IsEdge(AttachedGroupEdgeEids.B)))
-			{
-				return false;
-			}
-
-			// Gets vector pointing from the Vid along the edge.
-			auto GetEdgeUnitVector = [Mesh, Vid](int32 Eid, FVector3d& VectorOut)->bool
-			{
-				FIndex2i EdgeVids = Mesh->GetEdgeV(Eid);
-				// Make sure that the Vid is at EdgeVids.A
-				if (EdgeVids.B == Vid)
-				{
-					Swap(EdgeVids.A, EdgeVids.B);
-				}
-				VectorOut = Mesh->GetVertex(EdgeVids.B) - Mesh->GetVertex(EdgeVids.A);
-				return VectorOut.Normalize(KINDA_SMALL_NUMBER);
-			};
-
-			FVector Edge1, Edge2;
-			if (!GetEdgeUnitVector(AttachedGroupEdgeEids.A, Edge1) || !GetEdgeUnitVector(AttachedGroupEdgeEids.B, Edge2))
-			{
-				// If either edge was degenerate, we won't consider this a corner because otherwise we will end up
-				// with two corners connected by the degenerate edge, which is not ideal.
-				return false;
-			}
-
-			return Edge1.Dot(Edge2) >= ExtraCornerDotProductThreshold;
+			return TopologyProperties->bAddExtraCorners 
+				&& FGroupTopology::IsEdgeAngleSharp(GroupTopology.GetMesh(), Vid, AttachedGroupEdgeEids, ExtraCornerDotProductThreshold);
 		};
 	}
 
@@ -405,6 +373,10 @@ void UEditMeshPolygonsTool::Setup()
 	auto UpdateExtraCornerThreshold = [this]() { ExtraCornerDotProductThreshold = FMathd::Cos(TopologyProperties->ExtraCornerAngleThresholdDegrees * FMathd::DegToRad); };
 	UpdateExtraCornerThreshold();
 	TopologyProperties->WatchProperty(TopologyProperties->ExtraCornerAngleThresholdDegrees,
+		// Note: it may seem tempting to auto-rebuild the topology as the user drags the threshold slider (rather than waiting for
+		// the button click or next topology rebuild), but we have to transact corner additions/removals so that selection events in the 
+		// undo stack are able to refer to the same edges/corners (this is also why we store the current extra corners in FEditMeshPolygonsToolMeshChange).
+		// In an ideal world, we would know the end of a slider drag and only transact at that point, but we don't have that.
 		[this, UpdateExtraCornerThreshold](double) { UpdateExtraCornerThreshold(); });
 
 	Topology->RebuildTopology();
diff --git a/Engine/Plugins/Runtime/MeshModelingToolset/Source/ModelingComponentsEditorOnly/Private/Operations/SubdividePoly.cpp b/Engine/Plugins/Runtime/MeshModelingToolset/Source/ModelingComponentsEditorOnly/Private/Operations/SubdividePoly.cpp
--- a/Engine/Plugins/Runtime/MeshModelingToolset/Source/ModelingComponentsEditorOnly/Private/Operations/SubdividePoly.cpp
+++ b/Engine/Plugins/Runtime/MeshModelingToolset/Source/ModelingComponentsEditorOnly/Private/Operations/SubdividePoly.cpp
@@ -631,7 +631,28 @@ bool FSubdividePoly::ComputeTopologySubdivision()
 	using RefinerFactory = OpenSubdiv::Far::TopologyRefinerFactory<OpenSubdiv::Far::TopologyDescriptor>;
 	RefinerFactory::Options RefinerOptions;
 
-	RefinerOptions.schemeOptions.SetVtxBoundaryInterpolation(OpenSubdiv::Sdc::Options::VTX_BOUNDARY_EDGE_AND_CORNER);
+	OpenSubdiv::Sdc::Options::VtxBoundaryInterpolation BoundaryInterpolation = OpenSubdiv::Sdc::Options::VTX_BOUNDARY_EDGE_ONLY;
+	switch (BoundaryScheme)
+	{
+	case ESubdivisionBoundaryScheme::SmoothCorners:
+		BoundaryInterpolation = OpenSubdiv::Sdc::Options::VTX_BOUNDARY_EDGE_ONLY;
+		break;
+	case ESubdivisionBoundaryScheme::SharpCorners:
+		BoundaryInterpolation = OpenSubdiv::Sdc::Options::VTX_BOUNDARY_EDGE_AND_CORNER;
+		break;
+
+	// Note that supporting the below option is a little bit more work than just uncommenting.
+	// This option tags the boundary faces as holes, so we have to use the FTopologyLevel::IsFaceHole
+	// function where we output the faces, and then probably remove unused verts on the boundary.
+	//case ESubdivisionBoundaryScheme::NoBoundaryFaces:
+	//	BoundaryInterpolation = OpenSubdiv::Sdc::Options::VTX_BOUNDARY_NONE;
+	//	break;
+
+	default:
+		ensure(false);
+		break;
+	}
+	RefinerOptions.schemeOptions.SetVtxBoundaryInterpolation(BoundaryInterpolation);
 
 	switch (SubdivisionScheme)
 	{
@@ -1003,11 +1024,6 @@ FSubdividePoly::ETopologyCheckResult FSubdividePoly::ValidateTopology()
 		return ETopologyCheckResult::NoGroups;
 	}
 
-	if (GroupTopology.Groups.Num() < 2)
-	{
-		return ETopologyCheckResult::InsufficientGroups;
-	}
-
 	for (const FGroupTopology::FGroup& Group : GroupTopology.Groups)
 	{
 		if (Group.Boundaries.Num() == 0)
@@ -1029,6 +1045,12 @@ FSubdividePoly::ETopologyCheckResult FSubdividePoly::ValidateTopology()
 		}
 	}
 
+	// May be necessary if we ever support this option:
+	//if (BoundaryScheme == ESubdivisionBoundaryScheme::NoBoundaryFaces)
+	//{
+	//	// Verify that there is at least one non-boundary face.
+	//}
+
 	return ETopologyCheckResult::Ok;
 }
 
diff --git a/Engine/Plugins/Runtime/MeshModelingToolset/Source/ModelingComponentsEditorOnly/Public/Operations/SubdividePoly.h b/Engine/Plugins/Runtime/MeshModelingToolset/Source/ModelingComponentsEditorOnly/Public/Operations/SubdividePoly.h
--- a/Engine/Plugins/Runtime/MeshModelingToolset/Source/ModelingComponentsEditorOnly/Public/Operations/SubdividePoly.h
+++ b/Engine/Plugins/Runtime/MeshModelingToolset/Source/ModelingComponentsEditorOnly/Public/Operations/SubdividePoly.h
@@ -16,11 +16,36 @@ class FProgressCancel;
 UENUM()
 enum class ESubdivisionScheme : uint8
 {
+	// Subdivides like Catmull-Clark, but does not smooth the result (i.e. the vertices remain
+	// in their original planes).
 	Bilinear,
+	// Common subdivision scheme typically used with quad-dominant meshes. In our tools, it uses
+	// the group topology to build a cage that serves as the starting point.
 	CatmullClark UMETA(DisplayName = "Catmull-Clark"),
+	// Subdivision scheme developed by Charles Loop that operates directly on triangle meshes, and
+	// therefore does not use the group topology. 
 	Loop
 };
 
+UENUM()
+enum class ESubdivisionBoundaryScheme : uint8
+{
+	// Corners with only one adjoining face are smoothed along with the rest of the mesh boundary. For
+	// example, a square patch will have rounded corners.
+	SmoothCorners,
+	// Corners with only one adjoining face are constrained to be passed through. For example, this
+	// will keep the corners of a vertically deformed square patch sharp even as the rest of the patch
+	// is vertically smoothed.
+	SharpCorners,
+
+	//~ Not currently a supported option, and might not be one that is worth supporting since the same
+	//~ result can be achieved by deleting faces on the boundary after subdividing...
+	// Points along the boundary are treated purely as control points with no faces extending to them. This
+	// is not commonly used, except when trying to seamlessly join different patches by replicating vertices.
+	//NoBoundaryFaces
+};
+
+
 
 /**
  * Method for computing normals on the output mesh
@@ -85,6 +110,8 @@ public:
 
 	ESubdivisionScheme SubdivisionScheme = ESubdivisionScheme::CatmullClark;
 
+	ESubdivisionBoundaryScheme BoundaryScheme = ESubdivisionBoundaryScheme::SharpCorners;
+
 	ESubdivisionOutputNormals NormalComputationMethod = ESubdivisionOutputNormals::Generated;
 
 	ESubdivisionOutputUVs UVComputationMethod = ESubdivisionOutputUVs::Interpolated;
