Added MeshAutoUV to GeometryProcessingInterface
* Required for future changes to improve the mesh merging capability to cope with missing/invalid UVs when trying to perform material baking.

diff --git a/Engine/Plugins/Experimental/MeshModelingToolsetExp/Source/GeometryProcessingAdapters/Private/GeometryProcessingAdaptersModule.cpp b/Engine/Plugins/Experimental/MeshModelingToolsetExp/Source/GeometryProcessingAdapters/Private/GeometryProcessingAdaptersModule.cpp
--- a/Engine/Plugins/Experimental/MeshModelingToolsetExp/Source/GeometryProcessingAdapters/Private/GeometryProcessingAdaptersModule.cpp
+++ b/Engine/Plugins/Experimental/MeshModelingToolsetExp/Source/GeometryProcessingAdapters/Private/GeometryProcessingAdaptersModule.cpp
@@ -3,6 +3,10 @@
 #include "GeometryProcessingAdaptersModule.h"
 #include "Features/IModularFeatures.h"
 
+#include "GeometryProcessing/ApproximateActorsImpl.h"
+#include "GeometryProcessing/CombineMeshInstancesImpl.h"
+#include "GeometryProcessing/MeshAutoUVImpl.h"
+
 #define LOCTEXT_NAMESPACE "FGeometryProcessingAdaptersModule"
 
 using namespace UE::Geometry;
@@ -22,6 +26,12 @@ void FGeometryProcessingAdaptersModule::StartupModule()
 	{
 		IModularFeatures::Get().RegisterModularFeature(IGeometryProcessing_CombineMeshInstances::GetModularFeatureName(), CombineMeshInstances.Get());
 	}
+
+	MeshAutoUV = MakeShared<FMeshAutoUVImpl>();
+	if (CombineMeshInstances.IsValid())
+	{
+		IModularFeatures::Get().RegisterModularFeature(IGeometryProcessing_MeshAutoUV::GetModularFeatureName(), MeshAutoUV.Get());
+	}
 }
 
 void FGeometryProcessingAdaptersModule::ShutdownModule()
@@ -40,6 +50,12 @@ void FGeometryProcessingAdaptersModule::ShutdownModule()
 		IModularFeatures::Get().UnregisterModularFeature(IGeometryProcessing_CombineMeshInstances::GetModularFeatureName(), CombineMeshInstances.Get());
 		CombineMeshInstances = nullptr;
 	}
+
+	if (MeshAutoUV.IsValid())
+	{
+		IModularFeatures::Get().UnregisterModularFeature(IGeometryProcessing_MeshAutoUV::GetModularFeatureName(), MeshAutoUV.Get());
+		MeshAutoUV = nullptr;
+	}
 }
 
 #undef LOCTEXT_NAMESPACE
diff --git a/Engine/Plugins/Experimental/MeshModelingToolsetExp/Source/GeometryProcessingAdapters/Private/MeshAutoUVImpl.cpp b/Engine/Plugins/Experimental/MeshModelingToolsetExp/Source/GeometryProcessingAdapters/Private/MeshAutoUVImpl.cpp
new file mode 100644
--- /dev/null
+++ b/Engine/Plugins/Experimental/MeshModelingToolsetExp/Source/GeometryProcessingAdapters/Private/MeshAutoUVImpl.cpp
@@ -0,0 +1,69 @@
+// Copyright Epic Games, Inc. All Rights Reserved.
+
+#include "GeometryProcessing/MeshAutoUVImpl.h"
+#include "MeshDescriptionToDynamicMesh.h"
+#include "DynamicMeshToMeshDescription.h"
+#include "ParameterizationOps/ParameterizeMeshOp.h"
+
+
+using namespace UE::Geometry;
+
+
+IGeometryProcessing_MeshAutoUV::FOptions FMeshAutoUVImpl::ConstructDefaultOptions()
+{
+	return IGeometryProcessing_MeshAutoUV::FOptions();
+}
+
+void FMeshAutoUVImpl::GenerateUVs(FMeshDescription& InOutMesh, const IGeometryProcessing_MeshAutoUV::FOptions& Options, IGeometryProcessing_MeshAutoUV::FResults& ResultsOut)
+{
+	TSharedPtr<UE::Geometry::FDynamicMesh3, ESPMode::ThreadSafe> DynamicMesh = MakeShared<UE::Geometry::FDynamicMesh3, ESPMode::ThreadSafe>();
+	FMeshDescriptionToDynamicMesh MeshDescriptionToDynamicMesh;
+	MeshDescriptionToDynamicMesh.Convert(&InOutMesh, *DynamicMesh);
+
+	UE::Geometry::FParameterizeMeshOp ParameterizeMeshOp;
+	ParameterizeMeshOp.InputMesh = DynamicMesh;
+	ParameterizeMeshOp.Stretch = Options.UVAtlasStretch;
+	ParameterizeMeshOp.NumCharts = Options.UVAtlasNumCharts;
+	ParameterizeMeshOp.XAtlasMaxIterations = Options.XAtlasMaxIterations;
+
+	ParameterizeMeshOp.InitialPatchCount = Options.NumInitialPatches;
+	ParameterizeMeshOp.PatchCurvatureAlignmentWeight = Options.CurvatureAlignment;
+	ParameterizeMeshOp.PatchMergingMetricThresh = Options.MergingThreshold;
+	ParameterizeMeshOp.PatchMergingAngleThresh = Options.MaxAngleDeviationDeg;
+	ParameterizeMeshOp.ExpMapNormalSmoothingSteps = Options.SmoothingSteps;
+	ParameterizeMeshOp.ExpMapNormalSmoothingAlpha = Options.SmoothingAlpha;
+
+	ParameterizeMeshOp.bEnablePacking = Options.bAutoPack;
+	ParameterizeMeshOp.Width = ParameterizeMeshOp.Height = Options.PackingTargetWidth;
+
+	if (Options.Method == IGeometryProcessing_MeshAutoUV::EAutoUVMethod::UVAtlas)
+	{
+		ParameterizeMeshOp.Method = UE::Geometry::EParamOpBackend::UVAtlas;
+	}
+	else if (Options.Method == IGeometryProcessing_MeshAutoUV::EAutoUVMethod::XAtlas)
+	{
+		ParameterizeMeshOp.Method = UE::Geometry::EParamOpBackend::XAtlas;
+	}
+	else
+	{
+		ParameterizeMeshOp.Method = UE::Geometry::EParamOpBackend::PatchBuilder;
+	}
+
+	ParameterizeMeshOp.CalculateResult(nullptr);
+
+	TUniquePtr<UE::Geometry::FDynamicMesh3> ResultMesh = ParameterizeMeshOp.ExtractResult();
+
+	ResultsOut.ResultCode = IGeometryProcessing_MeshAutoUV::EResultCode::Success;
+
+	FConversionToMeshDescriptionOptions ConversionOptions;
+	ConversionOptions.bSetPolyGroups = false;
+	ConversionOptions.bUpdatePositions = false;
+	ConversionOptions.bUpdateNormals = false;
+	ConversionOptions.bUpdateTangents = false;
+	ConversionOptions.bUpdateUVs = true;
+	ConversionOptions.bUpdateVtxColors = false;
+	ConversionOptions.bTransformVtxColorsSRGBToLinear = false;
+
+	FDynamicMeshToMeshDescription Converter(ConversionOptions);
+	Converter.UpdateUsingConversionOptions(ResultMesh.Get(), InOutMesh);
+}
\ No newline at end of file
diff --git a/Engine/Plugins/Experimental/MeshModelingToolsetExp/Source/GeometryProcessingAdapters/Public/GeometryProcessing/MeshAutoUVImpl.h b/Engine/Plugins/Experimental/MeshModelingToolsetExp/Source/GeometryProcessingAdapters/Public/GeometryProcessing/MeshAutoUVImpl.h
new file mode 100644
--- /dev/null
+++ b/Engine/Plugins/Experimental/MeshModelingToolsetExp/Source/GeometryProcessingAdapters/Public/GeometryProcessing/MeshAutoUVImpl.h
@@ -0,0 +1,27 @@
+// Copyright Epic Games, Inc. All Rights Reserved.
+
+#pragma once
+
+#include "CoreMinimal.h"
+#include "GeometryProcessingInterfaces/MeshAutoUV.h"
+
+
+namespace UE
+{
+namespace Geometry
+{
+
+/**
+ * Implementation of IGeometryProcessing_MeshAutoUV
+ */
+class GEOMETRYPROCESSINGADAPTERS_API FMeshAutoUVImpl : public IGeometryProcessing_MeshAutoUV
+{
+public:
+	virtual FOptions ConstructDefaultOptions() override;
+
+	virtual void GenerateUVs(FMeshDescription& InOutMesh, const FOptions& Options, FResults& ResultsOut) override;
+};
+
+
+}
+}
diff --git a/Engine/Plugins/Experimental/MeshModelingToolsetExp/Source/GeometryProcessingAdapters/Public/GeometryProcessingAdaptersModule.h b/Engine/Plugins/Experimental/MeshModelingToolsetExp/Source/GeometryProcessingAdapters/Public/GeometryProcessingAdaptersModule.h
--- a/Engine/Plugins/Experimental/MeshModelingToolsetExp/Source/GeometryProcessingAdapters/Public/GeometryProcessingAdaptersModule.h
+++ b/Engine/Plugins/Experimental/MeshModelingToolsetExp/Source/GeometryProcessingAdapters/Public/GeometryProcessingAdaptersModule.h
@@ -5,8 +5,12 @@
 #include "CoreMinimal.h"
 #include "Modules/ModuleManager.h"
 
-#include "GeometryProcessing/ApproximateActorsImpl.h"
-#include "GeometryProcessing/CombineMeshInstancesImpl.h"
+namespace UE::Geometry
+{
+	class FApproximateActorsImpl;
+	class FCombineMeshInstancesImpl;
+	class FMeshAutoUVImpl;
+}
 
 class FGeometryProcessingAdaptersModule : public IModuleInterface
 {
@@ -20,4 +24,5 @@ public:
 protected:
 	TSharedPtr<UE::Geometry::FApproximateActorsImpl> ApproximateActors;
 	TSharedPtr<UE::Geometry::FCombineMeshInstancesImpl> CombineMeshInstances;
+	TSharedPtr<UE::Geometry::FMeshAutoUVImpl> MeshAutoUV;
 };
diff --git a/Engine/Source/Developer/GeometryProcessingInterfaces/Private/GeometryProcessingInterfacesModule.cpp b/Engine/Source/Developer/GeometryProcessingInterfaces/Private/GeometryProcessingInterfacesModule.cpp
--- a/Engine/Source/Developer/GeometryProcessingInterfaces/Private/GeometryProcessingInterfacesModule.cpp
+++ b/Engine/Source/Developer/GeometryProcessingInterfaces/Private/GeometryProcessingInterfacesModule.cpp
@@ -7,6 +7,7 @@
 
 #include "GeometryProcessingInterfaces/ApproximateActors.h"
 #include "GeometryProcessingInterfaces/CombineMeshInstances.h"
+#include "GeometryProcessingInterfaces/MeshAutoUV.h"
 
 
 IMPLEMENT_MODULE(FGeometryProcessingInterfacesModule, GeometryProcessingInterfaces);
@@ -17,37 +18,49 @@ void FGeometryProcessingInterfacesModule::StartupModule()
 
 }
 
-
 void FGeometryProcessingInterfacesModule::ShutdownModule()
 {
 	ApproximateActors = nullptr;
 	CombineMeshInstances = nullptr;
+	MeshAutoUV = nullptr;
 }
 
+namespace
+{
+	template <typename TModularFeatureInterface>
+	TModularFeatureInterface* GetModularFeatureImplementation()
+	{
+		TArray<TModularFeatureInterface*> AvailableImplementations =
+			IModularFeatures::Get().GetModularFeatureImplementations<TModularFeatureInterface>(TModularFeatureInterface::GetModularFeatureName());
+
+		return (AvailableImplementations.Num() > 0) ? AvailableImplementations[0] : nullptr;
+	}
+}
 
 IGeometryProcessing_ApproximateActors* FGeometryProcessingInterfacesModule::GetApproximateActorsImplementation()
 {
 	if (ApproximateActors == nullptr)
 	{
-		TArray<IGeometryProcessing_ApproximateActors*> ApproximateActorsOptions =
-			IModularFeatures::Get().GetModularFeatureImplementations<IGeometryProcessing_ApproximateActors>(IGeometryProcessing_ApproximateActors::GetModularFeatureName());
-
-		ApproximateActors = (ApproximateActorsOptions.Num() > 0) ? ApproximateActorsOptions[0] : nullptr;
+		ApproximateActors = GetModularFeatureImplementation<IGeometryProcessing_ApproximateActors>();
 	}
-
 	return ApproximateActors;
 }
 IGeometryProcessing_CombineMeshInstances* FGeometryProcessingInterfacesModule::GetCombineMeshInstancesImplementation()
 {
 	if (CombineMeshInstances == nullptr)
 	{
-		TArray<IGeometryProcessing_CombineMeshInstances*> CombineMeshInstancesOptions =
-			IModularFeatures::Get().GetModularFeatureImplementations<IGeometryProcessing_CombineMeshInstances>(IGeometryProcessing_CombineMeshInstances::GetModularFeatureName());
-
-		CombineMeshInstances = (CombineMeshInstancesOptions.Num() > 0) ? CombineMeshInstancesOptions[0] : nullptr;
+		CombineMeshInstances = GetModularFeatureImplementation<IGeometryProcessing_CombineMeshInstances>();
 	}
-
 	return CombineMeshInstances;
 }
 
+IGeometryProcessing_MeshAutoUV* FGeometryProcessingInterfacesModule::GetMeshAutoUVImplementation()
+{
+	if (MeshAutoUV == nullptr)
+	{
+		MeshAutoUV = GetModularFeatureImplementation<IGeometryProcessing_MeshAutoUV>();
+	}
+	return MeshAutoUV;
+}
+
 
diff --git a/Engine/Source/Developer/GeometryProcessingInterfaces/Private/GeometryProcessingInterfacesModule.h b/Engine/Source/Developer/GeometryProcessingInterfaces/Private/GeometryProcessingInterfacesModule.h
--- a/Engine/Source/Developer/GeometryProcessingInterfaces/Private/GeometryProcessingInterfacesModule.h
+++ b/Engine/Source/Developer/GeometryProcessingInterfaces/Private/GeometryProcessingInterfacesModule.h
@@ -7,6 +7,7 @@
 
 class IGeometryProcessing_ApproximateActors;
 class IGeometryProcessing_CombineMeshInstances;
+class IGeometryProcessing_MeshAutoUV;
 
 
 /**
@@ -20,7 +21,7 @@ public:
 	virtual void StartupModule() override;
 	virtual void ShutdownModule() override;
 	
-	
+
 	/**
 	 * @return implementation of IGeometryProcessing_ApproximateActors, if available, or nullptr (result is cached internally)
 	 */
@@ -31,8 +32,13 @@ public:
 	 */
 	virtual IGeometryProcessing_CombineMeshInstances* GetCombineMeshInstancesImplementation() override;
 
+	/**
+	 * @return implementation of IGeometryProcessing_MeshAutoUV, if available, or nullptr (result is cached internally)
+	 */
+	virtual IGeometryProcessing_MeshAutoUV* GetMeshAutoUVImplementation() override;
 
 private:
 	IGeometryProcessing_ApproximateActors* ApproximateActors = nullptr;
 	IGeometryProcessing_CombineMeshInstances* CombineMeshInstances = nullptr;
+	IGeometryProcessing_MeshAutoUV* MeshAutoUV = nullptr;
 };
diff --git a/Engine/Source/Developer/GeometryProcessingInterfaces/Public/GeometryProcessingInterfaces/MeshAutoUV.h b/Engine/Source/Developer/GeometryProcessingInterfaces/Public/GeometryProcessingInterfaces/MeshAutoUV.h
new file mode 100644
--- /dev/null
+++ b/Engine/Source/Developer/GeometryProcessingInterfaces/Public/GeometryProcessingInterfaces/MeshAutoUV.h
@@ -0,0 +1,72 @@
+// Copyright Epic Games, Inc. All Rights Reserved.
+
+#pragma once
+
+#include "CoreMinimal.h"
+#include "Containers/Array.h"
+#include "Features/IModularFeature.h"
+#include "Math/Vector2D.h"
+#include "MeshDescription.h"
+
+/**
+ * The CombineMeshInstances modular feature is used to provide a mechanism
+ * for merging a set of instances of meshes (ie mesh + transform + materials + ...)
+ * into a smaller set of meshes. Generally this involves creating simpler versions
+ * of the instances and appending them into one or a small number of combined meshes.
+ */
+class IGeometryProcessing_MeshAutoUV : public IModularFeature
+{
+public:
+	virtual ~IGeometryProcessing_MeshAutoUV() {}
+
+	enum class EAutoUVMethod : uint8
+	{
+		PatchBuilder = 0,
+		UVAtlas = 1,
+		XAtlas = 2
+	};
+
+	struct FOptions
+	{
+		EAutoUVMethod Method = EAutoUVMethod::PatchBuilder;
+
+		// UVAtlas parameters
+		double UVAtlasStretch = 0.11f;
+		int UVAtlasNumCharts = 0;
+
+		// XAtlas parameters
+		int XAtlasMaxIterations = 1;
+
+		// PatchBuilder parameters
+		int NumInitialPatches = 100;
+		double CurvatureAlignment = 1.0;
+		double MergingThreshold = 1.5;
+		double MaxAngleDeviationDeg = 45.0;
+		int SmoothingSteps = 5;
+		double SmoothingAlpha = 0.25;
+
+		bool bAutoPack = true;
+		int PackingTargetWidth = 512;
+	};
+
+	enum class EResultCode
+	{
+		Success,
+		UnknownError
+	};
+
+	struct FResults
+	{
+		EResultCode ResultCode = EResultCode::UnknownError;
+	};
+
+	virtual FOptions ConstructDefaultOptions() PURE_VIRTUAL(IGeometryProcessing_MeshAutoUV::ConstructDefaultOptions, return FOptions(); );
+
+	virtual void GenerateUVs(FMeshDescription& InOutMesh, const FOptions& Options, FResults& ResultsOut) PURE_VIRTUAL(IGeometryProcessing_MeshAutoUV::GenerateUVs,);
+
+	// Modular feature name to register for retrieval during runtime
+	static const FName GetModularFeatureName()
+	{
+		return TEXT("GeometryProcessing_MeshAutoUV");
+	}
+};
diff --git a/Engine/Source/Developer/GeometryProcessingInterfaces/Public/IGeometryProcessingInterfacesModule.h b/Engine/Source/Developer/GeometryProcessingInterfaces/Public/IGeometryProcessingInterfacesModule.h
--- a/Engine/Source/Developer/GeometryProcessingInterfaces/Public/IGeometryProcessingInterfacesModule.h
+++ b/Engine/Source/Developer/GeometryProcessingInterfaces/Public/IGeometryProcessingInterfacesModule.h
@@ -6,6 +6,7 @@
 
 class IGeometryProcessing_ApproximateActors;
 class IGeometryProcessing_CombineMeshInstances;
+class IGeometryProcessing_MeshAutoUV;
 
 /**
  * Abstract interface to a Module that provides functions to access
@@ -34,4 +35,9 @@ public:
 	 */
 	virtual IGeometryProcessing_CombineMeshInstances* GetCombineMeshInstancesImplementation() = 0;
 
+	/**
+	 * @return implementation of IGeometryProcessing_MeshAutoUV, if available, or nullptr (result is cached internally)
+	 */
+	virtual IGeometryProcessing_MeshAutoUV* GetMeshAutoUVImplementation() = 0;
+
 };
