Add option for pattern tool to randomly pick between elements rather than always using all elements

diff --git a/Engine/Plugins/Experimental/MeshModelingToolsetExp/Source/MeshModelingToolsExp/Private/PatternTool.cpp b/Engine/Plugins/Experimental/MeshModelingToolsetExp/Source/MeshModelingToolsExp/Private/PatternTool.cpp
--- a/Engine/Plugins/Experimental/MeshModelingToolsetExp/Source/MeshModelingToolsExp/Private/PatternTool.cpp
+++ b/Engine/Plugins/Experimental/MeshModelingToolsetExp/Source/MeshModelingToolsExp/Private/PatternTool.cpp
@@ -592,6 +592,7 @@ void UPatternTool::Setup()
 	Settings->WatchProperty(Settings->ProjectionOffset, [this](float NewValue) { MarkPatternDirty(); } );
 	Settings->WatchProperty(Settings->bHideSources, [this](bool bNewValue) { OnSourceVisibilityToggled(!bNewValue); } );
 	Settings->WatchProperty(Settings->bUseRelativeTransforms, [this](bool bNewValue) { MarkPatternDirty(); } );
+	Settings->WatchProperty(Settings->bRandomlyPickElements, [this](bool bNewValue) { MarkPatternDirty(); });
 
 	BoundingBoxSettings = NewObject<UPatternTool_BoundingBoxSettings>();
 	AddToolPropertySource(BoundingBoxSettings);
@@ -1451,41 +1452,68 @@ void UPatternTool::UpdatePattern()
 	int32 NumElements = Elements.Num();
 	check(PreviewComponents.Num() == NumElements);
 
-	for (int32 ElemIdx = 0; ElemIdx < NumElements; ++ElemIdx)
+	auto AddComponent = [this, &Pattern](int32 PatternItemIdx, const FPatternElement& Element, const FTransformSRT3d& ElementTransform, FComponentSet& ElemComponents)
 	{
-		FPatternElement& Element = Elements[ElemIdx];
-		FTransformSRT3d ElementTransform = Element.BaseRotateScale;
-		FComponentSet& ElemComponents = PreviewComponents[ElemIdx];
-
-		if (Settings->bUseRelativeTransforms)
+		UPrimitiveComponent* Component = nullptr;
+		if (Element.SourceDynamicMesh != nullptr)
 		{
-			ElementTransform.SetTranslation(ElementTransform.GetTranslation() + Element.RelativePosition);
+			Component = GetPreviewDynamicMesh(Element);
 		}
-		
+		else if (Element.SourceStaticMesh != nullptr)
+		{
+			Component = GetPreviewStaticMesh(Element);
+		}
+		if (Component != nullptr)
+		{
+			FTransform WorldTransform;
+			ComputeWorldTransform(WorldTransform, ElementTransform, Pattern[PatternItemIdx]);
+
+			Component->SetWorldTransform(WorldTransform);
+			Component->SetVisibility(true);
+
+			ElemComponents.Components.Add(Component);
+		}
+	};
+
+	if (Settings->bRandomlyPickElements)
+	{
+		FRandomStream Stream(Settings->Seed);
 		for (int32 k = 0; k < NumPatternItems; ++k)
 		{
-			UPrimitiveComponent* Component = nullptr;
-			if (Element.SourceDynamicMesh != nullptr)
+			int32 ElemIdx = Stream.RandHelper(NumElements);
+			FPatternElement& Element = Elements[ElemIdx];
+			FTransformSRT3d ElementTransform = Element.BaseRotateScale;
+			FComponentSet& ElemComponents = PreviewComponents[ElemIdx];
+
+			if (Settings->bUseRelativeTransforms)
 			{
-				Component = GetPreviewDynamicMesh(Element);
+				ElementTransform.SetTranslation(ElementTransform.GetTranslation() + Element.RelativePosition);
 			}
-			else if (Element.SourceStaticMesh != nullptr)
+
+			AddComponent(k, Element, ElementTransform, ElemComponents);
+		}
+	}
+	else // always pick all elements
+	{
+		for (int32 ElemIdx = 0; ElemIdx < NumElements; ++ElemIdx)
+		{
+			FPatternElement& Element = Elements[ElemIdx];
+			FTransformSRT3d ElementTransform = Element.BaseRotateScale;
+			FComponentSet& ElemComponents = PreviewComponents[ElemIdx];
+
+			if (Settings->bUseRelativeTransforms)
 			{
-				Component = GetPreviewStaticMesh(Element);
+				ElementTransform.SetTranslation(ElementTransform.GetTranslation() + Element.RelativePosition);
 			}
-			if (Component != nullptr)
+
+			for (int32 k = 0; k < NumPatternItems; ++k)
 			{
-				FTransform WorldTransform;
-				ComputeWorldTransform(WorldTransform, ElementTransform, Pattern[k]);
-				
-				Component->SetWorldTransform( WorldTransform );
-				Component->SetVisibility(true);
-		
-				ElemComponents.Components.Add(Component);
+				AddComponent(k, Element, ElementTransform, ElemComponents);
 			}
 		}
 	}
 
+
 	// ResetPreviews() does not hide the preview components because changing their visibility
 	// will destroy their SceneProxies, which is expensive, and in most cases the same set of
 	// objects will immediately be shown, re-creating the proxy. So instead they are just all
@@ -1558,7 +1586,7 @@ void UPatternTool::ComputeCombinedPatternBounds()
 	CombinedPatternBounds.Expand(BoundingBoxSettings->Adjustment);
 }
 
-UStaticMeshComponent* UPatternTool::GetPreviewStaticMesh(FPatternElement& Element)
+UStaticMeshComponent* UPatternTool::GetPreviewStaticMesh(const FPatternElement& Element)
 {
 	FComponentSet* FoundPool = StaticMeshPools.Find(Element.TargetIndex);
 	if (FoundPool == nullptr)
@@ -1614,7 +1642,7 @@ void UPatternTool::ReturnStaticMeshes(FPatternElement& Element, FComponentSet& C
 
 
 
-UDynamicMeshComponent* UPatternTool::GetPreviewDynamicMesh(FPatternElement& Element)
+UDynamicMeshComponent* UPatternTool::GetPreviewDynamicMesh(const FPatternElement& Element)
 {
 	FComponentSet* FoundPool = DynamicMeshPools.Find(Element.TargetIndex);
 	if (FoundPool == nullptr)
diff --git a/Engine/Plugins/Experimental/MeshModelingToolsetExp/Source/MeshModelingToolsExp/Public/PatternTool.h b/Engine/Plugins/Experimental/MeshModelingToolsetExp/Source/MeshModelingToolsExp/Public/PatternTool.h
--- a/Engine/Plugins/Experimental/MeshModelingToolsetExp/Source/MeshModelingToolsExp/Public/PatternTool.h
+++ b/Engine/Plugins/Experimental/MeshModelingToolsetExp/Source/MeshModelingToolsExp/Public/PatternTool.h
@@ -125,6 +125,10 @@ public:
 	/** If false, all pattern elements will be positioned at the origin of the first pattern element */
 	UPROPERTY(EditAnywhere, Category = General)
 	bool bUseRelativeTransforms = true;
+
+	/** Whether to randomly pick which source mesh is scattered at each location, or to always use all source meshes */
+	UPROPERTY(EditAnywhere, Category = General)
+	bool bRandomlyPickElements = false;
 	
 	/** Shape of the underlying Pattern */
 	UPROPERTY(EditAnywhere, Category = Shape)
@@ -594,11 +598,11 @@ protected:
 
 	
 	TMap<int32, FComponentSet> StaticMeshPools;
-	UStaticMeshComponent* GetPreviewStaticMesh(FPatternElement& Element);
+	UStaticMeshComponent* GetPreviewStaticMesh(const FPatternElement& Element);
 	void ReturnStaticMeshes(FPatternElement& Element, FComponentSet& ComponentSet);
 
 	TMap<int32, FComponentSet> DynamicMeshPools;
-	UDynamicMeshComponent* GetPreviewDynamicMesh(FPatternElement& Element);
+	UDynamicMeshComponent* GetPreviewDynamicMesh(const FPatternElement& Element);
 	void ReturnDynamicMeshes(FPatternElement& Element, FComponentSet& ComponentSet);
 
 	void HideReturnedPreviewMeshes();
