CombineMeshInstances: quality improvements. ComputePlanarPolygonApproximation now tries to snap vertices back to their original locations, which reduces cracks introduced by projecting the mesh being approximated to a plane. Split mesh retriangulation processes out of ComputeHiddenRemovalForLOD into OptimizeLODMeshTriangulation, now passing a geometric tolerance through OptimizeLODMeshTriangulation steps instead of using various hardcoded magic numbers, to provide some level of external control. Tolerances are now decided in BuildCombinedMesh, incrementally increases at each LOD based on SimplifyBaseTolerance as a starting point.

This change may result in small increase in output triangle count in some cases, but often because it was previously too aggressive in those cases and producing meshes with visible holes.

diff --git a/Engine/Plugins/Experimental/MeshModelingToolsetExp/Source/GeometryProcessingAdapters/Private/CombineMeshInstancesImpl.cpp b/Engine/Plugins/Experimental/MeshModelingToolsetExp/Source/GeometryProcessingAdapters/Private/CombineMeshInstancesImpl.cpp
--- a/Engine/Plugins/Experimental/MeshModelingToolsetExp/Source/GeometryProcessingAdapters/Private/CombineMeshInstancesImpl.cpp
+++ b/Engine/Plugins/Experimental/MeshModelingToolsetExp/Source/GeometryProcessingAdapters/Private/CombineMeshInstancesImpl.cpp
@@ -1838,6 +1838,8 @@ static void ComputePlanarPolygonApproximation(
 	check(SourcePartMesh.IsCompactT());
 	FFrame3d ProjectFrame(SourcePartMesh.GetTriCentroid(0), Direction);
 
+	double MaxDistanceZ = 0.0;		// maximum distance between vertices and the projection plane
+
 	FMeshBoundaryLoops Loops(&SourcePartMesh);
 	FPlanarComplexd PlanarComplex;
 	for (FEdgeLoop& Loop : Loops.Loops)
@@ -1848,6 +1850,7 @@ static void ComputePlanarPolygonApproximation(
 		for (FVector3d V : Vertices)
 		{
 			FVector LocalV = ProjectFrame.ToFramePoint(V);
+			MaxDistanceZ = FMathd::Max(MaxDistanceZ, FMathd::Abs(LocalV.Z));
 			Polygon.AppendVertex( FVector2d(LocalV.X, LocalV.Y) );
 		}
 		Polygon.Reverse();		// mesh orientation comes out backwards...
@@ -1939,8 +1942,40 @@ static void ComputePlanarPolygonApproximation(
 		return;
 	}
 
+	// map back to 3D plane
 	MeshTransforms::FrameCoordsToWorld(PolygonsMesh, ProjectFrame);
 
+	// todo trivially parallelizable...
+	// since we projected the mesh to plane, we may have introduced some cracks. Unfortunately since
+	// we did topological operations we cannot guarantee the verts are in the exact same place anymore, or
+	// that we have the same vertices at all. So instead we will try to find a vertex inside an epsilon-ball,
+	// and if found, we will snap to that location
+	double SnapTolerance = 2.0 * MaxDistanceZ;		// maybe should be based on edge length instead?
+	if (MaxDistanceZ > FMathd::ZeroTolerance)
+	{
+		for (int32 vid : PolygonsMesh.VertexIndicesItr())
+		{
+			FVector3d Pos = PolygonsMesh.GetVertex(vid);
+
+			FVector3d NearestOriginalPos = Pos;
+			double NearestDistSqr = 2.0 * SnapTolerance;
+			for (FVector3d OrigPos : SourcePartMesh.VerticesItr())
+			{
+				double DistSqr = DistanceSquared(Pos, OrigPos);
+				if (DistSqr < NearestDistSqr)
+				{
+					NearestDistSqr = DistSqr;
+					NearestOriginalPos = OrigPos;
+				}
+			}
+			if (NearestDistSqr < SnapTolerance * SnapTolerance)
+			{
+				PolygonsMesh.SetVertex(vid, NearestOriginalPos);
+			}
+		}
+	}
+
+
 	NewPlanarMesh = MoveTemp(PolygonsMesh);
 }
 
@@ -1949,13 +1984,13 @@ static void ComputePlanarPolygonApproximation(
 // Find sets of triangles that lie in the same 3D plane on TargetMesh,
 // and then extract out those areas, pull out the boundary polygons,
 // union them together, and do a 2D polygon-with-holes triangulation.
-static void RetriangulatePlanarFacePolygons(FDynamicMesh3& TargetMesh)
+static void RetriangulatePlanarFacePolygons(FDynamicMesh3& TargetMesh, double BaseGeometricTolerance)
 {
 	TArray<FFrame3d> PlaneSet;
 
+	// plane detection tolerances
 	const double AngleDotTol = 0.99;
-	// this distance probably should not be hardcoded...
-	const double DistanceTol = 0.1;
+	const double DistanceTol = BaseGeometricTolerance * 0.05;
 
 	if (TargetMesh.IsCompactT() == false)
 	{
@@ -2016,12 +2051,19 @@ static void RetriangulatePlanarFacePolygons(FDynamicMesh3& TargetMesh)
 
 	for (FDynamicMesh3& Mesh : SplitMeshes)		// ought to be trivially parallelizable...
 	{
-		if (Mesh.TriangleCount() <= 2) continue;
+		if (Mesh.TriangleCount() <= 2) continue;		// technically even 2 tris might form a single triangle...
 
-		const double MergeOffset = 1.0;
-		const double SimplifyTolerance = 1.0;
 		const double MinHoleArea = 10.0;
 
+		// planar areas can be very small, if we use full tolerance then they may end up being
+		// partially collapsed by topological operations, and creating cracks/holes. Deriving from
+		// area of submesh helps to prevent this, but results in too large of values for large areas,
+		// so clamp to geometric tolerance.
+		FVector2d VolumeArea = TMeshQueries<FDynamicMesh3>::GetVolumeArea(Mesh);
+		double EstSquareEdgeLen = FMathd::Sqrt(VolumeArea.Y);
+		double MergeOffset = FMath::Min(EstSquareEdgeLen * 0.02, BaseGeometricTolerance);
+		double SimplifyTolerance = FMath::Min(MergeOffset, BaseGeometricTolerance);
+
 		FDynamicMesh3 NewPlanarMesh;
 		ComputePlanarPolygonApproximation(Mesh, NewPlanarMesh,
 			Mesh.GetTriNormal(0), MergeOffset, SimplifyTolerance, MinHoleArea);
@@ -2109,7 +2151,7 @@ static void RetriangulatePlanarFacePolygons(FDynamicMesh3& TargetMesh)
 //
 static void PostProcessHiddenFaceRemovedMesh(
 	FDynamicMesh3& TargetMesh, 
-	double Tolerance,
+	double BaseGeometricTolerance,
 	bool bTryToMergeFaces,
 	bool bApplyPlanarRetriangulation,
 	TFunctionRef<FIndex3i(const FDynamicMesh3& Mesh, int32 TriangleID)> GetTriangleGroupingIDFunc)
@@ -2125,7 +2167,7 @@ static void PostProcessHiddenFaceRemovedMesh(
 		MeshEditor.SplitBowties(EditResult);
 
 		FMergeCoincidentMeshEdges Welder(&TargetMesh);
-		Welder.MergeVertexTolerance = Tolerance * 0.01;
+		Welder.MergeVertexTolerance = BaseGeometricTolerance * 0.01;
 		Welder.OnlyUniquePairs = false;
 		Welder.bWeldAttrsOnMergedEdges = true;
 		Welder.Apply();
@@ -2182,7 +2224,7 @@ static void PostProcessHiddenFaceRemovedMesh(
 
 		// try weld again just in case
 		FMergeCoincidentMeshEdges Welder(&SubRegionMesh);
-		Welder.MergeVertexTolerance = Tolerance * 0.01;
+		Welder.MergeVertexTolerance = BaseGeometricTolerance * 0.01;
 		Welder.OnlyUniquePairs = false;
 		Welder.bWeldAttrsOnMergedEdges = true;
 		Welder.Apply();
@@ -2216,7 +2258,7 @@ static void PostProcessHiddenFaceRemovedMesh(
 		// polygons that can be delaunay-triangulated. This is currently somewhat expensive...
 		if (bApplyPlanarRetriangulation)
 		{
-			RetriangulatePlanarFacePolygons(SubRegionMesh);
+			RetriangulatePlanarFacePolygons(SubRegionMesh, BaseGeometricTolerance);
 		}
 	}
 
@@ -2241,7 +2283,7 @@ static void PostProcessHiddenFaceRemovedMesh(
 		MeshEditor.SplitBowties(EditResult);
 
 		FMergeCoincidentMeshEdges Welder(&TargetMesh);
-		Welder.MergeVertexTolerance = Tolerance * 0.01;
+		Welder.MergeVertexTolerance = BaseGeometricTolerance * 0.01;
 		Welder.OnlyUniquePairs = false;
 		Welder.bWeldAttrsOnMergedEdges = true;
 		Welder.Apply();
@@ -2831,11 +2873,10 @@ static void SortMesh(FDynamicMesh3& Mesh)
 
 
 
-void ComputeHiddenRemovalForLOD(
+bool ComputeHiddenRemovalForLOD(
 	FDynamicMesh3& MeshLOD,
 	int32 LODIndex,
-	IGeometryProcessing_CombineMeshInstances::FOptions CombineOptions,
-	bool bSkipPostProcessing = false)
+	IGeometryProcessing_CombineMeshInstances::FOptions CombineOptions)
 {
 	bool bVerbose = CVarGeometryCombineMeshInstancesVerbose.GetValueOnAnyThread();
 
@@ -2859,23 +2900,32 @@ void ComputeHiddenRemovalForLOD(
 		UE_LOG(LogGeometry, Log, TEXT("    Remove Hidden Faces - [Tris %6d Verts %6d]"), MeshLOD.TriangleCount(), MeshLOD.VertexCount());
 	}
 
-	if ( bModified && bSkipPostProcessing == false)
-	{
-		TFunction<FIndex3i(const FDynamicMesh3& Mesh, int32 TriangleID)> GroupingIDFunc = CombineOptions.TriangleGroupingIDFunc;
-		if (!GroupingIDFunc)
-		{
-			GroupingIDFunc = [](const FDynamicMesh3&, int32) { return FIndex3i::Zero(); };
-		}
+	return bModified;
+}
+
+
+void OptimizeLODMeshTriangulation(
+	FDynamicMesh3& MeshLOD,
+	int32 LODIndex,
+	IGeometryProcessing_CombineMeshInstances::FOptions CombineOptions,
+	double BaseGeometricTolerance)
+{
+	bool bVerbose = CVarGeometryCombineMeshInstancesVerbose.GetValueOnAnyThread();
 
-		PostProcessHiddenFaceRemovedMesh(MeshLOD, 
-			CombineOptions.SimplifyBaseTolerance,
-			CombineOptions.bMergeCoplanarFaces && LODIndex >= CombineOptions.MergeCoplanarFacesStartLOD,
-			CombineOptions.PlanarPolygonRetriangulationStartLOD >= 0 && LODIndex >= CombineOptions.PlanarPolygonRetriangulationStartLOD,
-			GroupingIDFunc );
+	TFunction<FIndex3i(const FDynamicMesh3& Mesh, int32 TriangleID)> GroupingIDFunc = CombineOptions.TriangleGroupingIDFunc;
+	if (!GroupingIDFunc)
+	{
+		GroupingIDFunc = [](const FDynamicMesh3&, int32) { return FIndex3i::Zero(); };
 	}
+
+	PostProcessHiddenFaceRemovedMesh(MeshLOD, BaseGeometricTolerance,
+		CombineOptions.bMergeCoplanarFaces && LODIndex >= CombineOptions.MergeCoplanarFacesStartLOD,
+		CombineOptions.PlanarPolygonRetriangulationStartLOD >= 0 && LODIndex >= CombineOptions.PlanarPolygonRetriangulationStartLOD,
+		GroupingIDFunc);
 }
 
 
+
 // change this to build a single LOD, and separate versions for (eg) source mesh vs approx mesh
 // should we even bother w/ storing approx meshes? just generate them as needed?
 
@@ -3037,6 +3087,24 @@ void BuildCombinedMesh(
 		MeshLODs[FirstVoxWrappedIndex].Mesh = MeshLODs[CombineOptions.NumCopiedLODs-1].Mesh;
 	}
 
+	// make a list of per-LOD geometric tolerances tha will drive additional optimization. 
+	// For copied and first simplified LODs, use Simplify Base Tolerance, and then increment
+	// for each successive LOD.  (todo: have a separate initial tolerance for Approx LODs?)
+	double CurTolerance = CombineOptions.SimplifyBaseTolerance;
+	TArray<double> OptimizationTolerances;
+	for (int32 LODIndex = 0; LODIndex < NumLODs; ++LODIndex)
+	{
+		if (LODTypes[LODIndex] == ECombinedLODType::Simplified && LODTypes[LODIndex - 1] == ECombinedLODType::Simplified)
+		{
+			CurTolerance *= 2.0;
+		}
+		else if (LODTypes[LODIndex] == ECombinedLODType::Approximated)
+		{
+			CurTolerance *= 2.0;
+		}
+		OptimizationTolerances.Add(CurTolerance);
+	}
+
 	//
 	// start hidden-removal passes on all meshes up to voxel LODs here, because we can compute voxel LOD at the same time
 	//
@@ -3053,9 +3121,11 @@ void BuildCombinedMesh(
 				UE_LOG(LogGeometry, Log, TEXT("  Optimizing LOD%d - Tris %6d Verts %6d"), LODIndex, MeshLODs[LODIndex].Mesh.TriangleCount(), MeshLODs[LODIndex].Mesh.VertexCount());
 			}
 
-			UE::Tasks::FTask RemoveHiddenTask = UE::Tasks::Launch(UE_SOURCE_LOCATION, [&MeshLODs, &CombineOptions, LODIndex]()
+			double UseTolerance = OptimizationTolerances[LODIndex];
+			UE::Tasks::FTask RemoveHiddenTask = UE::Tasks::Launch(UE_SOURCE_LOCATION, [&MeshLODs, &CombineOptions, LODIndex, UseTolerance]()
 			{ 
 				ComputeHiddenRemovalForLOD(MeshLODs[LODIndex].Mesh, LODIndex, CombineOptions);
+				OptimizeLODMeshTriangulation(MeshLODs[LODIndex].Mesh, LODIndex, CombineOptions, UseTolerance);
 			});
 			PendingRemoveHiddenTasks.Add(RemoveHiddenTask);
 
@@ -3238,8 +3308,7 @@ void BuildCombinedMesh(
 					UE_LOG(LogGeometry, Log, TEXT("  Optimizing LOD%d - Tris %6d Verts %6d"), LODIndex, MeshLODs[LODIndex].Mesh.TriangleCount(), MeshLODs[LODIndex].Mesh.VertexCount());
 				}
 
-				// disable hidden-removal postprocessing for voxel LODs (this should not be part of ComputeHiddenRemovalForLOD...)
-				ComputeHiddenRemovalForLOD(MeshLODs[LODIndex].Mesh, LODIndex, CombineOptions, /*bSkipPostProcessing=*/ true);
+				ComputeHiddenRemovalForLOD(MeshLODs[LODIndex].Mesh, LODIndex, CombineOptions);
 			}
 		}, (bVerbose) ? EParallelForFlags::ForceSingleThread :  EParallelForFlags::None );
 	}
