Skeleton Modifier API
    - get all bone names
    - get children names

diff --git a/Engine/Plugins/Experimental/MeshModelingToolsetExp/Source/SkeletalMeshModifiers/Private/SkeletonModifier.cpp b/Engine/Plugins/Experimental/MeshModelingToolsetExp/Source/SkeletalMeshModifiers/Private/SkeletonModifier.cpp
--- a/Engine/Plugins/Experimental/MeshModelingToolsetExp/Source/SkeletalMeshModifiers/Private/SkeletonModifier.cpp
+++ b/Engine/Plugins/Experimental/MeshModelingToolsetExp/Source/SkeletalMeshModifiers/Private/SkeletonModifier.cpp
@@ -1333,6 +1333,11 @@ FTransform USkeletonModifier::GetBoneTransform(const FName InBoneName, const boo
 
 FName USkeletonModifier::GetParentName(const FName InBoneName) const
 {
+	if (!IsReferenceSkeletonValid())
+	{
+		return NAME_None;
+	}
+	
 	const int32 BoneIndex = ReferenceSkeleton->FindBoneIndex(InBoneName);
 	if (BoneIndex == INDEX_NONE)
 	{
@@ -1343,6 +1348,64 @@ FName USkeletonModifier::GetParentName(const FName InBoneName) const
 	return (ParentIndex > INDEX_NONE) ? BoneInfos[ParentIndex].Name : NAME_None;
 }
 
+TArray<FName> USkeletonModifier::GetChildrenNames(const FName InBoneName, const bool bRecursive) const
+{
+	TArray<FName> ChildrenNames;
+
+	if (!IsReferenceSkeletonValid())
+	{
+		return ChildrenNames;
+	}
+	
+	const int32 BoneIndex = ReferenceSkeleton->FindRawBoneIndex(InBoneName);
+	if (BoneIndex == INDEX_NONE)
+	{
+		return ChildrenNames; 
+	}
+
+	TArray<int32> ChildrenIndices;
+	if (bRecursive)
+	{
+		auto GetChildren = [&](const int32 InBoneIndex, auto&& GetChildren) -> void
+		{
+			TArray<int32> Children; ReferenceSkeleton->GetRawDirectChildBones(InBoneIndex, Children);
+			for (int32 ChildIndex: Children)
+			{
+				ChildrenIndices.Add(InBoneIndex);
+				GetChildren(ChildIndex, GetChildren);
+			}
+		};
+		GetChildren(BoneIndex, GetChildren);
+	}
+	else
+	{
+		ReferenceSkeleton->GetRawDirectChildBones(BoneIndex, ChildrenIndices);
+	}
+
+	ChildrenNames.Reserve(ChildrenIndices.Num());
+
+	const TArray<FMeshBoneInfo>& BoneInfos = ReferenceSkeleton->GetRawRefBoneInfo();
+	Algo::Transform(ChildrenIndices, ChildrenNames, [&BoneInfos](int32 BoneIndex) { return BoneInfos[BoneIndex].Name; });
+
+	return ChildrenNames;
+}
+
+TArray<FName> USkeletonModifier::GetAllBoneNames() const
+{
+	TArray<FName> BoneNames;
+	if (!IsReferenceSkeletonValid())
+	{
+		return BoneNames;
+	}
+
+	const TArray<FMeshBoneInfo>& BoneInfos = ReferenceSkeleton->GetRawRefBoneInfo();
+	BoneNames.Reserve(BoneInfos.Num());
+
+	Algo::Transform(BoneInfos, BoneNames, [](const FMeshBoneInfo& BoneInfo) { return BoneInfo.Name; });
+
+	return BoneNames;
+}
+
 const FTransform& USkeletonModifier::GetTransform(const int32 InBoneIndex, const bool bGlobal) const
 {
 	if (!ReferenceSkeleton.IsValid())
diff --git a/Engine/Plugins/Experimental/MeshModelingToolsetExp/Source/SkeletalMeshModifiers/Public/SkeletonModifier.h b/Engine/Plugins/Experimental/MeshModelingToolsetExp/Source/SkeletalMeshModifiers/Public/SkeletonModifier.h
--- a/Engine/Plugins/Experimental/MeshModelingToolsetExp/Source/SkeletalMeshModifiers/Public/SkeletonModifier.h
+++ b/Engine/Plugins/Experimental/MeshModelingToolsetExp/Source/SkeletalMeshModifiers/Public/SkeletonModifier.h
@@ -171,6 +171,20 @@ public:
 	 */
 	UFUNCTION(BlueprintCallable, Category = "Skeleton")
 	FName GetParentName(const FName InBoneName) const;
+
+	/** Get Children Names
+	 *  @param InBoneName The parent's name.
+	 *  @param bRecursive If set to true grand-children will also be added recursively
+	 *  @return \c The children names list 
+	 */
+	UFUNCTION(BlueprintCallable, Category = "Skeleton")
+	TArray<FName> GetChildrenNames(const FName InBoneName, const bool bRecursive = false) const;
+
+	/** Get All Bone Names
+	 *  @return \c All bone names list
+	 */
+	UFUNCTION(BlueprintCallable, Category = "Skeleton")
+	TArray<FName> GetAllBoneNames() const;
 	
 	/**
 	 * Actually applies the skeleton modifications to the skeletal mesh.
