ModelingTools: UX improvements to MeshVertexPaintTool. move bHardEdges setting to appropriate place, rename main category. Add 'Triangle Fill' main operation mode. Add more verbose display names for main operation types.

diff --git a/Engine/Plugins/Experimental/MeshModelingToolsetExp/Source/MeshModelingToolsExp/Private/MeshVertexPaintTool.cpp b/Engine/Plugins/Experimental/MeshModelingToolsetExp/Source/MeshModelingToolsExp/Private/MeshVertexPaintTool.cpp
--- a/Engine/Plugins/Experimental/MeshModelingToolsetExp/Source/MeshModelingToolsExp/Private/MeshVertexPaintTool.cpp
+++ b/Engine/Plugins/Experimental/MeshModelingToolsetExp/Source/MeshModelingToolsExp/Private/MeshVertexPaintTool.cpp
@@ -495,7 +495,8 @@ bool UMeshVertexPaintTool::IsInBrushSubMode() const
 {
 	return BasicProperties->SubToolType == EMeshVertexPaintInteractionType::Brush
 		|| BasicProperties->SubToolType == EMeshVertexPaintInteractionType::Fill
-		|| BasicProperties->SubToolType == EMeshVertexPaintInteractionType::GroupFill;
+		|| BasicProperties->SubToolType == EMeshVertexPaintInteractionType::GroupFill
+		|| BasicProperties->SubToolType == EMeshVertexPaintInteractionType::TriFill;
 }
 
 
@@ -962,7 +963,7 @@ void UMeshVertexPaintTool::InitializeElementROIFromTriangleROI(const TArray<int3
 	FDynamicMesh3* Mesh = GetSculptMesh();
 	FDynamicMeshColorOverlay* ColorOverlay = GetActiveColorOverlay();
 
-	if (FilterProperties->bHardEdges)
+	if (BasicProperties->bHardEdges)
 	{
 		for (int32 tid : Triangles)
 		{
@@ -2001,7 +2002,7 @@ void UMeshVertexPaintTool::UpdateSubToolType(EMeshVertexPaintInteractionType New
 {
 	FilterProperties->CurrentSubToolType = BasicProperties->SubToolType;
 
-	bool bSculptPropsVisible = (NewType == EMeshVertexPaintInteractionType::Brush);
+	bool bSculptPropsVisible = (NewType != EMeshVertexPaintInteractionType::PolyLasso);
 	SetToolPropertySourceEnabled(UMeshSculptToolBase::BrushProperties, bSculptPropsVisible);
 
 	//SetToolPropertySourceEnabled(BasicProperties, true);
diff --git a/Engine/Plugins/Experimental/MeshModelingToolsetExp/Source/MeshModelingToolsExp/Public/MeshVertexPaintTool.h b/Engine/Plugins/Experimental/MeshModelingToolsetExp/Source/MeshModelingToolsExp/Public/MeshVertexPaintTool.h
--- a/Engine/Plugins/Experimental/MeshModelingToolsetExp/Source/MeshModelingToolsExp/Public/MeshVertexPaintTool.h
+++ b/Engine/Plugins/Experimental/MeshModelingToolsetExp/Source/MeshModelingToolsExp/Public/MeshVertexPaintTool.h
@@ -52,12 +52,14 @@ public:
 UENUM()
 enum class EMeshVertexPaintInteractionType : uint8
 {
-	/** Paint with an Interactive Brush */
-	Brush UMETA(DisplayName = "Paint"),
+	/** Paint Vertices of hit triangles with a smooth falloff */
+	Brush UMETA(DisplayName = "Paint Vertices"),
+	/** Fill any painted triangles, by setting all 3 vertices to the same color */
+	TriFill UMETA(DisplayName = "Paint Triangles"),
 	/** Fill any triangles connected to the brushed triangles */
-	Fill,
+	Fill UMETA(DisplayName = "Flood Fill Connected"),
 	/** Fill any polygroups connected to the brushed triangles */
-	GroupFill,
+	GroupFill UMETA(DisplayName = "Flood Fill Groups"),
 	/** Paint any triangles inside polygonal or freehand Lassos drawn in the viewport */
 	PolyLasso,
 
@@ -170,32 +172,36 @@ public:
 	EMeshVertexPaintBrushType PrimaryBrushType = EMeshVertexPaintBrushType::Paint;
 
 	/** Painting Operation to apply when left-clicking and dragging */
-	UPROPERTY(EditAnywhere, Category = ActionType, meta = (DisplayName = "Action"))
+	UPROPERTY(EditAnywhere, Category = Settings, meta = (DisplayName = "Action"))
 	EMeshVertexPaintInteractionType SubToolType = EMeshVertexPaintInteractionType::Brush;
 
 	/** The Color that will be assigned to painted triangle vertices */
-	UPROPERTY(EditAnywhere, Category = ActionType, meta = ())
+	UPROPERTY(EditAnywhere, Category = Settings, meta = ())
 	FLinearColor PaintColor = FLinearColor::Red;
 
 	/** Controls how painted Colors will be combined with the existing Colors */
-	UPROPERTY(EditAnywhere, Category = ActionType)
+	UPROPERTY(EditAnywhere, Category = Settings)
 	EMeshVertexPaintColorBlendMode BlendMode = EMeshVertexPaintColorBlendMode::Lerp;
 
 	/** The Brush Operation that will be applied when holding the Shift key when in Painting */
-	UPROPERTY(EditAnywhere, Category = ActionType, meta = (DisplayName = "Secondary Brush"))
+	UPROPERTY(EditAnywhere, Category = Settings, meta = (DisplayName = "Secondary Brush"))
 	EMeshVertexPaintSecondaryActionType SecondaryActionType = EMeshVertexPaintSecondaryActionType::Erase;
 
 	/** Color to set when using Erase brush */
-	UPROPERTY(EditAnywhere, Category = ActionType, meta = (EditConditionHides, EditCondition = "SecondaryActionType != EMeshVertexPaintSecondaryActionType::Smooth"))
+	UPROPERTY(EditAnywhere, Category = Settings, meta = (EditConditionHides, EditCondition = "SecondaryActionType != EMeshVertexPaintSecondaryActionType::Smooth"))
 	FLinearColor EraseColor = FLinearColor::White;
 
 	/** Strength of Smooth Brush */
-	UPROPERTY(EditAnywhere, Category = ActionType, meta = (UIMin = 0, UIMax = 1, EditConditionHides, EditCondition = "SecondaryActionType == EMeshVertexPaintSecondaryActionType::Smooth"))
+	UPROPERTY(EditAnywhere, Category = Settings, meta = (UIMin = 0, UIMax = 1, EditConditionHides, EditCondition = "SecondaryActionType == EMeshVertexPaintSecondaryActionType::Smooth"))
 	float SmoothStrength = 0.25;
 
 	/** Controls which Color Channels will be affected by Operations. Only enabled Channels are rendered. */
-	UPROPERTY(EditAnywhere, Category = ActionType)
+	UPROPERTY(EditAnywhere, Category = Settings)
 	FModelingToolsColorChannelFilter ChannelFilter;
+
+	/** Create Split Colors / Hard Color Edges at the borders of the painted area. Use Soften operations to un-split. */
+	UPROPERTY(EditAnywhere, Category = Settings, meta = (DisplayName = "Hard Edges"))
+	bool bHardEdges = false;
 };
 
 
@@ -224,10 +230,6 @@ public:
 	UPROPERTY(EditAnywhere, Category = Filters, meta = (DisplayName = "Hard Normals", EditCondition = "CurrentSubToolType != EMeshVertexPaintInteractionType::PolyLasso && BrushAreaMode == EMeshVertexPaintBrushAreaType::Connected"))
 	bool bNormalSeams = false;
 
-	/** Create Split Colors / Hard Color Edges at the borders of the painted area. Use Soften operations to un-split. */
-	UPROPERTY(EditAnywhere, Category = Filters, meta = (DisplayName = "Hard Edges"))
-	bool bHardEdges = false;
-
 	/** Control which triangles can be affected by the current operation based on visibility. Applied after all other filters. */
 	UPROPERTY(EditAnywhere, Category = Filters)
 	EMeshVertexPaintVisibilityType VisibilityFilter = EMeshVertexPaintVisibilityType::None;
@@ -297,7 +299,7 @@ public:
 	/**
 	 * Fill all Vertex Colors with the current Paint color. Current Channel Filter still applies.
 	 */
-	UFUNCTION(CallInEditor, Category = Operations, meta = (DisplayPriority = 12))
+	UFUNCTION(CallInEditor, Category = QuickActions, meta = (DisplayPriority = 12))
 	void PaintAll()
 	{
 		PostAction(EMeshVertexPaintToolActions::PaintAll);
@@ -306,7 +308,7 @@ public:
 	/**
 	 * Fill all Vertex Colors with the current Erase color. Current Channel Filter still applies.
 	 */
-	UFUNCTION(CallInEditor, Category = Operations, meta = (DisplayPriority = 13))
+	UFUNCTION(CallInEditor, Category = QuickActions, meta = (DisplayPriority = 13))
 	void EraseAll()
 	{
 		PostAction(EMeshVertexPaintToolActions::EraseAll);
@@ -315,7 +317,7 @@ public:
 	/**
 	 * Fill all Vertex Colors with the Color (0,0,0,1). Current Channel Filter still applies.
 	 */
-	UFUNCTION(CallInEditor, Category = Operations, meta = (DisplayPriority = 14))
+	UFUNCTION(CallInEditor, Category = QuickActions, meta = (DisplayPriority = 14))
 	void FillBlack()
 	{
 		PostAction(EMeshVertexPaintToolActions::FillBlack);
@@ -324,7 +326,7 @@ public:
 	/**
 	 * Fill all Vertex Colors with the Color (1,1,1,1). Current Channel Filter still applies.
 	 */
-	UFUNCTION(CallInEditor, Category = Operations, meta = (DisplayPriority = 14))
+	UFUNCTION(CallInEditor, Category = QuickActions, meta = (DisplayPriority = 14))
 	void FillWhite()
 	{
 		PostAction(EMeshVertexPaintToolActions::FillWhite);
@@ -382,17 +384,17 @@ public:
 	/**
 	 * Operation to apply to current Vertex Colors
 	 */
-	UPROPERTY(EditAnywhere, Category = Utility, meta = (NoResetToDefault, DisplayPriority = 1))
+	UPROPERTY(EditAnywhere, Category = UtilityOperations, meta = (NoResetToDefault, DisplayPriority = 1))
 	EMeshVertexPaintToolUtilityOperations Operation = EMeshVertexPaintToolUtilityOperations::BlendAllSeams;
 
-	UPROPERTY(EditAnywhere, Category = Utility, meta = (EditConditionHides, EditCondition = "Operation == EMeshVertexPaintToolUtilityOperations::CopyChannelToChannel || Operation == EMeshVertexPaintToolUtilityOperations::SwapChannels"))
+	UPROPERTY(EditAnywhere, Category = UtilityOperations, meta = (EditConditionHides, EditCondition = "Operation == EMeshVertexPaintToolUtilityOperations::CopyChannelToChannel || Operation == EMeshVertexPaintToolUtilityOperations::SwapChannels"))
 	EMeshVertexPaintColorChannel SourceChannel = EMeshVertexPaintColorChannel::Red;
 
-	UPROPERTY(EditAnywhere, Category = Utility, meta = (UIMin = 0, UIMax = 1, EditConditionHides, EditCondition = "Operation == EMeshVertexPaintToolUtilityOperations::FillChannels"))
+	UPROPERTY(EditAnywhere, Category = UtilityOperations, meta = (UIMin = 0, UIMax = 1, EditConditionHides, EditCondition = "Operation == EMeshVertexPaintToolUtilityOperations::FillChannels"))
 	float SourceValue = 0.0f;
 
 	/** Target Vertex Weight Map */
-	UPROPERTY(EditAnywhere, Category = Utility, meta = (TransientToolProperty, NoResetToDefault, GetOptions = GetWeightMapsFunc, EditConditionHides, EditCondition = "Operation == EMeshVertexPaintToolUtilityOperations::CopyFromWeightMap"))
+	UPROPERTY(EditAnywhere, Category = UtilityOperations, meta = (TransientToolProperty, NoResetToDefault, GetOptions = GetWeightMapsFunc, EditConditionHides, EditCondition = "Operation == EMeshVertexPaintToolUtilityOperations::CopyFromWeightMap"))
 	FName WeightMap;
 
 	// this function is called provide set of available weight maps
@@ -403,18 +405,18 @@ public:
 	UPROPERTY(meta = (TransientToolProperty))
 	TArray<FString> WeightMapsList;
 
-	UPROPERTY(EditAnywhere, Category = Utility, meta = (EditConditionHides, EditCondition = "Operation == EMeshVertexPaintToolUtilityOperations::FillChannels || Operation == EMeshVertexPaintToolUtilityOperations::InvertChannels || Operation == EMeshVertexPaintToolUtilityOperations::CopyChannelToChannel || Operation == EMeshVertexPaintToolUtilityOperations::CopyFromWeightMap"))
+	UPROPERTY(EditAnywhere, Category = UtilityOperations, meta = (EditConditionHides, EditCondition = "Operation == EMeshVertexPaintToolUtilityOperations::FillChannels || Operation == EMeshVertexPaintToolUtilityOperations::InvertChannels || Operation == EMeshVertexPaintToolUtilityOperations::CopyChannelToChannel || Operation == EMeshVertexPaintToolUtilityOperations::CopyFromWeightMap"))
 	FModelingToolsColorChannelFilter TargetChannels;
 
-	UPROPERTY(EditAnywhere, Category = Utility, meta = (EditConditionHides, EditCondition = "Operation == EMeshVertexPaintToolUtilityOperations::SwapChannels"))
+	UPROPERTY(EditAnywhere, Category = UtilityOperations, meta = (EditConditionHides, EditCondition = "Operation == EMeshVertexPaintToolUtilityOperations::SwapChannels"))
 	EMeshVertexPaintColorChannel TargetChannel = EMeshVertexPaintColorChannel::Green;
 
 	/** Copy colors to HiRes Source Mesh, if it exists */
-	UPROPERTY(EditAnywhere, Category = Utility, meta = (EditConditionHides, EditCondition = "Operation == EMeshVertexPaintToolUtilityOperations::CopyToOtherLODs"))
+	UPROPERTY(EditAnywhere, Category = UtilityOperations, meta = (EditConditionHides, EditCondition = "Operation == EMeshVertexPaintToolUtilityOperations::CopyToOtherLODs"))
 	bool bCopyToHiRes = false;
 
 	/** Target LOD to copy Colors to */
-	UPROPERTY(EditAnywhere, Category = Utility, meta = (TransientToolProperty, DisplayName = "Copy To LOD", NoResetToDefault, GetOptions = GetLODNamesFunc, EditConditionHides, EditCondition = "Operation == EMeshVertexPaintToolUtilityOperations::CopyToSingleLOD"))
+	UPROPERTY(EditAnywhere, Category = UtilityOperations, meta = (TransientToolProperty, DisplayName = "Copy To LOD", NoResetToDefault, GetOptions = GetLODNamesFunc, EditConditionHides, EditCondition = "Operation == EMeshVertexPaintToolUtilityOperations::CopyToSingleLOD"))
 	FString CopyToLODName;
 
 	UFUNCTION()
@@ -427,7 +429,7 @@ public:
 	/**
 	 * Apply the Operation currently selected below
 	 */
-	UFUNCTION(CallInEditor, Category = Utility, meta = (DisplayPriority = 10))
+	UFUNCTION(CallInEditor, Category = UtilityOperations, meta = (DisplayPriority = 10))
 	void ApplySelectedOperation()
 	{
 		PostAction(EMeshVertexPaintToolActions::ApplyCurrentUtility);
