Modeling Mode: Fix non-unity compile error

#rnx

diff --git a/Engine/Plugins/Experimental/MeshModelingToolsetExp/Source/MeshModelingToolsExp/Private/Physics/SimpleCollisionEditorTool.cpp b/Engine/Plugins/Experimental/MeshModelingToolsetExp/Source/MeshModelingToolsExp/Private/Physics/SimpleCollisionEditorTool.cpp
--- a/Engine/Plugins/Experimental/MeshModelingToolsetExp/Source/MeshModelingToolsExp/Private/Physics/SimpleCollisionEditorTool.cpp
+++ b/Engine/Plugins/Experimental/MeshModelingToolsetExp/Source/MeshModelingToolsExp/Private/Physics/SimpleCollisionEditorTool.cpp
@@ -19,6 +19,8 @@
 #include "Components/DynamicMeshComponent.h"
 #include "Components/StaticMeshComponent.h"
 
+#include "UObject/UObjectIterator.h"
+
 #include "TargetInterfaces/PrimitiveComponentBackedTarget.h"
 #include "TargetInterfaces/PhysicsDataSource.h"
 #include "ToolTargetManager.h"
