Skinning tool improvements:
details customization hides/unhides depending on paint/vertex modes
Fixed bug with vertex selection in skin mode.
Frame selection fixed in skin tool.
Tool mode settings are saved/restored
All fields have tooltips now.
Vertex color ramp now uses min/max colors

diff --git a/Engine/Plugins/Experimental/Animation/SkeletalMeshModelingTools/Source/SkeletalMeshModelingTools/Private/SkeletalMeshModelingToolsEditorMode.cpp b/Engine/Plugins/Experimental/Animation/SkeletalMeshModelingTools/Source/SkeletalMeshModelingTools/Private/SkeletalMeshModelingToolsEditorMode.cpp
--- a/Engine/Plugins/Experimental/Animation/SkeletalMeshModelingTools/Source/SkeletalMeshModelingTools/Private/SkeletalMeshModelingToolsEditorMode.cpp
+++ b/Engine/Plugins/Experimental/Animation/SkeletalMeshModelingTools/Source/SkeletalMeshModelingTools/Private/SkeletalMeshModelingToolsEditorMode.cpp
@@ -269,7 +269,6 @@ void USkeletalMeshModelingToolsEditorMode::CreateToolkit()
 	Toolkit = MakeShareable(new FSkeletalMeshModelingToolsEditorModeToolkit);
 }
 
-
 void USkeletalMeshModelingToolsEditorMode::Tick(FEditorViewportClient* ViewportClient, float DeltaTime)
 {
 	Super::Tick(ViewportClient, DeltaTime);
@@ -298,6 +297,19 @@ bool USkeletalMeshModelingToolsEditorMode::HandleClick(FEditorViewportClient* In
 
 bool USkeletalMeshModelingToolsEditorMode::ComputeBoundingBoxForViewportFocus(AActor* Actor, UPrimitiveComponent* PrimitiveComponent, FBox& InOutBox) const
 {
+	// if Tool supports custom Focus box, use that first
+	if (GetToolManager()->HasAnyActiveTool())
+	{
+		UInteractiveTool* Tool = GetToolManager()->GetActiveTool(EToolSide::Mouse);
+		IInteractiveToolCameraFocusAPI* FocusAPI = Cast<IInteractiveToolCameraFocusAPI>(Tool);
+		if (FocusAPI && FocusAPI->SupportsWorldSpaceFocusBox())
+		{
+			InOutBox = FocusAPI->GetWorldSpaceFocusBox();
+			return true;
+		}
+	}
+
+	// focus using selected bones in skel mesh editor
 	if (const USkeletalMeshComponent* Component = Cast<USkeletalMeshComponent>(PrimitiveComponent))
 	{
 		check(Component->GetSkeletalMeshAsset());
@@ -338,18 +350,6 @@ bool USkeletalMeshModelingToolsEditorMode::ComputeBoundingBoxForViewportFocus(AA
 			}
 		}
 	}
-
-	// if Tool supports custom Focus box, use that
-	if (GetToolManager()->HasAnyActiveTool())
-	{
-		UInteractiveTool* Tool = GetToolManager()->GetActiveTool(EToolSide::Mouse);
-		IInteractiveToolCameraFocusAPI* FocusAPI = Cast<IInteractiveToolCameraFocusAPI>(Tool);
-		if (FocusAPI && FocusAPI->SupportsWorldSpaceFocusBox())
-		{
-			InOutBox = FocusAPI->GetWorldSpaceFocusBox();
-			return true;
-		}
-	}
 	
 	return Super::ComputeBoundingBoxForViewportFocus(Actor, PrimitiveComponent, InOutBox);
 }
diff --git a/Engine/Plugins/Experimental/Animation/SkeletalMeshModelingTools/Source/SkeletalMeshModelingTools/Private/SkeletalMeshModelingToolsEditorMode.h b/Engine/Plugins/Experimental/Animation/SkeletalMeshModelingTools/Source/SkeletalMeshModelingTools/Private/SkeletalMeshModelingToolsEditorMode.h
--- a/Engine/Plugins/Experimental/Animation/SkeletalMeshModelingTools/Source/SkeletalMeshModelingTools/Private/SkeletalMeshModelingToolsEditorMode.h
+++ b/Engine/Plugins/Experimental/Animation/SkeletalMeshModelingTools/Source/SkeletalMeshModelingTools/Private/SkeletalMeshModelingToolsEditorMode.h
@@ -31,18 +31,12 @@ public:
 
 	// UEdMode overrides
 	virtual void Initialize() override;
-
 	virtual void Enter() override;
 	virtual void Exit() override;
 	virtual void CreateToolkit() override;
-
 	virtual void Tick(FEditorViewportClient* InViewportClient, float InDeltaTime) override;
-	// void Render(const FSceneView* InView, FViewport* InViewport, FPrimitiveDrawInterface* InPDI) override;
-	// void DrawHUD(FEditorViewportClient* InViewportClient, FViewport* InViewport, const FSceneView* InView, FCanvas* InCanvas) override;
 	virtual bool HandleClick(FEditorViewportClient* InViewportClient, HHitProxy *HitProxy, const FViewportClick &Click) override;
-
 	virtual bool ComputeBoundingBoxForViewportFocus(AActor* Actor, UPrimitiveComponent* PrimitiveComponent, FBox& InOutBox) const override;
-	
 	virtual bool UsesToolkits() const override { return true; }
 
 	// binding
diff --git a/Engine/Plugins/Experimental/Animation/SkeletalMeshModelingTools/Source/SkeletalMeshModelingTools/Private/SkeletalMeshModelingToolsModule.cpp b/Engine/Plugins/Experimental/Animation/SkeletalMeshModelingTools/Source/SkeletalMeshModelingTools/Private/SkeletalMeshModelingToolsModule.cpp
--- a/Engine/Plugins/Experimental/Animation/SkeletalMeshModelingTools/Source/SkeletalMeshModelingTools/Private/SkeletalMeshModelingToolsModule.cpp
+++ b/Engine/Plugins/Experimental/Animation/SkeletalMeshModelingTools/Source/SkeletalMeshModelingTools/Private/SkeletalMeshModelingToolsModule.cpp
@@ -157,6 +157,7 @@ void FSkeletalMeshModelingToolsModule::OnToggleEditingToolsMode(TWeakPtr<ISkelet
 		{
 			EditorModeManager.ActivateMode(USkeletalMeshModelingToolsEditorMode::Id, true);
 
+			// bind notifications to skeletal mesh modeling editor mode
 			UEdMode* ActiveMode = EditorModeManager.GetActiveScriptableMode(USkeletalMeshModelingToolsEditorMode::Id);
 			if (USkeletalMeshModelingToolsEditorMode* SkeletalMeshMode = CastChecked<USkeletalMeshModelingToolsEditorMode>(ActiveMode))
 			{
diff --git a/Engine/Plugins/Experimental/MeshModelingToolsetExp/Source/MeshModelingToolsEditorOnlyExp/Private/SkeletalMesh/SkinWeightsPaintTool.cpp b/Engine/Plugins/Experimental/MeshModelingToolsetExp/Source/MeshModelingToolsEditorOnlyExp/Private/SkeletalMesh/SkinWeightsPaintTool.cpp
--- a/Engine/Plugins/Experimental/MeshModelingToolsetExp/Source/MeshModelingToolsEditorOnlyExp/Private/SkeletalMesh/SkinWeightsPaintTool.cpp
+++ b/Engine/Plugins/Experimental/MeshModelingToolsetExp/Source/MeshModelingToolsEditorOnlyExp/Private/SkeletalMesh/SkinWeightsPaintTool.cpp
@@ -83,6 +83,18 @@ USkinWeightsPaintToolProperties::USkinWeightsPaintToolProperties()
 
 	MinColor = FLinearColor::Black;
 	MaxColor = FLinearColor::White;
+
+	BrushConfigs.Add(EWeightEditOperation::Add, &BrushConfigAdd);
+	BrushConfigs.Add(EWeightEditOperation::Replace, &BrushConfigReplace);
+	BrushConfigs.Add(EWeightEditOperation::Multiply, &BrushConfigMultiply);
+	BrushConfigs.Add(EWeightEditOperation::Relax, &BrushConfigRelax);
+
+	LoadConfig();
+}
+
+FSkinWeightBrushConfig& USkinWeightsPaintToolProperties::GetBrushConfig()
+{
+	return *BrushConfigs[BrushMode];
 }
 
 void FMultiBoneWeightEdits::MergeSingleEdit(
@@ -260,7 +272,7 @@ void FSkinToolDeformer::UpdateVertexDeformation(USkinWeightsPaintTool* Tool)
 	}
 
 	// update data structures used by the selection mode
-	if (EditingMode == EWeightEditMode::Selection)
+	if (EditingMode == EWeightEditMode::Vertices)
 	{
 		// update AABB Tree for vertex selection
 		TRACE_CPUPROFILER_EVENT_SCOPE(SkinTool::UpdateAABBTree);
@@ -753,13 +765,37 @@ void USkinWeightsPaintTool::Render(IToolsContextRenderAPI* RenderAPI)
 	{
 		Super::Render(RenderAPI);	
 	}
-	else if (PolygonSelectionMechanic && WeightToolProperties->EditingMode == EWeightEditMode::Selection)
+	else if (PolygonSelectionMechanic && WeightToolProperties->EditingMode == EWeightEditMode::Vertices)
 	{
 		PolygonSelectionMechanic->Render(RenderAPI);
 	}
 	
 }
 
+FBox USkinWeightsPaintTool::GetWorldSpaceFocusBox()
+{
+	// TODO add support for selected bones
+
+	const FGroupTopologySelection& Selection = PolygonSelectionMechanic->GetActiveSelection();
+	if (!Selection.SelectedCornerIDs.IsEmpty())
+	{
+		TArray<VertexIndex> SelectedVertexIndices = Selection.SelectedCornerIDs.Array();
+		const FDynamicMesh3* Mesh = PreviewMesh->GetMesh();
+		FTransform3d Transform(PreviewMesh->GetTransform());
+		FAxisAlignedBox3d Bounds = FAxisAlignedBox3d::Empty();
+		for (const int32 VertexID : SelectedVertexIndices)
+		{
+			Bounds.Contain(Transform.TransformPosition(Mesh->GetVertex(VertexID)));
+		}
+		if (Bounds.MaxDim() > FMathf::ZeroTolerance)
+		{
+			return static_cast<FBox>(Bounds);
+		}
+	}
+	
+	return PreviewMesh->GetActor()->GetComponentsBoundingBox();
+}
+
 void USkinWeightsPaintTool::OnTick(float DeltaTime)
 {
 	if (bStampPending)
@@ -897,7 +933,7 @@ void USkinWeightsPaintTool::CalculateVertexROI(
 		return UseFalloff;
 	};
 	
-	if (WeightToolProperties->FalloffMode == EWeightBrushFalloffMode::Volume)
+	if (WeightToolProperties->GetBrushConfig().FalloffMode == EWeightBrushFalloffMode::Volume)
 	{
 		const IPrimitiveComponentBackedTarget* TargetComponent = Cast<IPrimitiveComponentBackedTarget>(Target);
 		const FTransform3d Transform(TargetComponent->GetWorldTransform());
@@ -918,7 +954,7 @@ void USkinWeightsPaintTool::CalculateVertexROI(
 		return;
 	}
 
-	if (WeightToolProperties->FalloffMode == EWeightBrushFalloffMode::Surface)
+	if (WeightToolProperties->GetBrushConfig().FalloffMode == EWeightBrushFalloffMode::Surface)
 	{
 		// get coordinate frame from stamp
 		auto GetFrameFromStamp = [](const FBrushStampData& InStamp) -> FFrame3d
@@ -982,9 +1018,9 @@ void USkinWeightsPaintTool::CalculateVertexROI(
 FVector4f USkinWeightsPaintTool::WeightToColor(float Value) const
 {
 	// optional greyscale mode
-	if (WeightToolProperties->ColorMode == EWeightColorMode::Greyscale)
+	if (WeightToolProperties->ColorMode == EWeightColorMode::MinMax)
 	{
-		return FLinearColor::LerpUsingHSV(FLinearColor::Black, FLinearColor::White, Value);	
+		return FLinearColor::LerpUsingHSV(WeightToolProperties->MinColor, WeightToolProperties->MaxColor, Value);	
 	}
 	
 	// early out zero weights to min color
@@ -1783,7 +1819,11 @@ void USkinWeightsPaintTool::HandleSkeletalMeshModified(const TArray<FName>& InBo
 	case ESkeletalMeshNotifyType::BonesRemoved:
 		break;
 	case ESkeletalMeshNotifyType::BonesMoved:
-		break;
+		{
+			// TODO update only vertices weighted to modified bones (AND CHILDREN!?)
+			Weights.Deformer.SetAllVerticesToBeUpdated();
+			break;	
+		}
 	case ESkeletalMeshNotifyType::BonesSelected:
 		{
 			// store selected bones
@@ -1815,7 +1855,7 @@ void USkinWeightsPaintTool::ToggleEditingMode()
 {
 	Weights.Deformer.SetAllVerticesToBeUpdated();
 	SetBrushEnabled(WeightToolProperties->EditingMode == EWeightEditMode::Brush);
-	PolygonSelectionMechanic->SetIsEnabled(WeightToolProperties->EditingMode == EWeightEditMode::Selection);
+	PolygonSelectionMechanic->SetIsEnabled(WeightToolProperties->EditingMode == EWeightEditMode::Vertices);
 }
 
 void USkinWeightsPaintTool::OnPropertyModified(UObject* ModifiedObject, FProperty* ModifiedProperty)
diff --git a/Engine/Plugins/Experimental/MeshModelingToolsetExp/Source/MeshModelingToolsEditorOnlyExp/Public/SkeletalMesh/SkinWeightsPaintTool.h b/Engine/Plugins/Experimental/MeshModelingToolsetExp/Source/MeshModelingToolsEditorOnlyExp/Public/SkeletalMesh/SkinWeightsPaintTool.h
--- a/Engine/Plugins/Experimental/MeshModelingToolsetExp/Source/MeshModelingToolsEditorOnlyExp/Public/SkeletalMesh/SkinWeightsPaintTool.h
+++ b/Engine/Plugins/Experimental/MeshModelingToolsetExp/Source/MeshModelingToolsEditorOnlyExp/Public/SkeletalMesh/SkinWeightsPaintTool.h
@@ -35,15 +35,15 @@ UENUM()
 enum class EWeightEditMode : uint8
 {
 	Brush,
-	Selection,
+	Vertices
 };
 
 // weight color mode
 UENUM()
 enum class EWeightColorMode : uint8
 {
-	Greyscale,
-	ColorRamp,
+	MinMax,
+	Ramp,
 };
 
 // brush falloff mode
@@ -239,8 +239,27 @@ public:
 	UMeshSurfacePointTool* CreateNewTool(const FToolBuilderState& SceneState) const override;
 };
 
+// for saveing/restoring the brush settings separately for each brush mode (Add, Replace, etc...)
+USTRUCT()
+struct FSkinWeightBrushConfig
+{
+	GENERATED_BODY()
+	
+	UPROPERTY()
+	float Strength = 1.f;
+
+	UPROPERTY()
+	float Radius = 20.0f;
+	
+	UPROPERTY()
+	float Falloff = 1.0f;
+
+	UPROPERTY()
+	EWeightBrushFalloffMode FalloffMode = EWeightBrushFalloffMode::Surface;
+};
+
 // Container for properties displayed in Details panel while using USkinWeightsPaintTool
-UCLASS()
+UCLASS(config = EditorSettings)
 class MESHMODELINGTOOLSEDITORONLYEXP_API USkinWeightsPaintToolProperties : public UBrushBaseProperties
 {
 	GENERATED_BODY()
@@ -256,8 +275,6 @@ public:
 	// custom brush modes and falloff types
 	UPROPERTY()
 	EWeightEditOperation BrushMode;
-	UPROPERTY()
-	EWeightBrushFalloffMode FalloffMode;
 
 	// weight color properties
 	UPROPERTY(EditAnywhere, Category = WeightColors)
@@ -271,17 +288,28 @@ public:
 	bool bColorModeChanged = false;
 
 	// weight editing arguments
-	UPROPERTY()
+	UPROPERTY(Config = EditorSettings)
 	TEnumAsByte<EAxis::Type> MirrorAxis = EAxis::X;
-	UPROPERTY()
+	UPROPERTY(Config = EditorSettings)
 	EMirrorDirection MirrorDirection = EMirrorDirection::PositiveToNegative;
-	UPROPERTY()
-	EWeightEditOperation FloodMode;
-	UPROPERTY()
+	UPROPERTY(Config = EditorSettings)
 	float FloodValue = 1.f;
-	UPROPERTY()
+	UPROPERTY(Config = EditorSettings)
 	float PruneValue = 0.01;
 
+	// save/restore user specified settings for each tool mode
+	FSkinWeightBrushConfig& GetBrushConfig();
+	TMap<EWeightEditOperation, FSkinWeightBrushConfig*> BrushConfigs;
+	UPROPERTY(Config = EditorSettings)
+	FSkinWeightBrushConfig BrushConfigAdd;
+	UPROPERTY(Config = EditorSettings)
+	FSkinWeightBrushConfig BrushConfigReplace;
+	UPROPERTY(Config = EditorSettings)
+	FSkinWeightBrushConfig BrushConfigMultiply;
+	UPROPERTY(Config = EditorSettings)
+	FSkinWeightBrushConfig BrushConfigRelax;
+
+	// pointer back to paint tool
 	TObjectPtr<USkinWeightsPaintTool> WeightTool;
 };
 
@@ -310,6 +338,10 @@ public:
 	virtual void DrawHUD(FCanvas* Canvas, IToolsContextRenderAPI* RenderAPI) override;
 	virtual void Render(IToolsContextRenderAPI* RenderAPI) override;
 
+	// IInteractiveToolCameraFocusAPI
+	virtual bool SupportsWorldSpaceFocusBox() override { return true; }
+	virtual FBox GetWorldSpaceFocusBox() override;
+
 	void ExternalUpdateWeights(const int32 BoneIndex, const TMap<int32, float>& IndexValues);
 
 	// weight editing operations (selection based)
diff --git a/Engine/Plugins/Experimental/MeshModelingToolsetExp/Source/ModelingEditorUI/Private/SkinWeightDetailCustomization.cpp b/Engine/Plugins/Experimental/MeshModelingToolsetExp/Source/ModelingEditorUI/Private/SkinWeightDetailCustomization.cpp
--- a/Engine/Plugins/Experimental/MeshModelingToolsetExp/Source/ModelingEditorUI/Private/SkinWeightDetailCustomization.cpp
+++ b/Engine/Plugins/Experimental/MeshModelingToolsetExp/Source/ModelingEditorUI/Private/SkinWeightDetailCustomization.cpp
@@ -14,8 +14,16 @@
 
 #define LOCTEXT_NAMESPACE "SkinWeightToolSettingsEditor"
 
+// layout constants
+float FSkinWeightDetailCustomization::WeightSliderWidths = 150.0f;
+float FSkinWeightDetailCustomization::WeightEditingLabelsPercent = 0.40f;
+float FSkinWeightDetailCustomization::WeightEditVerticalPadding = 4.0f;
+float FSkinWeightDetailCustomization::WeightEditHorizontalPadding = 2.0f;
+
 void FSkinWeightDetailCustomization::CustomizeDetails(IDetailLayoutBuilder& DetailBuilder)
 {
+	CurrentDetailBuilder = &DetailBuilder;
+	
 	TArray<TWeakObjectPtr<UObject>> DetailObjects;
 	DetailBuilder.GetObjectsBeingCustomized(DetailObjects);
 
@@ -23,15 +31,10 @@ void FSkinWeightDetailCustomization::CustomizeDetails(IDetailLayoutBuilder& Deta
 	ensure(DetailObjects.Num()==1);
 	SkinToolSettings = Cast<USkinWeightsPaintToolProperties>(DetailObjects[0]);
 
-	// layout constants
-	constexpr float WeightSliderWidths = 150.0f;
-	constexpr float WeightEditingLabelsPercent = 0.40f;
-	constexpr float WeightEditVerticalPadding = 4.0f;
-
 	// custom display of falloff mode as segmented toggle buttons
 	IDetailCategoryBuilder& EditModeCategory = DetailBuilder.EditCategory("Weight Editing Mode", FText::GetEmpty(), ECategoryPriority::Important);
 
-	// add segmented control toggle for brush falloff modes ("Surface" or "Volume")
+	// add segmented control toggle for editing modes ("Brush" or "Selection")
 	EditModeCategory.AddCustomRow(LOCTEXT("EditModeCategory", "Weight Editing Mode"), false)
 	.WholeRowContent()
 	[
@@ -50,14 +53,89 @@ void FSkinWeightDetailCustomization::CustomizeDetails(IDetailLayoutBuilder& Deta
 			{
 				SkinToolSettings->EditingMode = Mode;
 				SkinToolSettings->WeightTool->ToggleEditingMode();
+				if (CurrentDetailBuilder)
+				{
+					CurrentDetailBuilder->ForceRefreshDetails();
+				}
 			})
 			+SSegmentedControl<EWeightEditMode>::Slot(EWeightEditMode::Brush)
 			.Text(LOCTEXT("BrushEditMode", "Brush"))
-			+ SSegmentedControl<EWeightEditMode>::Slot(EWeightEditMode::Selection)
-			.Text(LOCTEXT("SelectionEditMode", "Selection"))
+			+ SSegmentedControl<EWeightEditMode>::Slot(EWeightEditMode::Vertices)
+			.Text(LOCTEXT("VertexEditMode", "Vertices"))
 		]
 	];
+
+	// BRUSH editing mode UI
+	if (SkinToolSettings->EditingMode == EWeightEditMode::Brush)
+	{
+		AddBrushUI(DetailBuilder);
+	}
+
+	// VERTEX editing mode UI
+	if (SkinToolSettings->EditingMode == EWeightEditMode::Vertices)
+	{
+		AddSelectionUI(DetailBuilder);
+	}
 	
+	// COLOR MODE category
+	IDetailCategoryBuilder& WeightColorsCategory = DetailBuilder.EditCategory("WeightColors", FText::GetEmpty(), ECategoryPriority::Important);
+	WeightColorsCategory.InitiallyCollapsed(true);
+	WeightColorsCategory.AddCustomRow(LOCTEXT("ColorModeCategory", "Color Mode"), false)
+	.NameContent()
+	[
+		SNew(STextBlock)
+		.Text(LOCTEXT("ColorModeLabel", "Color Mode"))
+		.Font(FAppStyle::GetFontStyle(TEXT("PropertyWindow.NormalFont")))
+		.ToolTipText(LOCTEXT("ColorModeTooltip", "Determines how the weight colors are displayed."))
+	]
+	.ValueContent()
+	[
+		SNew(SHorizontalBox)
+		+ SHorizontalBox::Slot()
+		[
+			SNew(SBox)
+			.Padding(2.0f)
+			.HAlign(HAlign_Center)
+			[
+				SNew(SSegmentedControl<EWeightColorMode>)
+				.Value_Lambda([this]()
+				{
+					return SkinToolSettings->ColorMode;
+				})
+				.OnValueChanged_Lambda([this](EWeightColorMode Mode)
+				{
+					SkinToolSettings->ColorMode = Mode;
+					SkinToolSettings->bColorModeChanged = true;
+				})
+				+ SSegmentedControl<EWeightColorMode>::Slot(EWeightColorMode::MinMax)
+				.Text(LOCTEXT("MinMaxMode", "Min / Max"))
+				+SSegmentedControl<EWeightColorMode>::Slot(EWeightColorMode::Ramp)
+				.Text(LOCTEXT("RampMode", "Color Ramp"))
+			]
+		]
+	];
+
+	// hide all base brush properties that have been customized
+	const TSharedRef<IPropertyHandle> BrushModeHandle = DetailBuilder.GetProperty(GET_MEMBER_NAME_CHECKED(USkinWeightsPaintToolProperties, BrushMode));
+	DetailBuilder.HideProperty(BrushModeHandle);
+	const TSharedRef<IPropertyHandle> BrushSizeHandle = DetailBuilder.GetProperty(GET_MEMBER_NAME_CHECKED(UBrushBaseProperties, BrushSize), UBrushBaseProperties::StaticClass());
+	DetailBuilder.HideProperty(BrushSizeHandle);
+	const TSharedRef<IPropertyHandle> BrushStrengthHandle = DetailBuilder.GetProperty(GET_MEMBER_NAME_CHECKED(UBrushBaseProperties, BrushStrength), UBrushBaseProperties::StaticClass());
+	DetailBuilder.HideProperty(BrushStrengthHandle);
+	const TSharedRef<IPropertyHandle> BrushFalloffHandle = DetailBuilder.GetProperty(GET_MEMBER_NAME_CHECKED(UBrushBaseProperties, BrushFalloffAmount), UBrushBaseProperties::StaticClass());
+	DetailBuilder.HideProperty(BrushFalloffHandle);
+	const TSharedRef<IPropertyHandle> BrushRadiusHandle = DetailBuilder.GetProperty(GET_MEMBER_NAME_CHECKED(UBrushBaseProperties, BrushRadius), UBrushBaseProperties::StaticClass());
+	DetailBuilder.HideProperty(BrushRadiusHandle);
+	const TSharedRef<IPropertyHandle> SpecifyRadiusHandle = DetailBuilder.GetProperty(GET_MEMBER_NAME_CHECKED(UBrushBaseProperties, bSpecifyRadius), UBrushBaseProperties::StaticClass());
+	DetailBuilder.HideProperty(SpecifyRadiusHandle);
+	const TSharedRef<IPropertyHandle> EditModePropHandle = DetailBuilder.GetProperty(GET_MEMBER_NAME_CHECKED(USkinWeightsPaintToolProperties, EditingMode));
+	DetailBuilder.HideProperty(EditModePropHandle);
+	const TSharedRef<IPropertyHandle> ColorModePropHandle = DetailBuilder.GetProperty(GET_MEMBER_NAME_CHECKED(USkinWeightsPaintToolProperties, ColorMode));
+	DetailBuilder.HideProperty(ColorModePropHandle);
+}
+
+void FSkinWeightDetailCustomization::AddBrushUI(IDetailLayoutBuilder& DetailBuilder)
+{
 	// custom display of falloff mode as segmented toggle buttons
 	IDetailCategoryBuilder& BrushCategory = DetailBuilder.EditCategory("Brush", FText::GetEmpty(), ECategoryPriority::Important);
 
@@ -69,15 +147,14 @@ void FSkinWeightDetailCustomization::CustomizeDetails(IDetailLayoutBuilder& Deta
 		.Padding(2.0f)
 		[
 			SNew(SSegmentedControl<EWeightEditOperation>)
-			.ToolTipText(LOCTEXT("WeightEditModeTooltip",
-					"Add: applies the current weight PLUS the strength value to the new weight.\n"
-					"Remove: applies the current weight MINUS the strength value to the new weight.\n"
-					"Multiply: applies the current weight MULTIPLIED by the strength value to the new weight.\n"
-					"Relax: applies the average of the connected (by edge) vertex weights to the new vertex weight, blended by the strength.\n"))
-			.IsEnabled_Lambda([this]()
-			{
-				return SkinToolSettings->EditingMode == EWeightEditMode::Brush;
-			})
+			.ToolTipText(LOCTEXT("FloodTooltip",
+				"Add: applies the current weight plus the flood value to the new weight.\n"
+				"Replace: applies the current weight minus the strength value to the new weight.\n"
+				"Multiply: applies the current weight multiplied by the strength value to the new weight.\n"
+				"Relax: applies the average of the connected (by edge) vertex weights to the new vertex weight, blended by the strength.\n"
+				"This command operates on the selected bone(s) and selected vertices.\n"
+				"If no bones are selected, ALL bones are considered.\n"
+				"If no vertices are selected, ALL vertices are considered."))
 			.Value_Lambda([this]()
 			{
 				return SkinToolSettings->BrushMode;
@@ -108,18 +185,15 @@ void FSkinWeightDetailCustomization::CustomizeDetails(IDetailLayoutBuilder& Deta
 			.ToolTipText(LOCTEXT("BrushFalloffModeTooltip",
 					"Surface: falloff is based on the distance along the surface from the brush center to nearby connected vertices.\n"
 					"Volume: falloff is based on the straight-line distance from the brush center to surrounding vertices.\n"))
-			.IsEnabled_Lambda([this]()
-			{
-				return SkinToolSettings->EditingMode == EWeightEditMode::Brush;
-			})
 			.Value_Lambda([this]()
 			{
-				return SkinToolSettings->FalloffMode;
+				return SkinToolSettings->GetBrushConfig().FalloffMode;
 			})
 			.OnValueChanged_Lambda([this](EWeightBrushFalloffMode Mode)
 			{
-				SkinToolSettings->FalloffMode = Mode;
 				SkinToolSettings->bColorModeChanged = true;
+				SkinToolSettings->GetBrushConfig().FalloffMode = Mode;
+				SkinToolSettings->SaveConfig();
 			})
 			+SSegmentedControl<EWeightBrushFalloffMode>::Slot(EWeightBrushFalloffMode::Surface)
 			.Text(LOCTEXT("SurfaceMode", "Surface"))
@@ -143,20 +217,21 @@ void FSkinWeightDetailCustomization::CustomizeDetails(IDetailLayoutBuilder& Deta
 		.MaxSliderValue(20.f)
 		.Value(10.0f)
 		.SupportDynamicSliderMaxValue(true)
-		.IsEnabled_Lambda([this]()
-		{
-			return SkinToolSettings->EditingMode == EWeightEditMode::Brush;
-		})
 		.Value_Lambda([this]()
 		{
-			return SkinToolSettings->BrushRadius;
+			return SkinToolSettings->GetBrushConfig().Radius;
 		})
 		.OnValueChanged_Lambda([this](float NewValue)
 		{
 			SkinToolSettings->BrushRadius = NewValue;
+			SkinToolSettings->GetBrushConfig().Radius = NewValue;
 			FPropertyChangedEvent PropertyChangedEvent(UBrushBaseProperties::StaticClass()->FindPropertyByName(GET_MEMBER_NAME_CHECKED(UBrushBaseProperties, BrushRadius)));
 			SkinToolSettings->PostEditChangeProperty(PropertyChangedEvent);
 		})
+		.OnValueCommitted_Lambda([this](float NewValue, ETextCommit::Type CommitType)
+		{
+			SkinToolSettings->SaveConfig();
+		})
 	];
 
 	BrushCategory.AddCustomRow(LOCTEXT("BrushStrengthCategory", "Brush Strength"), false)
@@ -175,20 +250,21 @@ void FSkinWeightDetailCustomization::CustomizeDetails(IDetailLayoutBuilder& Deta
 		.MaxSliderValue(1.f)
 		.Value(1.0f)
 		.SupportDynamicSliderMaxValue(true)
-		.IsEnabled_Lambda([this]()
-		{
-			return SkinToolSettings->EditingMode == EWeightEditMode::Brush;
-		})
 		.Value_Lambda([this]()
 		{
-			return SkinToolSettings->BrushStrength;
+			return SkinToolSettings->GetBrushConfig().Strength;
 		})
 		.OnValueChanged_Lambda([this](float NewValue)
 		{
 			SkinToolSettings->BrushStrength = NewValue;
+			SkinToolSettings->GetBrushConfig().Strength = NewValue;
 			FPropertyChangedEvent PropertyChangedEvent(UBrushBaseProperties::StaticClass()->FindPropertyByName(GET_MEMBER_NAME_CHECKED(UBrushBaseProperties, BrushStrength)));
 			SkinToolSettings->PostEditChangeProperty(PropertyChangedEvent);
 		})
+		.OnValueCommitted_Lambda([this](float NewValue, ETextCommit::Type CommitType)
+		{
+			SkinToolSettings->SaveConfig();
+		})
 	];
 
 	BrushCategory.AddCustomRow(LOCTEXT("BrushFalloffCategory", "Brush Falloff"), false)
@@ -204,22 +280,26 @@ void FSkinWeightDetailCustomization::CustomizeDetails(IDetailLayoutBuilder& Deta
 		SNew(SSpinBox<float>)
 		.MinValue(0.f)
 		.MaxValue(1.f)
-		.IsEnabled_Lambda([this]()
-		{
-			return SkinToolSettings->EditingMode == EWeightEditMode::Brush;
-		})
 		.Value_Lambda([this]()
 		{
-			return SkinToolSettings->BrushFalloffAmount;
+			return SkinToolSettings->GetBrushConfig().Falloff;
 		})
 		.OnValueChanged_Lambda([this](float NewValue)
 		{
 			SkinToolSettings->BrushFalloffAmount = NewValue;
+			SkinToolSettings->GetBrushConfig().Falloff = NewValue;
 			FPropertyChangedEvent PropertyChangedEvent(UBrushBaseProperties::StaticClass()->FindPropertyByName(GET_MEMBER_NAME_CHECKED(UBrushBaseProperties, BrushFalloffAmount)));
 			SkinToolSettings->PostEditChangeProperty(PropertyChangedEvent);
 		})
+		.OnValueCommitted_Lambda([this](float NewValue, ETextCommit::Type CommitType)
+		{
+			SkinToolSettings->SaveConfig();
+		})
 	];
+}
 
+void FSkinWeightDetailCustomization::AddSelectionUI(IDetailLayoutBuilder& DetailBuilder)
+{
 	// custom display of weight editing tools
 	IDetailCategoryBuilder& EditWeightsCategory = DetailBuilder.EditCategory("EditWeights", FText::GetEmpty(), ECategoryPriority::Important);
 	EditWeightsCategory.InitiallyCollapsed(true);
@@ -242,10 +322,6 @@ void FSkinWeightDetailCustomization::CustomizeDetails(IDetailLayoutBuilder& Deta
 					"This command operates on the selected bone(s) and selected vertices.\n "
 					"If no bones are selected, ALL bone weights are considered.\n "
 					"If no vertices are selected, ALL vertices are considered."))
-			.IsEnabled_Lambda([this]()
-			{
-				return SkinToolSettings->EditingMode == EWeightEditMode::Selection;
-			})
 			.OnClicked_Lambda([this]()
 			{
 				SkinToolSettings->WeightTool->AverageWeights();
@@ -266,7 +342,7 @@ void FSkinWeightDetailCustomization::CustomizeDetails(IDetailLayoutBuilder& Deta
 					"If no vertices are selected, ALL vertices are considered."))
 			.IsEnabled_Lambda([this]()
 			{
-				return SkinToolSettings->EditingMode == EWeightEditMode::Selection;
+				return SkinToolSettings->EditingMode == EWeightEditMode::Vertices;
 			})
 			.OnClicked_Lambda([this]()
 			{
@@ -283,7 +359,7 @@ void FSkinWeightDetailCustomization::CustomizeDetails(IDetailLayoutBuilder& Deta
 		SNew(SVerticalBox)
 		
 		+ SVerticalBox::Slot()
-		.Padding(0.f, WeightEditVerticalPadding)
+		.Padding(WeightEditHorizontalPadding, WeightEditVerticalPadding)
 		[
 			SNew(SHorizontalBox)
 			
@@ -308,10 +384,6 @@ void FSkinWeightDetailCustomization::CustomizeDetails(IDetailLayoutBuilder& Deta
 						"X: copies weights across the YZ plane.\n"
 						"Y: copies weights across the XZ plane.\n"
 						"Z: copies weights across the XY plane."))
-					.IsEnabled_Lambda([this]()
-					{
-						return SkinToolSettings->EditingMode == EWeightEditMode::Selection;
-					})
 					.Value_Lambda([this]()
 					{
 						return SkinToolSettings->MirrorAxis;
@@ -332,10 +404,6 @@ void FSkinWeightDetailCustomization::CustomizeDetails(IDetailLayoutBuilder& Deta
 				[
 					SNew(SSegmentedControl<EMirrorDirection>)
 					.ToolTipText(LOCTEXT("MirrorDirectionTooltip", "The direction that determines what side of the plane to copy weights from."))
-					.IsEnabled_Lambda([this]()
-					{
-						return SkinToolSettings->EditingMode == EWeightEditMode::Selection;
-					})
 					.Value_Lambda([this]()
 					{
 						return SkinToolSettings->MirrorDirection;
@@ -365,10 +433,6 @@ void FSkinWeightDetailCustomization::CustomizeDetails(IDetailLayoutBuilder& Deta
 					"This command operates on the selected bone(s) and selected vertices.\n "
 					"If no bones are selected, ALL bone weights are considered.\n "
 					"If no vertices are selected, ALL vertices are considered."))
-				.IsEnabled_Lambda([this]()
-				{
-					return SkinToolSettings->EditingMode == EWeightEditMode::Selection;
-				})
 				.OnClicked_Lambda([this]()
 				{
 					SkinToolSettings->WeightTool->MirrorWeights(SkinToolSettings->MirrorAxis, SkinToolSettings->MirrorDirection);
@@ -408,10 +472,6 @@ void FSkinWeightDetailCustomization::CustomizeDetails(IDetailLayoutBuilder& Deta
 				.MaxSliderValue(1.f)
 				.Value(1.0f)
 				.SupportDynamicSliderMaxValue(true)
-				.IsEnabled_Lambda([this]()
-				{
-					return SkinToolSettings->EditingMode == EWeightEditMode::Selection;
-				})
 				.Value_Lambda([this]()
 				{
 					return SkinToolSettings->FloodValue;
@@ -420,76 +480,63 @@ void FSkinWeightDetailCustomization::CustomizeDetails(IDetailLayoutBuilder& Deta
 				{
 					SkinToolSettings->FloodValue = NewValue;
 				})
+				.OnValueCommitted_Lambda([this](float NewValue, ETextCommit::Type CommitType)
+				{
+					SkinToolSettings->SaveConfig();
+				})
 			]
 		]
 
-		+ SVerticalBox::Slot()
-		.Padding(0.f, WeightEditVerticalPadding)
+		+SVerticalBox::Slot()
+		.Padding(WeightEditHorizontalPadding, WeightEditVerticalPadding)
 		[
 			SNew(SHorizontalBox)
-				
-			+SHorizontalBox::Slot()
-			.VAlign(VAlign_Center)
-			.FillWidth(WeightEditingLabelsPercent)
-			[
-				SNew(STextBlock)
-				.Text(LOCTEXT("FloodOperationLabel", "Flood Operation"))
-				.Font(FAppStyle::GetFontStyle(TEXT("PropertyWindow.NormalFont")))
-				.ToolTipText(LOCTEXT("FloodOperationTooltip", "The operation to perform when Flood button is pressed."))
-			]
 
-			+SHorizontalBox::Slot()
-			.FillWidth(1.f)
+			+ SHorizontalBox::Slot()
 			[
-				SNew(SSegmentedControl<EWeightEditOperation>)
-				.ToolTipText(LOCTEXT("FloodOperationButtonTooltip",
-					"Add: applies the current weight PLUS the flood value to the new weight.\n"
-					"Remove: applies the current weight MINUS the flood value to the new weight.\n"
-					"Multiply: applies the current weight MULTIPLIED by the flood value to the new weight.\n"
-					"Relax: applies the average of the connected (by edge) vertex weights to the new vertex weight.\n"))
-				.IsEnabled_Lambda([this]()
-				{
-					return SkinToolSettings->EditingMode == EWeightEditMode::Selection;
-				})
-				.Value_Lambda([this]()
+				SNew(SButton)
+				.HAlign(HAlign_Center)
+				.Text(LOCTEXT("AddWeightsButtonLabel", "Add"))
+				.ToolTipText(LOCTEXT("AddOpTooltip", "Add: applies the current weight plus the flood value to the new weight.\n"))
+				.OnClicked_Lambda([this]()
 				{
-					return SkinToolSettings->FloodMode;
+					SkinToolSettings->WeightTool->FloodWeights(SkinToolSettings->FloodValue, EWeightEditOperation::Add);
+					return FReply::Handled();
 				})
-				.OnValueChanged_Lambda([this](EWeightEditOperation Mode)
+			]
+			+ SHorizontalBox::Slot()
+			[
+				SNew(SButton)
+				.HAlign(HAlign_Center)
+				.Text(LOCTEXT("ReplaceWeightsButtonLabel", "Replace"))
+				.ToolTipText(LOCTEXT("ReplaceOpTooltip", "Replace: applies the current weight minus the strength value to the new weight.\n"))
+				.OnClicked_Lambda([this]()
 				{
-					SkinToolSettings->FloodMode = Mode;
+					SkinToolSettings->WeightTool->FloodWeights(SkinToolSettings->FloodValue, EWeightEditOperation::Replace);
+					return FReply::Handled();
 				})
-				+SSegmentedControl<EWeightEditOperation>::Slot(EWeightEditOperation::Add)
-				.Text(LOCTEXT("FloodAddMode", "Add"))
-				+ SSegmentedControl<EWeightEditOperation>::Slot(EWeightEditOperation::Replace)
-				.Text(LOCTEXT("FloodReplaceMode", "Replace"))
-				+ SSegmentedControl<EWeightEditOperation>::Slot(EWeightEditOperation::Multiply)
-				.Text(LOCTEXT("FloodMultiplyMode", "Multiply"))
-				+ SSegmentedControl<EWeightEditOperation>::Slot(EWeightEditOperation::Relax)
-				.Text(LOCTEXT("FloodRelaxMode", "Relax"))
 			]
-		]
-		
-		+SVerticalBox::Slot()
-		.Padding(0.f, WeightEditVerticalPadding)
-		[
-			SNew(SBox)
+			+ SHorizontalBox::Slot()
 			[
 				SNew(SButton)
 				.HAlign(HAlign_Center)
-				.Text(LOCTEXT("FloodWeightsButtonLabel", "Flood"))
-				.ToolTipText(LOCTEXT("FloodButtonTooltip",
-					"Modifies the vertex weights according to the chosen operation. See the operation tooltip for details.\n"
-					"This command operates on the selected bone(s) and selected vertices.\n"
-					"If no bones are selected, ALL bones are considered.\n"
-					"If no vertices are selected, ALL vertices are considered."))
-				.IsEnabled_Lambda([this]()
+				.Text(LOCTEXT("MultiplyeightsButtonLabel", "Multiply"))
+				.ToolTipText(LOCTEXT("MultiplyOpTooltip", "Multiply: applies the current weight multiplied by the strength value to the new weight.\n"))
+				.OnClicked_Lambda([this]()
 				{
-					return SkinToolSettings->EditingMode == EWeightEditMode::Selection;
+					SkinToolSettings->WeightTool->FloodWeights(SkinToolSettings->FloodValue, EWeightEditOperation::Multiply);
+					return FReply::Handled();
 				})
+			]
+			+ SHorizontalBox::Slot()
+			[
+				SNew(SButton)
+				.HAlign(HAlign_Center)
+				.Text(LOCTEXT("RelaxWeightsButtonLabel", "Relax"))
+				.ToolTipText(LOCTEXT("RelaxeOpTooltip", "Relax: applies the average of the connected (by edge) vertex weights to the new vertex weight, blended by the strength.\n"))
 				.OnClicked_Lambda([this]()
 				{
-					SkinToolSettings->WeightTool->FloodWeights(SkinToolSettings->FloodValue, SkinToolSettings->FloodMode);
+					SkinToolSettings->WeightTool->FloodWeights(SkinToolSettings->FloodValue, EWeightEditOperation::Relax);
 					return FReply::Handled();
 				})
 			]
@@ -503,7 +550,7 @@ void FSkinWeightDetailCustomization::CustomizeDetails(IDetailLayoutBuilder& Deta
 		SNew(SVerticalBox)
 
 		+ SVerticalBox::Slot()
-		.Padding(0.f, WeightEditVerticalPadding)
+		.Padding(WeightEditHorizontalPadding, WeightEditVerticalPadding)
 		[
 			SNew(SHorizontalBox)
 					
@@ -523,10 +570,6 @@ void FSkinWeightDetailCustomization::CustomizeDetails(IDetailLayoutBuilder& Deta
 				SNew(SSpinBox<float>)
 				.MinValue(0.f)
 				.MaxValue(1.f)
-				.IsEnabled_Lambda([this]()
-				{
-					return SkinToolSettings->EditingMode == EWeightEditMode::Selection;
-				})
 				.Value_Lambda([this]()
 				{
 					return SkinToolSettings->PruneValue;
@@ -535,11 +578,15 @@ void FSkinWeightDetailCustomization::CustomizeDetails(IDetailLayoutBuilder& Deta
 				{
 					SkinToolSettings->PruneValue = NewValue;
 				})
+				.OnValueCommitted_Lambda([this](float NewValue, ETextCommit::Type CommitType)
+				{
+					SkinToolSettings->SaveConfig();
+				})
 			]
 		]
 		
 		+SVerticalBox::Slot()
-		.Padding(0.f, WeightEditVerticalPadding)
+		.Padding(WeightEditHorizontalPadding, WeightEditVerticalPadding)
 		[
 			SNew(SBox)
 			.MinDesiredWidth(WeightSliderWidths)
@@ -552,10 +599,6 @@ void FSkinWeightDetailCustomization::CustomizeDetails(IDetailLayoutBuilder& Deta
 					"This command operates on the selected bone(s) and selected vertices.\n "
 					"If no bones are selected, ALL bone weights are considered.\n "
 					"If no vertices are selected, ALL vertices are considered."))
-				.IsEnabled_Lambda([this]()
-				{
-					return SkinToolSettings->EditingMode == EWeightEditMode::Selection;
-				})
 				.OnClicked_Lambda([this]()
 				{
 					SkinToolSettings->WeightTool->PruneWeights(SkinToolSettings->PruneValue);
@@ -564,64 +607,6 @@ void FSkinWeightDetailCustomization::CustomizeDetails(IDetailLayoutBuilder& Deta
 			]
 		]
 	];
-	
-	// COLOR MODE category
-	IDetailCategoryBuilder& WeightColorsCategory = DetailBuilder.EditCategory("WeightColors", FText::GetEmpty(), ECategoryPriority::Important);
-	WeightColorsCategory.InitiallyCollapsed(true);
-	WeightColorsCategory.AddCustomRow(LOCTEXT("ColorModeCategory", "Color Mode"), false)
-	.NameContent()
-	[
-		SNew(STextBlock)
-		.Text(LOCTEXT("ColorModeLabel", "Color Mode"))
-		.Font(FAppStyle::GetFontStyle(TEXT("PropertyWindow.NormalFont")))
-		.ToolTipText(LOCTEXT("ColorModeTooltip", "Determines how the weight colors are displayed."))
-	]
-	.ValueContent()
-	[
-		SNew(SHorizontalBox)
-		+ SHorizontalBox::Slot()
-		[
-			SNew(SBox)
-			.Padding(2.0f)
-			.HAlign(HAlign_Center)
-			[
-				SNew(SSegmentedControl<EWeightColorMode>)
-				.Value_Lambda([this]()
-				{
-					return SkinToolSettings->ColorMode;
-				})
-				.OnValueChanged_Lambda([this](EWeightColorMode Mode)
-				{
-					SkinToolSettings->ColorMode = Mode;
-					SkinToolSettings->bColorModeChanged = true;
-				})
-				+ SSegmentedControl<EWeightColorMode>::Slot(EWeightColorMode::Greyscale)
-				.Text(LOCTEXT("GreyscaleMode", "Greyscale"))
-				+SSegmentedControl<EWeightColorMode>::Slot(EWeightColorMode::ColorRamp)
-				.Text(LOCTEXT("RampMode", "Color Ramp"))
-			]
-		]
-	];
-
-	// hide all base brush properties that have been customized
-	const TSharedRef<IPropertyHandle> BrushModeHandle = DetailBuilder.GetProperty(GET_MEMBER_NAME_CHECKED(USkinWeightsPaintToolProperties, BrushMode));
-	DetailBuilder.HideProperty(BrushModeHandle);
-	const TSharedRef<IPropertyHandle> BrushSizeHandle = DetailBuilder.GetProperty(GET_MEMBER_NAME_CHECKED(UBrushBaseProperties, BrushSize), UBrushBaseProperties::StaticClass());
-	DetailBuilder.HideProperty(BrushSizeHandle);
-	const TSharedRef<IPropertyHandle> BrushStrengthHandle = DetailBuilder.GetProperty(GET_MEMBER_NAME_CHECKED(UBrushBaseProperties, BrushStrength), UBrushBaseProperties::StaticClass());
-	DetailBuilder.HideProperty(BrushStrengthHandle);
-	const TSharedRef<IPropertyHandle> BrushFalloffHandle = DetailBuilder.GetProperty(GET_MEMBER_NAME_CHECKED(UBrushBaseProperties, BrushFalloffAmount), UBrushBaseProperties::StaticClass());
-	DetailBuilder.HideProperty(BrushFalloffHandle);
-	const TSharedRef<IPropertyHandle> BrushRadiusHandle = DetailBuilder.GetProperty(GET_MEMBER_NAME_CHECKED(UBrushBaseProperties, BrushRadius), UBrushBaseProperties::StaticClass());
-	DetailBuilder.HideProperty(BrushRadiusHandle);
-	const TSharedRef<IPropertyHandle> SpecifyRadiusHandle = DetailBuilder.GetProperty(GET_MEMBER_NAME_CHECKED(UBrushBaseProperties, bSpecifyRadius), UBrushBaseProperties::StaticClass());
-	DetailBuilder.HideProperty(SpecifyRadiusHandle);
-	const TSharedRef<IPropertyHandle> FalloffPropHandle = DetailBuilder.GetProperty(GET_MEMBER_NAME_CHECKED(USkinWeightsPaintToolProperties, FalloffMode));
-	DetailBuilder.HideProperty(FalloffPropHandle);
-	const TSharedRef<IPropertyHandle> EditModePropHandle = DetailBuilder.GetProperty(GET_MEMBER_NAME_CHECKED(USkinWeightsPaintToolProperties, EditingMode));
-	DetailBuilder.HideProperty(EditModePropHandle);
-	const TSharedRef<IPropertyHandle> ColorModePropHandle = DetailBuilder.GetProperty(GET_MEMBER_NAME_CHECKED(USkinWeightsPaintToolProperties, ColorMode));
-	DetailBuilder.HideProperty(ColorModePropHandle);
 }
 
-#undef LOCTEXT_NAMESPACE
\ No newline at end of file
+#undef LOCTEXT_NAMESPACE
diff --git a/Engine/Plugins/Experimental/MeshModelingToolsetExp/Source/ModelingEditorUI/Public/SkinWeightDetailCustomization.h b/Engine/Plugins/Experimental/MeshModelingToolsetExp/Source/ModelingEditorUI/Public/SkinWeightDetailCustomization.h
--- a/Engine/Plugins/Experimental/MeshModelingToolsetExp/Source/ModelingEditorUI/Public/SkinWeightDetailCustomization.h
+++ b/Engine/Plugins/Experimental/MeshModelingToolsetExp/Source/ModelingEditorUI/Public/SkinWeightDetailCustomization.h
@@ -21,5 +21,15 @@ public:
 	virtual void CustomizeDetails(IDetailLayoutBuilder& DetailBuilder) override;
 
 private:
+
+	void AddBrushUI(IDetailLayoutBuilder& DetailBuilder);
+	void AddSelectionUI(IDetailLayoutBuilder& DetailBuilder);
+
+	IDetailLayoutBuilder* CurrentDetailBuilder;
 	TWeakObjectPtr<USkinWeightsPaintToolProperties> SkinToolSettings;
+
+	static float WeightSliderWidths;
+	static float WeightEditingLabelsPercent;
+	static float WeightEditVerticalPadding;
+	static float WeightEditHorizontalPadding;
 };
