Fixes crash if framing selection while closing skeletal mesh editor (minor edge case)

diff --git a/Engine/Plugins/Experimental/MeshModelingToolsetExp/Source/MeshModelingToolsEditorOnlyExp/Private/SkeletalMesh/SkinWeightsPaintTool.cpp b/Engine/Plugins/Experimental/MeshModelingToolsetExp/Source/MeshModelingToolsEditorOnlyExp/Private/SkeletalMesh/SkinWeightsPaintTool.cpp
--- a/Engine/Plugins/Experimental/MeshModelingToolsetExp/Source/MeshModelingToolsEditorOnlyExp/Private/SkeletalMesh/SkinWeightsPaintTool.cpp
+++ b/Engine/Plugins/Experimental/MeshModelingToolsetExp/Source/MeshModelingToolsEditorOnlyExp/Private/SkeletalMesh/SkinWeightsPaintTool.cpp
@@ -776,20 +776,24 @@ FBox USkinWeightsPaintTool::GetWorldSpaceFocusBox()
 {
 	// TODO add support for selected bones
 
-	const FGroupTopologySelection& Selection = PolygonSelectionMechanic->GetActiveSelection();
-	if (!Selection.SelectedCornerIDs.IsEmpty())
+	// frame vertex selection if there is one
+	if (PolygonSelectionMechanic)
 	{
-		TArray<VertexIndex> SelectedVertexIndices = Selection.SelectedCornerIDs.Array();
-		const FDynamicMesh3* Mesh = PreviewMesh->GetMesh();
-		FTransform3d Transform(PreviewMesh->GetTransform());
-		FAxisAlignedBox3d Bounds = FAxisAlignedBox3d::Empty();
-		for (const int32 VertexID : SelectedVertexIndices)
-		{
-			Bounds.Contain(Transform.TransformPosition(Mesh->GetVertex(VertexID)));
-		}
-		if (Bounds.MaxDim() > FMathf::ZeroTolerance)
+		const FGroupTopologySelection& Selection = PolygonSelectionMechanic->GetActiveSelection();
+		if (!Selection.SelectedCornerIDs.IsEmpty())
 		{
-			return static_cast<FBox>(Bounds);
+			TArray<VertexIndex> SelectedVertexIndices = Selection.SelectedCornerIDs.Array();
+			const FDynamicMesh3* Mesh = PreviewMesh->GetMesh();
+			FTransform3d Transform(PreviewMesh->GetTransform());
+			FAxisAlignedBox3d Bounds = FAxisAlignedBox3d::Empty();
+			for (const int32 VertexID : SelectedVertexIndices)
+			{
+				Bounds.Contain(Transform.TransformPosition(Mesh->GetVertex(VertexID)));
+			}
+			if (Bounds.MaxDim() > FMathf::ZeroTolerance)
+			{
+				return static_cast<FBox>(Bounds);
+			}
 		}
 	}
 	
