CombineMeshInstances: add concept of a 'Subset' of instances, which are processed separately for hidden removal and mesh-merging optimizations. Refactor internals to process things by-subset. Behavior is unchanged if Subset Indices are not initialized.

diff --git a/Engine/Plugins/Experimental/MeshModelingToolsetExp/Source/GeometryProcessingAdapters/Private/CombineMeshInstancesImpl.cpp b/Engine/Plugins/Experimental/MeshModelingToolsetExp/Source/GeometryProcessingAdapters/Private/CombineMeshInstancesImpl.cpp
--- a/Engine/Plugins/Experimental/MeshModelingToolsetExp/Source/GeometryProcessingAdapters/Private/CombineMeshInstancesImpl.cpp
+++ b/Engine/Plugins/Experimental/MeshModelingToolsetExp/Source/GeometryProcessingAdapters/Private/CombineMeshInstancesImpl.cpp
@@ -99,6 +99,8 @@
 
 	// allow FMeshPartInstance to maintain link to external representation of instance
 	FIndex3i ExternalInstanceIndex = FIndex3i::Invalid();
+
+	int32 SubsetID = 0;
 };
 
 
@@ -206,6 +208,7 @@
 
 		FMeshPartInstance NewInstance;
 		NewInstance.ExternalInstanceIndex = FIndex3i(Index, 0, -1);
+		NewInstance.SubsetID = SourceMeshInstance.InstanceSubsetID;
 
 		if ( SourceMeshInstance.GroupDataIndex >= 0 && SourceMeshInstance.GroupDataIndex < SourceInstanceList.InstanceGroupDatas.Num() )
 		{
@@ -260,6 +263,7 @@
 
 		FMeshPartInstance NewInstance;
 		NewInstance.ExternalInstanceIndex = FIndex3i(Index, 1, -1);
+		NewInstance.SubsetID = SourceMeshInstance.InstanceSubsetID;
 
 		if (SourceMeshInstance.GroupDataIndex >= 0 && SourceMeshInstance.GroupDataIndex < SourceInstanceList.InstanceGroupDatas.Num())
 		{
@@ -1295,7 +1299,7 @@
 	int32 NumApproxLODs = FMath::Max(1, 
 		CombineOptions.NumLODs - CombineOptions.NumCopiedLODs - CombineOptions.NumSimplifiedLODs);
 
-	bool bVerbose = CVarGeometryCombineMeshInstancesVerbose.GetValueOnGameThread();
+	bool bVerbose = CVarGeometryCombineMeshInstancesVerbose.GetValueOnAnyThread();
 	ParallelFor(NumParts, [&](int32 Index)
 	{
 		TUniquePtr<FMeshPart>& Part = Assembly.Parts[Index];
@@ -2774,7 +2778,9 @@
 {
 	FDynamicMesh3 Mesh;
 	FDynamicMeshEditor Editor;
-	FDynamicMeshMaterialAttribute* MaterialIDs;
+	FDynamicMeshMaterialAttribute* MaterialIDs = nullptr;
 
+	FDynamicMeshPolygroupAttribute* SubsetIDs = nullptr;
+
 	FCombinedMeshLOD()
 		: Editor(&Mesh)
@@ -2788,6 +2794,14 @@
 		MaterialIDs = Mesh.Attributes()->GetMaterialID();
 	}
 
+
+	void SetMesh(FDynamicMesh3&& ExternalMesh)
+	{
+		Mesh = MoveTemp(ExternalMesh);
+		Editor = FDynamicMeshEditor(&Mesh);
+		check(Mesh.HasAttributes() && Mesh.Attributes()->HasPrimaryColors() && Mesh.Attributes()->HasMaterialID())
+		MaterialIDs = Mesh.Attributes()->GetMaterialID();
+	}
 };
 
 }
@@ -2926,204 +2940,36 @@
 
 
 
-// change this to build a single LOD, and separate versions for (eg) source mesh vs approx mesh
-// should we even bother w/ storing approx meshes? just generate them as needed?
 
-void BuildCombinedMesh(
-	const FMeshPartsAssembly& Assembly,
-	IGeometryProcessing_CombineMeshInstances::FOptions CombineOptions,
-	TArray<FDynamicMesh3>& CombinedMeshLODs)
+
+void ProcessCombinedLODChain(
+	TArray<FCombinedMeshLOD>& MeshLODs,
+	const TArray<double>& OptimizationTolerances,
+	int32 FirstVoxWrappedIndex,
+	IGeometryProcessing_CombineMeshInstances::FOptions CombineOptions
+)
 {
 	using namespace UE::Geometry;
-	bool bVerbose = CVarGeometryCombineMeshInstancesVerbose.GetValueOnGameThread();
+	bool bVerbose = CVarGeometryCombineMeshInstancesVerbose.GetValueOnAnyThread();
 
-	const bool bAppendMinimalLOD = false;
-
 	int32 NumLODs = CombineOptions.NumLODs;
-	int NumExtraLODs = (bAppendMinimalLOD) ? 1 : 0;
-	TArray<FCombinedMeshLOD> MeshLODs;
-	MeshLODs.SetNum(NumLODs + NumExtraLODs);
 
-	int FirstVoxWrappedIndex = 9999;
-	TArray<ECombinedLODType> LODTypes;
-	LODTypes.Init(ECombinedLODType::Approximated, NumLODs);
-	for (int32 LODLevel = 0; LODLevel < NumLODs; ++LODLevel)
-	{
-		if (LODLevel < CombineOptions.NumCopiedLODs)
-		{
-			LODTypes[LODLevel] = ECombinedLODType::Copied;
-		}
-		else if (LODLevel < CombineOptions.NumCopiedLODs + CombineOptions.NumSimplifiedLODs)
-		{
-			LODTypes[LODLevel] = ECombinedLODType::Simplified;
-		}
-		else if (LODLevel >= NumLODs - CombineOptions.NumCoarseLODs)
-		{
-			LODTypes[LODLevel] = ECombinedLODType::VoxWrapped;
-			FirstVoxWrappedIndex = FMath::Min(LODLevel, FirstVoxWrappedIndex);
-		}
-	}
-
-	//CombinedLOD0.Attributes()->SetNumPolygroupLayers(2);
-	//FDynamicMeshPolygroupAttribute* PartIDAttrib = AccumMesh.Attributes()->GetPolygroupLayer(0);
-	//FDynamicMeshPolygroupAttribute* PartInstanceMapAttrib = AccumMesh.Attributes()->GetPolygroupLayer(1);
-
-	int32 NumParts = Assembly.Parts.Num();
-
-	//for ( int32 SetIndex = 0; SetIndex < NumParts; ++SetIndex )
-	//{
-	//	CombinedLOD0.EnableMatchingAttributes( Assembly.SourceMeshGeometry[Index].OriginalMesh, false, false );
-	//}
-
-	for ( int32 SetIndex = 0; SetIndex < NumParts; ++SetIndex )
-	{
-		const TUniquePtr<FMeshPart>& Part = Assembly.Parts[SetIndex];
-		const FSourceGeometry& SourceGeometry = Assembly.SourceMeshGeometry[SetIndex];
-		const FOptimizedGeometry& OptimizedGeometry = Assembly.OptimizedMeshGeometry[SetIndex];
-
-		check(Part->bAllowMerging == true);		// do not support this yet
-
-		FMeshIndexMappings Mappings;
-
-		for (int32 LODLevel = 0; LODLevel < NumLODs; ++LODLevel)
-		{
-			const FDynamicMesh3* SourceAppendMesh = nullptr;
-			const FDynamicMesh3* ApproximateAppendMesh = nullptr;
-			const FDynamicMesh3* UseAppendMesh = nullptr;
-
-			// default approximate mesh to lowest-quality approximation (box), need to do this
-			// so that we always have something to swap to for Decorative parts
-			ApproximateAppendMesh = &OptimizedGeometry.ApproximateMeshLODs.Last();
-
-			ECombinedLODType LevelLODType = LODTypes[LODLevel];
-			if (LevelLODType == ECombinedLODType::Copied)
-			{
-				SourceAppendMesh = (LODLevel < SourceGeometry.SourceMeshLODs.Num()) ? 
-					&SourceGeometry.SourceMeshLODs[LODLevel] : &SourceGeometry.SourceMeshLODs.Last();
-				UseAppendMesh = SourceAppendMesh;
-			}
-			else if (LevelLODType == ECombinedLODType::Simplified)
-			{
-				int32 SimplifiedLODIndex = LODLevel - CombineOptions.NumCopiedLODs;
-				SourceAppendMesh = &OptimizedGeometry.SimplifiedMeshLODs[SimplifiedLODIndex];
-				UseAppendMesh = SourceAppendMesh;
-			}
-			else if (LevelLODType == ECombinedLODType::VoxWrapped)
-			{
-				SourceAppendMesh = &SourceGeometry.SourceMeshLODs.Last();
-				UseAppendMesh = SourceAppendMesh;
-			}
-			else // ECombinedLODType::Approximated
-			{
-				int32 ApproxLODIndex = LODLevel - CombineOptions.NumCopiedLODs - CombineOptions.NumSimplifiedLODs;
-				ApproximateAppendMesh = &OptimizedGeometry.ApproximateMeshLODs[ApproxLODIndex];
-				UseAppendMesh = ApproximateAppendMesh;
-			}
-
-			FCombinedMeshLOD& CombinedMeshLODData = MeshLODs[LODLevel];
-
-			for ( const FMeshPartInstance& Instance : Part->Instances )
-			{
-				bool bIsDecorativePart = (Instance.DetailLevel == EMeshDetailLevel::Decorative);
-
-				if (bIsDecorativePart)
-				{
-					// filter out detail parts at higher LODs, or if we are doing VoxWrap LOD
-					if ( LODLevel >= CombineOptions.FilterDecorativePartsLODLevel || LevelLODType == ECombinedLODType::VoxWrapped )
-					{
-						continue;
-					}
-					// at last detail part LOD, switch to approximate mesh
-					if (LODLevel >= (CombineOptions.FilterDecorativePartsLODLevel - CombineOptions.ApproximateDecorativePartLODs) )
-					{
-						UseAppendMesh = ApproximateAppendMesh;
-					}
-				}
-
-				// need to make a copy to run pre-process func
-				FDynamicMesh3 TempAppendMesh(*UseAppendMesh);
-				if (Assembly.PreProcessInstanceMeshFunc)
-				{
-					Assembly.PreProcessInstanceMeshFunc(TempAppendMesh, Instance);
-				}
-
-				// if part does not require UVs, get rid of them here to encourage merging downstream
-				if (LODLevel > 0 && Part->bPreserveUVs == false && TempAppendMesh.HasAttributes())
-				{
-					for (int32 UVLayer = 0; UVLayer < TempAppendMesh.Attributes()->NumUVLayers(); ++UVLayer)
-					{
-						TempAppendMesh.Attributes()->GetUVLayer(UVLayer)->ClearElements();
-					}
-					// if we have no UVs then tangents are invalid
-					TempAppendMesh.Attributes()->DisableTangents();
-				}
-
-				Mappings.Reset();
-				CombinedMeshLODData.Editor.AppendMesh(&TempAppendMesh, Mappings,
-					[&](int, const FVector3d& Pos) { return Instance.WorldTransform.TransformPosition(Pos); },
-					[&](int, const FVector3d& Normal) { return Instance.WorldTransform.TransformNormal(Normal); });
-
-				// append part ID stuff here
-
-				// could precompute these indexes for each instance?
-				// also for source mesh we could transfer material IDs correctly...
-				UMaterialInterface* UseMaterial = Instance.Materials[0];
-				const int32* FoundMaterialIndex = Assembly.MaterialMap.Find(UseMaterial);
-				int32 AssignMaterialIndex = (FoundMaterialIndex != nullptr) ? *FoundMaterialIndex : 0;
-
-				for (int32 tid : TempAppendMesh.TriangleIndicesItr())
-				{
-					CombinedMeshLODData.MaterialIDs->SetValue( Mappings.GetNewTriangle(tid), AssignMaterialIndex );
-				}
-			}
-		}
-	}
-
-	// Above we copied the last approximation LOD into the first VoxWrap LOD slot. However 
-	// approximate LOD may be very coarse, it's better to use a known-good starting mesh?
-	// Need to do this before starting hidden-removal task...
-	if (FirstVoxWrappedIndex < MeshLODs.Num())
-	{
-		MeshLODs[FirstVoxWrappedIndex].Mesh = MeshLODs[CombineOptions.NumCopiedLODs-1].Mesh;
-	}
-
-	// make a list of per-LOD geometric tolerances tha will drive additional optimization. 
-	// For copied and first simplified LODs, use Simplify Base Tolerance, and then increment
-	// for each successive LOD.  (todo: have a separate initial tolerance for Approx LODs?)
-	double CurTolerance = CombineOptions.SimplifyBaseTolerance;
-	TArray<double> OptimizationTolerances;
-	for (int32 LODIndex = 0; LODIndex < NumLODs; ++LODIndex)
-	{
-		if (LODTypes[LODIndex] == ECombinedLODType::Simplified && LODTypes[LODIndex - 1] == ECombinedLODType::Simplified)
-		{
-			CurTolerance *= 2.0;
-		}
-		else if (LODTypes[LODIndex] == ECombinedLODType::Approximated)
-		{
-			CurTolerance *= 2.0;
-		}
-		OptimizationTolerances.Add(CurTolerance);
-	}
-
-	//
-	// start hidden-removal passes on all meshes up to voxel LODs here, because we can compute voxel LOD at the same time
-	//
 	TArray<UE::Tasks::FTask> PendingRemoveHiddenTasks;
-	bool bRemoveHiddenFaces = 
-		(CombineOptions.RemoveHiddenFacesMethod != IGeometryProcessing_CombineMeshInstances::ERemoveHiddenFacesMode::None 
-		&& CVarGeometryCombineMeshInstancesRemoveHidden.GetValueOnGameThread() > 0);
+	bool bRemoveHiddenFaces =
+		(CombineOptions.RemoveHiddenFacesMethod != IGeometryProcessing_CombineMeshInstances::ERemoveHiddenFacesMode::None
+			&& CVarGeometryCombineMeshInstancesRemoveHidden.GetValueOnAnyThread() > 0);
 	if (bRemoveHiddenFaces)
 	{
 		for (int32 LODIndex = CombineOptions.RemoveHiddenStartLOD; LODIndex < NumLODs && LODIndex < FirstVoxWrappedIndex; ++LODIndex)
 		{
-			if ( bVerbose )
+			if (bVerbose)
 			{
 				UE_LOG(LogGeometry, Log, TEXT("  Optimizing LOD%d - Tris %6d Verts %6d"), LODIndex, MeshLODs[LODIndex].Mesh.TriangleCount(), MeshLODs[LODIndex].Mesh.VertexCount());
 			}
 
 			double UseTolerance = OptimizationTolerances[LODIndex];
 			UE::Tasks::FTask RemoveHiddenTask = UE::Tasks::Launch(UE_SOURCE_LOCATION, [&MeshLODs, &CombineOptions, LODIndex, UseTolerance]()
-			{ 
+			{
 				ComputeHiddenRemovalForLOD(MeshLODs[LODIndex].Mesh, LODIndex, CombineOptions);
 				OptimizeLODMeshTriangulation(MeshLODs[LODIndex].Mesh, LODIndex, CombineOptions, UseTolerance);
 			});
@@ -3271,11 +3117,6 @@
 	// wait...
 	UE::Tasks::Wait(PendingRemoveHiddenTasks);
 
-	if (bAppendMinimalLOD)
-	{
-		check(false);		// for future use
-	}
-
 	// parallel regenerate UVs and potentially tangents for any areas of LODs that are missing UVs
 	TArray<UE::Tasks::FTask> PendingAutoUVTasks;
 	bool bComputeTangents = (CombineOptions.bAutoGenerateMissingUVs && CombineOptions.bAutoGenerateTangents);
@@ -3295,7 +3136,7 @@
 		}
 	}
 
-	// remove hidden faces on voxel LODs (todo: can do this via shape sorting, much faster)
+	// remove hidden faces on voxel LODs
 	if (bRemoveHiddenFaces && bUsingCoarseSweepApproximation == false)
 	{
 		TRACE_CPUPROFILER_EVENT_SCOPE(RemoveHidden);
@@ -3316,8 +3157,26 @@
 
 	// make sure AutoUV is done
 	UE::Tasks::Wait(PendingAutoUVTasks);
+}
 
 
+struct FCombinedSubAssembly
+{
+	TArray<UE::Geometry::FDynamicMesh3> MeshLODs;
+	int32 SubAssemblyID = 0;
+};
+
+
+static void BuildOutputSubAssembly(
+	TArray<FCombinedMeshLOD>& MeshLODs,
+	int SubAssemblyID,
+	int32 FirstVoxWrappedIndex,
+	FCombinedSubAssembly& OutputSubAssembly
+)
+{
+	OutputSubAssembly.SubAssemblyID = SubAssemblyID;
+
+	// collect output meshes
 	// can't replace voxel LODs if we are generating UVs for them!
 	// (no way to communicate this upwards...)
 	//int MaxReplaceLOD = MeshLODs.Num();
@@ -3330,12 +3189,285 @@
 		// This can happen particular with VoxWrap LODs
 		if (LODLevel > 0 && LODLevel < MaxReplaceLOD)
 		{
-			if (LODMesh.TriangleCount() > CombinedMeshLODs.Last().TriangleCount())
+			if (LODMesh.TriangleCount() > OutputSubAssembly.MeshLODs.Last().TriangleCount())
 			{
-				LODMesh = CombinedMeshLODs.Last();
+				LODMesh = OutputSubAssembly.MeshLODs.Last();
 			}
 		}
-		CombinedMeshLODs.Add(MoveTemp(LODMesh));
+		OutputSubAssembly.MeshLODs.Add(MoveTemp(LODMesh));
+	}
+}
+
+
+void BuildCombinedMesh(
+	const FMeshPartsAssembly& Assembly,
+	IGeometryProcessing_CombineMeshInstances::FOptions CombineOptions,
+	TArray<FCombinedSubAssembly>& CombinedResults)
+{
+	using namespace UE::Geometry;
+	bool bVerbose = CVarGeometryCombineMeshInstancesVerbose.GetValueOnAnyThread();
+
+	int32 NumLODs = CombineOptions.NumLODs;
+	TArray<FCombinedMeshLOD> MeshLODs;
+	MeshLODs.SetNum(NumLODs);
+
+	int FirstVoxWrappedIndex = 9999;
+	TArray<ECombinedLODType> LODTypes;
+	LODTypes.Init(ECombinedLODType::Approximated, NumLODs);
+	for (int32 LODLevel = 0; LODLevel < NumLODs; ++LODLevel)
+	{
+		if (LODLevel < CombineOptions.NumCopiedLODs)
+		{
+			LODTypes[LODLevel] = ECombinedLODType::Copied;
+		}
+		else if (LODLevel < CombineOptions.NumCopiedLODs + CombineOptions.NumSimplifiedLODs)
+		{
+			LODTypes[LODLevel] = ECombinedLODType::Simplified;
+		}
+		else if (LODLevel >= NumLODs - CombineOptions.NumCoarseLODs)
+		{
+			LODTypes[LODLevel] = ECombinedLODType::VoxWrapped;
+			FirstVoxWrappedIndex = FMath::Min(LODLevel, FirstVoxWrappedIndex);
+		}
+	}
+
+	//CombinedLOD0.Attributes()->SetNumPolygroupLayers(2);
+	//FDynamicMeshPolygroupAttribute* PartIDAttrib = AccumMesh.Attributes()->GetPolygroupLayer(0);
+	//FDynamicMeshPolygroupAttribute* PartInstanceMapAttrib = AccumMesh.Attributes()->GetPolygroupLayer(1);
+
+	int32 NumParts = Assembly.Parts.Num();
+
+	//for ( int32 SetIndex = 0; SetIndex < NumParts; ++SetIndex )
+	//{
+	//	CombinedLOD0.EnableMatchingAttributes( Assembly.SourceMeshGeometry[Index].OriginalMesh, false, false );
+	//}
+
+	// determine if we have multiple part subsets. In this case we need to be able to split the mesh
+	// by part later, which we will do by appending a polygroup layer
+	TArray<int32> SubsetIDs;
+	for (const TUniquePtr<FMeshPart>& Part : Assembly.Parts)
+	{
+		for (const FMeshPartInstance& Instance : Part->Instances)
+		{
+			SubsetIDs.AddUnique(Instance.SubsetID);
+		}
+	}
+	bool bHaveMultiplePartSubsets = (SubsetIDs.Num() > 1);
+	if (bHaveMultiplePartSubsets)
+	{
+		for (FCombinedMeshLOD& LODMeshData : MeshLODs)
+		{
+			LODMeshData.Mesh.Attributes()->SetNumPolygroupLayers(1);
+			LODMeshData.SubsetIDs = LODMeshData.Mesh.Attributes()->GetPolygroupLayer(0);
+		}
+	}
+
+
+	// iterate over part sets, then for each part, over part LODs, and
+	// for each instance append the part LOD to the accumulate LOD mesh
+	for ( int32 SetIndex = 0; SetIndex < NumParts; ++SetIndex )
+	{
+		const TUniquePtr<FMeshPart>& Part = Assembly.Parts[SetIndex];
+		const FSourceGeometry& SourceGeometry = Assembly.SourceMeshGeometry[SetIndex];
+		const FOptimizedGeometry& OptimizedGeometry = Assembly.OptimizedMeshGeometry[SetIndex];
+
+		check(Part->bAllowMerging == true);		// do not support this yet
+
+		FMeshIndexMappings Mappings;
+
+		for (int32 LODLevel = 0; LODLevel < NumLODs; ++LODLevel)
+		{
+			const FDynamicMesh3* SourceAppendMesh = nullptr;
+			const FDynamicMesh3* ApproximateAppendMesh = nullptr;
+			const FDynamicMesh3* UseAppendMesh = nullptr;
+
+			// default approximate mesh to lowest-quality approximation (box), need to do this
+			// so that we always have something to swap to for Decorative parts
+			ApproximateAppendMesh = &OptimizedGeometry.ApproximateMeshLODs.Last();
+
+			ECombinedLODType LevelLODType = LODTypes[LODLevel];
+			if (LevelLODType == ECombinedLODType::Copied)
+			{
+				SourceAppendMesh = (LODLevel < SourceGeometry.SourceMeshLODs.Num()) ? 
+					&SourceGeometry.SourceMeshLODs[LODLevel] : &SourceGeometry.SourceMeshLODs.Last();
+				UseAppendMesh = SourceAppendMesh;
+			}
+			else if (LevelLODType == ECombinedLODType::Simplified)
+			{
+				int32 SimplifiedLODIndex = LODLevel - CombineOptions.NumCopiedLODs;
+				SourceAppendMesh = &OptimizedGeometry.SimplifiedMeshLODs[SimplifiedLODIndex];
+				UseAppendMesh = SourceAppendMesh;
+			}
+			else if (LevelLODType == ECombinedLODType::VoxWrapped)
+			{
+				SourceAppendMesh = &SourceGeometry.SourceMeshLODs.Last();
+				UseAppendMesh = SourceAppendMesh;
+			}
+			else // ECombinedLODType::Approximated
+			{
+				int32 ApproxLODIndex = LODLevel - CombineOptions.NumCopiedLODs - CombineOptions.NumSimplifiedLODs;
+				ApproximateAppendMesh = &OptimizedGeometry.ApproximateMeshLODs[ApproxLODIndex];
+				UseAppendMesh = ApproximateAppendMesh;
+			}
+
+			FCombinedMeshLOD& CombinedMeshLODData = MeshLODs[LODLevel];
+
+			for ( const FMeshPartInstance& Instance : Part->Instances )
+			{
+				bool bIsDecorativePart = (Instance.DetailLevel == EMeshDetailLevel::Decorative);
+
+				if (bIsDecorativePart)
+				{
+					// filter out detail parts at higher LODs, or if we are doing VoxWrap LOD
+					if ( LODLevel >= CombineOptions.FilterDecorativePartsLODLevel || LevelLODType == ECombinedLODType::VoxWrapped )
+					{
+						continue;
+					}
+					// at last detail part LOD, switch to approximate mesh
+					if (LODLevel >= (CombineOptions.FilterDecorativePartsLODLevel - CombineOptions.ApproximateDecorativePartLODs) )
+					{
+						UseAppendMesh = ApproximateAppendMesh;
+					}
+				}
+
+				// need to make a copy to run pre-process func
+				FDynamicMesh3 TempAppendMesh(*UseAppendMesh);
+				if (Assembly.PreProcessInstanceMeshFunc)
+				{
+					Assembly.PreProcessInstanceMeshFunc(TempAppendMesh, Instance);
+				}
+
+				// if part does not require UVs, get rid of them here to encourage merging downstream
+				if (LODLevel > 0 && Part->bPreserveUVs == false && TempAppendMesh.HasAttributes())
+				{
+					for (int32 UVLayer = 0; UVLayer < TempAppendMesh.Attributes()->NumUVLayers(); ++UVLayer)
+					{
+						TempAppendMesh.Attributes()->GetUVLayer(UVLayer)->ClearElements();
+					}
+					// if we have no UVs then tangents are invalid
+					TempAppendMesh.Attributes()->DisableTangents();
+				}
+
+				Mappings.Reset();
+				CombinedMeshLODData.Editor.AppendMesh(&TempAppendMesh, Mappings,
+					[&](int, const FVector3d& Pos) { return Instance.WorldTransform.TransformPosition(Pos); },
+					[&](int, const FVector3d& Normal) { return Instance.WorldTransform.TransformNormal(Normal); });
+
+				// transfer part IDs etc
+				if (CombinedMeshLODData.SubsetIDs != nullptr)
+				{
+					for (int32 tid : TempAppendMesh.TriangleIndicesItr())
+					{
+						CombinedMeshLODData.SubsetIDs->SetValue(Mappings.GetNewTriangle(tid), Instance.SubsetID);
+					}
+				}
+
+				// could precompute these indexes for each instance?
+				// also for source mesh we could transfer material IDs correctly...
+				UMaterialInterface* UseMaterial = Instance.Materials[0];
+				const int32* FoundMaterialIndex = Assembly.MaterialMap.Find(UseMaterial);
+				int32 AssignMaterialIndex = (FoundMaterialIndex != nullptr) ? *FoundMaterialIndex : 0;
+
+				for (int32 tid : TempAppendMesh.TriangleIndicesItr())
+				{
+					CombinedMeshLODData.MaterialIDs->SetValue( Mappings.GetNewTriangle(tid), AssignMaterialIndex );
+				}
+			}
+		}
+	}
+
+	// Above we copied the last approximation LOD into the first VoxWrap LOD slot. However 
+	// approximate LOD may be very coarse, it's better to use a known-good starting mesh?
+	// Need to do this before starting hidden-removal task...
+	if (FirstVoxWrappedIndex < MeshLODs.Num())
+	{
+		MeshLODs[FirstVoxWrappedIndex].Mesh = MeshLODs[CombineOptions.NumCopiedLODs-1].Mesh;
+	}
+
+	// make a list of per-LOD geometric tolerances tha will drive additional optimization. 
+	// For copied and first simplified LODs, use Simplify Base Tolerance, and then increment
+	// for each successive LOD.  (todo: have a separate initial tolerance for Approx LODs?)
+	double CurTolerance = CombineOptions.SimplifyBaseTolerance;
+	TArray<double> OptimizationTolerances;
+	for (int32 LODIndex = 0; LODIndex < NumLODs; ++LODIndex)
+	{
+		if (LODTypes[LODIndex] == ECombinedLODType::Simplified && LODTypes[LODIndex - 1] == ECombinedLODType::Simplified)
+		{
+			CurTolerance *= 2.0;
+		}
+		else if (LODTypes[LODIndex] == ECombinedLODType::Approximated)
+		{
+			CurTolerance *= 2.0;
+		}
+		OptimizationTolerances.Add(CurTolerance);
+	}
+
+
+	if (bHaveMultiplePartSubsets)
+	{
+		TArray<int32> OrderedSubsetIDs(SubsetIDs);
+
+		int32 NumSubsets = SubsetIDs.Num();
+		TArray<TArray<FCombinedMeshLOD>> SubsetMeshLODChains;
+		SubsetMeshLODChains.SetNum(NumSubsets);
+		for (int32 SubsetIndex = 0; SubsetIndex < NumSubsets; ++SubsetIndex)
+		{
+			SubsetMeshLODChains[SubsetIndex].SetNum(NumLODs);
+		}
+
+		for (int32 LODIndex = 0; LODIndex < NumLODs; ++LODIndex)
+		{
+			// split the LOD by subset ID
+			FDynamicMesh3& LODMesh = MeshLODs[LODIndex].Mesh;
+			FDynamicMeshPolygroupAttribute* SubsetIDAttrib = MeshLODs[LODIndex].SubsetIDs;
+			TArray<FDynamicMesh3> SplitMeshes;
+			FDynamicMeshEditor::SplitMesh(&LODMesh, SplitMeshes, [&](int32 tid) { return OrderedSubsetIDs.IndexOfByKey(SubsetIDAttrib->GetValue(tid)); });
+
+			// code below assumes this. If it's not the case, then we have some more complex processing to figure out...
+			check(SplitMeshes.Num() == NumSubsets);		
+
+			// give each subset submesh to the 
+			for (int32 SubsetIndex = 0; SubsetIndex < NumSubsets; ++SubsetIndex)
+			{
+				TArray<FCombinedMeshLOD>& LODChain = SubsetMeshLODChains[SubsetIndex];
+				FDynamicMesh3& Submesh = SplitMeshes[SubsetIndex];
+				LODChain[LODIndex].SetMesh( MoveTemp(Submesh) );
+			}
+		}
+
+		TArray<UE::Tasks::FTask> PendingSubsetTasks;
+
+		for (int32 SubsetIndex = 0; SubsetIndex < NumSubsets; ++SubsetIndex)
+		{
+			TArray<FCombinedMeshLOD>& LODChain = SubsetMeshLODChains[SubsetIndex];
+			UE::Tasks::FTask ProcessSubsetTask = UE::Tasks::Launch(UE_SOURCE_LOCATION, [&LODChain, &CombineOptions, &OptimizationTolerances, FirstVoxWrappedIndex]()
+			{
+				ProcessCombinedLODChain(LODChain, OptimizationTolerances, FirstVoxWrappedIndex, CombineOptions);
+			});
+			PendingSubsetTasks.Add(ProcessSubsetTask);
+			if (bVerbose)
+			{
+				ProcessSubsetTask.Wait();
+			}
+		}
+
+		// wait for all subsets to finish processing
+		UE::Tasks::Wait(PendingSubsetTasks);
+
+		CombinedResults.SetNum(NumSubsets);
+		for (int32 SubsetIndex = 0; SubsetIndex < NumSubsets; ++SubsetIndex)
+		{
+			TArray<FCombinedMeshLOD>& LODChain = SubsetMeshLODChains[SubsetIndex];
+			BuildOutputSubAssembly(LODChain, OrderedSubsetIDs[SubsetIndex], FirstVoxWrappedIndex, CombinedResults[SubsetIndex]);
+		}
+
+	}
+	else
+	{
+		ProcessCombinedLODChain(MeshLODs, OptimizationTolerances, FirstVoxWrappedIndex, CombineOptions);
+
+		CombinedResults.SetNum(1);
+		BuildOutputSubAssembly(MeshLODs, 0, FirstVoxWrappedIndex, CombinedResults[0]);
 	}
 
 }
@@ -3436,10 +3568,13 @@
 
 void BuildCombinedCollisionShapes(
 	const FMeshPartsAssembly& Assembly,
+	TArray<int32> SubsetIDsOrdering,
 	IGeometryProcessing_CombineMeshInstances::FOptions CombineOptions,
-	FSimpleShapeSet3d& CombinedCollisionShapes)
+	TArray<FSimpleShapeSet3d>& CombinedCollisionShapes)
 {
 	int32 NumParts = Assembly.Parts.Num();
+	int32 NumSubsets = SubsetIDsOrdering.Num();
+	CombinedCollisionShapes.SetNum(NumSubsets);
 
 	for ( int32 SetIndex = 0; SetIndex < NumParts; ++SetIndex )
 	{
@@ -3447,22 +3582,73 @@
 		const FSourceGeometry& SourceGeometry = Assembly.SourceMeshGeometry[SetIndex];
 		for ( const FMeshPartInstance& Instance : Part->Instances )
 		{
+			int32 SubsetIndex = SubsetIDsOrdering.IndexOfByKey(Instance.SubsetID);
+
 			bool bIsDecorativePart = (Instance.DetailLevel == EMeshDetailLevel::Decorative);
 			if ( ! bIsDecorativePart )
 			{
-				CombinedCollisionShapes.Append( SourceGeometry.CollisionShapes, Instance.WorldTransform );
+				CombinedCollisionShapes[SubsetIndex].Append(SourceGeometry.CollisionShapes, Instance.WorldTransform);
 			}
 		}
 	}
 
 	// trivially merge any adjacent boxes that merge to a perfect combined-box
-	CombineCollisionShapes(CombinedCollisionShapes, 0.01);
+	for (FSimpleShapeSet3d& ShapeSet : CombinedCollisionShapes)
+	{
+		CombineCollisionShapes(ShapeSet, 0.01);
+	}
 }
 
 
 
 
 
+void FilterUnusedMaterials(
+	TArray<FDynamicMesh3>& Meshes,
+	TArray<UMaterialInterface*>& Materials)
+{
+	FDynamicMeshMaterialAttribute* MaterialIDs = Meshes[0].Attributes()->GetMaterialID();
+	if (!MaterialIDs) return;
+
+	TArray<bool> UsedMaterials;
+	UsedMaterials.SetNum(Materials.Num());
+	for (int32 tid : Meshes[0].TriangleIndicesItr())
+	{
+		int32 MaterialID = MaterialIDs->GetValue(tid);
+		UsedMaterials[MaterialID] = true;
+	}
+	if (UsedMaterials.Num() == Materials.Num())
+	{
+		return;
+	}
+	
+	TArray<UMaterialInterface*> NewMaterials;
+	TArray<int32> MaterialMap;
+	MaterialMap.SetNum(Materials.Num());
+	for (int32 k = 0; k < Materials.Num(); ++k)
+	{
+		if (UsedMaterials[k])
+		{
+			MaterialMap[k] = NewMaterials.Num();
+			NewMaterials.Add(Materials[k]);
+		}
+	}
+
+	for (FDynamicMesh3& LODMesh : Meshes)
+	{
+		FDynamicMeshMaterialAttribute* LODMaterialIDs = LODMesh.Attributes()->GetMaterialID();
+		for (int32 tid : LODMesh.TriangleIndicesItr())
+		{
+			int32 MaterialID = LODMaterialIDs->GetValue(tid);
+			int32 NewMaterialID = MaterialMap[MaterialID];
+			LODMaterialIDs->SetValue(tid, NewMaterialID);
+		}
+	}
+
+	Materials = MoveTemp(NewMaterials);
+}
 
 
 
+
+
@@ -3526,7 +3712,7 @@
 {
 	TRACE_CPUPROFILER_EVENT_SCOPE(CombineMeshInstances);
 
-	bool bVerbose = CVarGeometryCombineMeshInstancesVerbose.GetValueOnGameThread();
+	bool bVerbose = CVarGeometryCombineMeshInstancesVerbose.GetValueOnAnyThread();
 	if (bVerbose)
 	{
 		int32 NumApproxLODs = FMath::Max(0, Options.NumLODs - Options.NumCopiedLODs - Options.NumSimplifiedLODs - Options.NumCoarseLODs);
@@ -3587,31 +3773,50 @@
 		}
 	};
 
-	TArray<UE::Geometry::FDynamicMesh3> CombinedMeshLODs;
+	// build combined mesh LOD chains for each sub-assembly
+	TArray<FCombinedSubAssembly> CombinedResults;
+	TArray<int32> SubsetIDsOrdering;
 	{
 		TRACE_CPUPROFILER_EVENT_SCOPE(CombineMeshInst_BuildMeshes);
-		BuildCombinedMesh(PartAssembly, Options, CombinedMeshLODs);
+		BuildCombinedMesh(PartAssembly, Options, CombinedResults);
+		for (FCombinedSubAssembly& Assembly : CombinedResults)
+		{
+			SubsetIDsOrdering.Add(Assembly.SubAssemblyID);
+		}
 	}
 
-	FSimpleShapeSet3d CombinedCollisionShapes;
+	// build combined collision shapes, repsecting sub-assembly ordering
+	TArray<FSimpleShapeSet3d> CombinedCollisionShapes;
 	{
 		TRACE_CPUPROFILER_EVENT_SCOPE(CombineMeshInst_BuildCollision);
-		BuildCombinedCollisionShapes(PartAssembly, Options, CombinedCollisionShapes);
+		BuildCombinedCollisionShapes(PartAssembly, SubsetIDsOrdering, Options, CombinedCollisionShapes);
 
-		if (bVerbose)
+		if (bVerbose && CombinedCollisionShapes.Num() == 1)
 		{
-			UE_LOG(LogGeometry, Log, TEXT("  CombinedCollisionShapes contains %d Boxes, %d Convexes"), 
-				CombinedCollisionShapes.Boxes.Num(), CombinedCollisionShapes.Convexes.Num());
+			UE_LOG(LogGeometry, Log, TEXT("  CombinedCollisionShapes[0] contains %d Boxes, %d Convexes"), 
+				CombinedCollisionShapes[0].Boxes.Num(), CombinedCollisionShapes[0].Convexes.Num());
 		}
 	}
-	FPhysicsDataCollection PhysicsData;
-	PhysicsData.Geometry = CombinedCollisionShapes;
-	PhysicsData.CopyGeometryToAggregate();		// need FPhysicsDataCollection to convert to agg geom, should fix this
 
 
-	ResultsOut.CombinedMeshes.SetNum(1);
-	IGeometryProcessing_CombineMeshInstances::FOutputMesh& OutputMesh = ResultsOut.CombinedMeshes[0];
-	OutputMesh.MeshLODs = MoveTemp(CombinedMeshLODs);
-	OutputMesh.MaterialSet = PartAssembly.UniqueMaterials;
-	OutputMesh.SimpleCollisionShapes = PhysicsData.AggGeom;
+	// build final results data structure
+	ResultsOut.CombinedMeshes.SetNum(CombinedResults.Num());
+	for (int32 k = 0; k < CombinedResults.Num(); ++k)
+	{
+		ResultsOut.CombinedMeshes[k].MeshLODs = MoveTemp(CombinedResults[k].MeshLODs);
+
+		FPhysicsDataCollection PhysicsData;
+		PhysicsData.Geometry = CombinedCollisionShapes[k];
+		PhysicsData.CopyGeometryToAggregate();		// need FPhysicsDataCollection to convert to agg geom, should fix this
+		ResultsOut.CombinedMeshes[k].SimpleCollisionShapes = PhysicsData.AggGeom;
+
+		ResultsOut.CombinedMeshes[k].MaterialSet = PartAssembly.UniqueMaterials;
+		if (CombinedResults.Num() > 1)	
+		{
+			// if we have multiple outputs, they may not each use the full material set, in this case we will filter the materials (should this be optional?)
+			FilterUnusedMaterials(ResultsOut.CombinedMeshes[k].MeshLODs, ResultsOut.CombinedMeshes[k].MaterialSet);
+		}
+
+		ResultsOut.CombinedMeshes[k].InstanceSubsetID = CombinedResults[k].SubAssemblyID;
+	}
 }
diff --git a/Engine/Source/Developer/GeometryProcessingInterfaces/Public/GeometryProcessingInterfaces/CombineMeshInstances.h b/Engine/Source/Developer/GeometryProcessingInterfaces/Public/GeometryProcessingInterfaces/CombineMeshInstances.h
--- a/Engine/Source/Developer/GeometryProcessingInterfaces/Public/GeometryProcessingInterfaces/CombineMeshInstances.h
+++ b/Engine/Source/Developer/GeometryProcessingInterfaces/Public/GeometryProcessingInterfaces/CombineMeshInstances.h
@@ -56,7 +56,13 @@ public:
 	{
 		EMeshDetailLevel DetailLevel = EMeshDetailLevel::Standard;
 		TArray<FTransform3d> TransformSequence;		// set of transforms on this instance. Often just a single transform.
-		int32 GroupDataIndex = -1;		// index into FSourceInstanceList::InstanceGroupDatas
+		int32 GroupDataIndex = -1;					// index into FSourceInstanceList::InstanceGroupDatas
+
+		// in some cases it may be desirable to have "groups" of instances which should be output as separate meshes, but
+		// be jointly processed in terms of (eg) the part LODs. If any InstanceSubsetID is non-zero, then instance subsets
+		// are grouped/extracted by integer ID and will be returned as separate FOutputMesh's in the FResults. 
+		// No hidden-removal, triangle merging optimization, etc will be performed between Instance Subsets.
+		int32 InstanceSubsetID = 0;
 	};
 
 
@@ -273,6 +279,9 @@ public:
 		TArray<UMaterialInterface*> MaterialSet;
 
 		FKAggregateGeom SimpleCollisionShapes;
+
+		// All part instances accumulated into the MeshLODs will have had this InstanceSubsetID in their input FBaseMeshInstance's
+		int32 InstanceSubsetID = 0;
 	};
 
 
