Fixing deprecating warnings due to unused config category in skin paint properties.

diff --git a/Engine/Plugins/Experimental/MeshModelingToolsetExp/Source/MeshModelingToolsEditorOnlyExp/Public/SkeletalMesh/SkinWeightsPaintTool.h b/Engine/Plugins/Experimental/MeshModelingToolsetExp/Source/MeshModelingToolsEditorOnlyExp/Public/SkeletalMesh/SkinWeightsPaintTool.h
--- a/Engine/Plugins/Experimental/MeshModelingToolsetExp/Source/MeshModelingToolsEditorOnlyExp/Public/SkeletalMesh/SkinWeightsPaintTool.h
+++ b/Engine/Plugins/Experimental/MeshModelingToolsetExp/Source/MeshModelingToolsEditorOnlyExp/Public/SkeletalMesh/SkinWeightsPaintTool.h
@@ -288,25 +288,25 @@ public:
 	bool bColorModeChanged = false;
 
 	// weight editing arguments
-	UPROPERTY(Config = EditorSettings)
+	UPROPERTY(Config)
 	TEnumAsByte<EAxis::Type> MirrorAxis = EAxis::X;
-	UPROPERTY(Config = EditorSettings)
+	UPROPERTY(Config)
 	EMirrorDirection MirrorDirection = EMirrorDirection::PositiveToNegative;
-	UPROPERTY(Config = EditorSettings)
+	UPROPERTY(Config)
 	float FloodValue = 1.f;
-	UPROPERTY(Config = EditorSettings)
+	UPROPERTY(Config)
 	float PruneValue = 0.01;
 
 	// save/restore user specified settings for each tool mode
 	FSkinWeightBrushConfig& GetBrushConfig();
 	TMap<EWeightEditOperation, FSkinWeightBrushConfig*> BrushConfigs;
-	UPROPERTY(Config = EditorSettings)
+	UPROPERTY(Config)
 	FSkinWeightBrushConfig BrushConfigAdd;
-	UPROPERTY(Config = EditorSettings)
+	UPROPERTY(Config)
 	FSkinWeightBrushConfig BrushConfigReplace;
-	UPROPERTY(Config = EditorSettings)
+	UPROPERTY(Config)
 	FSkinWeightBrushConfig BrushConfigMultiply;
-	UPROPERTY(Config = EditorSettings)
+	UPROPERTY(Config)
 	FSkinWeightBrushConfig BrushConfigRelax;
 
 	// pointer back to paint tool
