Added the ability to set the base LOD index when copying from the original mesh.

diff --git a/Engine/Plugins/Experimental/MeshModelingToolsetExp/Source/GeometryProcessingAdapters/Private/CombineMeshInstancesImpl.cpp b/Engine/Plugins/Experimental/MeshModelingToolsetExp/Source/GeometryProcessingAdapters/Private/CombineMeshInstancesImpl.cpp
--- a/Engine/Plugins/Experimental/MeshModelingToolsetExp/Source/GeometryProcessingAdapters/Private/CombineMeshInstancesImpl.cpp
+++ b/Engine/Plugins/Experimental/MeshModelingToolsetExp/Source/GeometryProcessingAdapters/Private/CombineMeshInstancesImpl.cpp
@@ -3732,7 +3732,7 @@ void FCombineMeshInstancesImpl::CombineMeshInstances(
 				PartAssembly.Parts.Num(), PartAssembly.UniqueMaterials.Num());
 		}
 
-		InitializeAssemblySourceMeshesFromLOD(PartAssembly, 0, Options.NumCopiedLODs);
+		InitializeAssemblySourceMeshesFromLOD(PartAssembly, Options.BaseCopiedLOD, Options.NumCopiedLODs);
 		InitializePartAssemblySpatials(PartAssembly);
 	}
 	{
@@ -3819,4 +3819,4 @@ void FCombineMeshInstancesImpl::CombineMeshInstances(
 
 		ResultsOut.CombinedMeshes[k].InstanceSubsetID = CombinedResults[k].SubAssemblyID;
 	}
-}
\ No newline at end of file
+}
diff --git a/Engine/Source/Developer/GeometryProcessingInterfaces/Public/GeometryProcessingInterfaces/CombineMeshInstances.h b/Engine/Source/Developer/GeometryProcessingInterfaces/Public/GeometryProcessingInterfaces/CombineMeshInstances.h
--- a/Engine/Source/Developer/GeometryProcessingInterfaces/Public/GeometryProcessingInterfaces/CombineMeshInstances.h
+++ b/Engine/Source/Developer/GeometryProcessingInterfaces/Public/GeometryProcessingInterfaces/CombineMeshInstances.h
@@ -143,6 +143,7 @@ public:
 		int32 NumLODs = 5;
 
 		// settings for Copied LODs, these are directly copied from Source Geometry
+		int32 BaseCopiedLOD = 0;
 		int32 NumCopiedLODs = 1;
 
 		//
