CombineMeshInstances: add support for collision shapes in FMeshLODSet path

diff --git a/Engine/Plugins/Experimental/MeshModelingToolsetExp/Source/GeometryProcessingAdapters/Private/CombineMeshInstancesImpl.cpp b/Engine/Plugins/Experimental/MeshModelingToolsetExp/Source/GeometryProcessingAdapters/Private/CombineMeshInstancesImpl.cpp
--- a/Engine/Plugins/Experimental/MeshModelingToolsetExp/Source/GeometryProcessingAdapters/Private/CombineMeshInstancesImpl.cpp
+++ b/Engine/Plugins/Experimental/MeshModelingToolsetExp/Source/GeometryProcessingAdapters/Private/CombineMeshInstancesImpl.cpp
@@ -427,9 +427,15 @@ void InitializeAssemblySourceMeshesFromLOD(
 			if ( UBodySetup* BodySetup = StaticMesh->GetBodySetup() )
 			{
 				UE::Geometry::GetShapeSet(BodySetup->AggGeom, Target.CollisionShapes);
-				// todo: detect boxes?
 			}
 		}
+		else if (Part->SourceMeshLODSet != nullptr)
+		{
+			UE::Geometry::GetShapeSet(Part->SourceMeshLODSet->SimpleCollisionShapes, Target.CollisionShapes);
+		}
+
+		// sometimes simple collision is a convex when it's actually a box - could try to detect here?
+
 	}
 
 }
diff --git a/Engine/Source/Developer/GeometryProcessingInterfaces/Public/GeometryProcessingInterfaces/CombineMeshInstances.h b/Engine/Source/Developer/GeometryProcessingInterfaces/Public/GeometryProcessingInterfaces/CombineMeshInstances.h
--- a/Engine/Source/Developer/GeometryProcessingInterfaces/Public/GeometryProcessingInterfaces/CombineMeshInstances.h
+++ b/Engine/Source/Developer/GeometryProcessingInterfaces/Public/GeometryProcessingInterfaces/CombineMeshInstances.h
@@ -82,6 +82,7 @@ public:
 	struct FMeshLODSet
 	{
 		TArray<const FMeshDescription*> ReferencedMeshLODs;
+		FKAggregateGeom SimpleCollisionShapes;
 	};
 
 	/**
