CombineMeshInstances: fix issue w/ planar retriangulation and bowties. Add sanity check on area change during planar retriangulation.

diff --git a/Engine/Plugins/Experimental/MeshModelingToolsetExp/Source/GeometryProcessingAdapters/Private/CombineMeshInstancesImpl.cpp b/Engine/Plugins/Experimental/MeshModelingToolsetExp/Source/GeometryProcessingAdapters/Private/CombineMeshInstancesImpl.cpp
--- a/Engine/Plugins/Experimental/MeshModelingToolsetExp/Source/GeometryProcessingAdapters/Private/CombineMeshInstancesImpl.cpp
+++ b/Engine/Plugins/Experimental/MeshModelingToolsetExp/Source/GeometryProcessingAdapters/Private/CombineMeshInstancesImpl.cpp
@@ -2007,6 +2007,8 @@ static void RetriangulatePlanarFacePolygons(FDynamicMesh3& TargetMesh, double Ba
 		TargetMesh.CompactInPlace();
 	}
 
+	double InitialArea = TMeshQueries<FDynamicMesh3>::GetVolumeArea(TargetMesh).Y;
+
 	TArray<int32> TriPlaneID;
 	TriPlaneID.Init(-1, TargetMesh.MaxTriangleID());
 
@@ -2074,10 +2076,25 @@ static void RetriangulatePlanarFacePolygons(FDynamicMesh3& TargetMesh, double Ba
 		double MergeOffset = FMath::Min(EstSquareEdgeLen * 0.02, BaseGeometricTolerance);
 		double SimplifyTolerance = FMath::Min(MergeOffset, BaseGeometricTolerance);
 
+		// ideally this is better than just using one arbitrary triangle normal
+		FVector3d AverageNormal = FVector3d::Zero();
+		for (int32 tid : Mesh.TriangleIndicesItr())
+		{
+			AverageNormal += Mesh.GetTriNormal(tid);
+		}
+		AverageNormal.Normalize();
+
+		// subset of triangles can easily have bowties, and since we are going to be using boundary loops
+		// this can lead to messy situations, so just split any bowties
+		FDynamicMeshEditor BowtieSplitter(&Mesh);
+		FDynamicMeshEditResult TmpEditResult;
+		BowtieSplitter.SplitBowties(TmpEditResult);
+
 		FDynamicMesh3 NewPlanarMesh;
 		ComputePlanarPolygonApproximation(Mesh, NewPlanarMesh,
-			Mesh.GetTriNormal(0), MergeOffset, SimplifyTolerance, MinHoleArea);
+			AverageNormal, MergeOffset, SimplifyTolerance, MinHoleArea);
 
+		// only take this new mesh if we actually improved the situation
 		if (NewPlanarMesh.TriangleCount() < Mesh.TriangleCount())
 		{
 			Mesh = MoveTemp(NewPlanarMesh);
@@ -2129,6 +2146,15 @@ static void RetriangulatePlanarFacePolygons(FDynamicMesh3& TargetMesh, double Ba
 		return;
 	}
 
+	// Sanity check that we have not dramatically changed the mesh area. Some
+	// area change is expected because of (eg) filling holes, merging, etc,
+	// so the tolerance here is quite large and mainly intended to catch catastrophic failures
+	double FinalArea = TMeshQueries<FDynamicMesh3>::GetVolumeArea(NewMesh).Y;
+	if (FinalArea < 0.5 * InitialArea)
+	{
+		return;
+	}
+
 	// currently assuming input mesh has been split by MaterialID, so no projection here,
 	// just use any MaterialID
 	if (bTargetHasAttributes)
