Skin paint tools "flood" operation UI modified.

diff --git a/Engine/Plugins/Experimental/MeshModelingToolsetExp/Source/ModelingEditorUI/Private/SkinWeightDetailCustomization.cpp b/Engine/Plugins/Experimental/MeshModelingToolsetExp/Source/ModelingEditorUI/Private/SkinWeightDetailCustomization.cpp
--- a/Engine/Plugins/Experimental/MeshModelingToolsetExp/Source/ModelingEditorUI/Private/SkinWeightDetailCustomization.cpp
+++ b/Engine/Plugins/Experimental/MeshModelingToolsetExp/Source/ModelingEditorUI/Private/SkinWeightDetailCustomization.cpp
@@ -488,57 +488,72 @@ void FSkinWeightDetailCustomization::AddSelectionUI(IDetailLayoutBuilder& Detail
 		]
 
 		+SVerticalBox::Slot()
-		.Padding(WeightEditHorizontalPadding, WeightEditVerticalPadding)
+		.Padding(0.f, WeightEditVerticalPadding)
 		[
 			SNew(SHorizontalBox)
-
-			+ SHorizontalBox::Slot()
-			[
-				SNew(SButton)
-				.HAlign(HAlign_Center)
-				.Text(LOCTEXT("AddWeightsButtonLabel", "Add"))
-				.ToolTipText(LOCTEXT("AddOpTooltip", "Add: applies the current weight plus the flood value to the new weight.\n"))
-				.OnClicked_Lambda([this]()
-				{
-					SkinToolSettings->WeightTool->FloodWeights(SkinToolSettings->FloodValue, EWeightEditOperation::Add);
-					return FReply::Handled();
-				})
-			]
-			+ SHorizontalBox::Slot()
-			[
-				SNew(SButton)
-				.HAlign(HAlign_Center)
-				.Text(LOCTEXT("ReplaceWeightsButtonLabel", "Replace"))
-				.ToolTipText(LOCTEXT("ReplaceOpTooltip", "Replace: applies the current weight minus the strength value to the new weight.\n"))
-				.OnClicked_Lambda([this]()
-				{
-					SkinToolSettings->WeightTool->FloodWeights(SkinToolSettings->FloodValue, EWeightEditOperation::Replace);
-					return FReply::Handled();
-				})
-			]
-			+ SHorizontalBox::Slot()
+			
+			+SHorizontalBox::Slot()
+			.VAlign(VAlign_Center)
+			.FillWidth(WeightEditingLabelsPercent)
 			[
-				SNew(SButton)
-				.HAlign(HAlign_Center)
-				.Text(LOCTEXT("MultiplyeightsButtonLabel", "Multiply"))
-				.ToolTipText(LOCTEXT("MultiplyOpTooltip", "Multiply: applies the current weight multiplied by the strength value to the new weight.\n"))
-				.OnClicked_Lambda([this]()
-				{
-					SkinToolSettings->WeightTool->FloodWeights(SkinToolSettings->FloodValue, EWeightEditOperation::Multiply);
-					return FReply::Handled();
-				})
+				SNew(STextBlock)
+				.Text(LOCTEXT("FloodOperationLabel", "Flood Operation"))
+				.Font(FAppStyle::GetFontStyle(TEXT("PropertyWindow.NormalFont")))
+				.ToolTipText(LOCTEXT("FloodOperationTooltip", "The various flood weight operations."))
 			]
-			+ SHorizontalBox::Slot()
+
+			+SHorizontalBox::Slot()
 			[
-				SNew(SButton)
-				.HAlign(HAlign_Center)
-				.Text(LOCTEXT("RelaxWeightsButtonLabel", "Relax"))
-				.ToolTipText(LOCTEXT("RelaxeOpTooltip", "Relax: applies the average of the connected (by edge) vertex weights to the new vertex weight, blended by the strength.\n"))
-				.OnClicked_Lambda([this]()
-				{
-					SkinToolSettings->WeightTool->FloodWeights(SkinToolSettings->FloodValue, EWeightEditOperation::Relax);
-					return FReply::Handled();
-				})
+				SNew(SHorizontalBox)
+				+ SHorizontalBox::Slot()
+				[
+					SNew(SButton)
+					.HAlign(HAlign_Center)
+					.ButtonStyle(&FAppStyle::Get().GetWidgetStyle<FButtonStyle>("VerticalToolBar.Button"))
+					.Text(LOCTEXT("AddWeightsButtonLabel", "Add"))
+					.ToolTipText(LOCTEXT("AddOpTooltip", "Add: applies the current weight plus the flood amount to the new weight."))
+					.OnClicked_Lambda([this]()
+					{
+						SkinToolSettings->WeightTool->FloodWeights(SkinToolSettings->FloodValue, EWeightEditOperation::Add);
+						return FReply::Handled();
+					})
+				]
+				+ SHorizontalBox::Slot()
+				[
+					SNew(SButton)
+					.HAlign(HAlign_Center)
+					.Text(LOCTEXT("ReplaceWeightsButtonLabel", "Replace"))
+					.ToolTipText(LOCTEXT("ReplaceOpTooltip", "Replace: applies the current weight minus the flood amount to the new weight."))
+					.OnClicked_Lambda([this]()
+					{
+						SkinToolSettings->WeightTool->FloodWeights(SkinToolSettings->FloodValue, EWeightEditOperation::Replace);
+						return FReply::Handled();
+					})
+				]
+				+ SHorizontalBox::Slot()
+				[
+					SNew(SButton)
+					.HAlign(HAlign_Center)
+					.Text(LOCTEXT("MultiplyeightsButtonLabel", "Multiply"))
+					.ToolTipText(LOCTEXT("MultiplyOpTooltip", "Multiply: applies the current weight multiplied by the flood amount to the new weight."))
+					.OnClicked_Lambda([this]()
+					{
+						SkinToolSettings->WeightTool->FloodWeights(SkinToolSettings->FloodValue, EWeightEditOperation::Multiply);
+						return FReply::Handled();
+					})
+				]
+				+ SHorizontalBox::Slot()
+				[
+					SNew(SButton)
+					.HAlign(HAlign_Center)
+					.Text(LOCTEXT("RelaxWeightsButtonLabel", "Relax"))
+					.ToolTipText(LOCTEXT("RelaxeOpTooltip", "Relax: applies the average of the connected (by edge) vertex weights to the new vertex weight, scaled by the flood amount."))
+					.OnClicked_Lambda([this]()
+					{
+						SkinToolSettings->WeightTool->FloodWeights(SkinToolSettings->FloodValue, EWeightEditOperation::Relax);
+						return FReply::Handled();
+					})
+				]
 			]
 		]
 	];
