Skeleton Editing Tool: notify hierachy changed when mirroring bones

diff --git a/Engine/Plugins/Experimental/MeshModelingToolsetExp/Source/MeshModelingToolsEditorOnlyExp/Private/SkeletalMesh/SkeletonEditingTool.cpp b/Engine/Plugins/Experimental/MeshModelingToolsetExp/Source/MeshModelingToolsEditorOnlyExp/Private/SkeletalMesh/SkeletonEditingTool.cpp
--- a/Engine/Plugins/Experimental/MeshModelingToolsetExp/Source/MeshModelingToolsEditorOnlyExp/Private/SkeletalMesh/SkeletonEditingTool.cpp
+++ b/Engine/Plugins/Experimental/MeshModelingToolsetExp/Source/MeshModelingToolsEditorOnlyExp/Private/SkeletalMesh/SkeletonEditingTool.cpp
@@ -456,6 +456,11 @@ void USkeletonEditingTool::MirrorBones()
 	const bool bBonesMirrored = Modifier->MirrorBones(GetSelectedBones(), MirroringProperties->Options);
 	if (bBonesMirrored)
 	{
+		if (NeedsNotification())
+		{
+			GetNotifier().Notify({}, ESkeletalMeshNotifyType::HierarchyChanged);
+		}
+	
 		EndChange();
 		return;		
 	}
