CombineMeshInstances: properly handle multiple materials on source/simplified parts. Add FBaseMeshInstance::bAllowApproximation, if set to false, only copied or simplified LODs should be used for that part instance, to support special downstream part processing after merging. Currently only partially supported, budget may replace simplified LOD with approximation before appending instances, this needs to be revisited.

diff --git a/Engine/Plugins/Experimental/MeshModelingToolsetExp/Source/GeometryProcessingAdapters/Private/CombineMeshInstancesImpl.cpp b/Engine/Plugins/Experimental/MeshModelingToolsetExp/Source/GeometryProcessingAdapters/Private/CombineMeshInstancesImpl.cpp
--- a/Engine/Plugins/Experimental/MeshModelingToolsetExp/Source/GeometryProcessingAdapters/Private/CombineMeshInstancesImpl.cpp
+++ b/Engine/Plugins/Experimental/MeshModelingToolsetExp/Source/GeometryProcessingAdapters/Private/CombineMeshInstancesImpl.cpp
@@ -97,6 +97,8 @@ struct FMeshPartInstance
 
 	EMeshDetailLevel DetailLevel = EMeshDetailLevel::Standard;
 
+	bool bAllowApproximation = true;
+
 	// allow FMeshPartInstance to maintain link to external representation of instance
 	FIndex3i ExternalInstanceIndex = FIndex3i::Invalid();
 
@@ -116,7 +118,6 @@ struct FMeshPart
 	bool bPreserveUVs = false;
 	bool bAllowMerging = true;
 
-
 	int32 GetNumTriangles() const 
 	{
 		if (SourceAsset != nullptr)
@@ -223,6 +224,7 @@ void InitializeMeshPartAssembly(
 		NewInstance.SourceComponent = SourceMeshInstance.SourceComponent;
 		NewInstance.SourceInstanceIndex = SourceMeshInstance.SourceInstanceIndex;
 		NewInstance.DetailLevel = static_cast<EMeshDetailLevel>( static_cast<int32>(SourceMeshInstance.DetailLevel) );
+		NewInstance.bAllowApproximation = SourceMeshInstance.bAllowApproximation;
 		for ( FTransform3d Transform : SourceMeshInstance.TransformSequence )
 		{
 			NewInstance.WorldTransform.Append( Transform );
@@ -278,6 +280,7 @@ void InitializeMeshPartAssembly(
 		NewInstance.SourceComponent = nullptr;
 		NewInstance.SourceInstanceIndex = 0;
 		NewInstance.DetailLevel = static_cast<EMeshDetailLevel>(static_cast<int32>(SourceMeshInstance.DetailLevel));
+		NewInstance.bAllowApproximation = SourceMeshInstance.bAllowApproximation;
 		for (FTransform3d Transform : SourceMeshInstance.TransformSequence)
 		{
 			NewInstance.WorldTransform.Append(Transform);
@@ -1443,9 +1446,9 @@ void ComputeMeshApproximations(
 			{
 				CostInfo[SetIndex].LODChainMeshes.Add(&SimplifiedLODMesh);
 			}
-			for (FDynamicMesh3& SimplifiedLODMesh : Assembly.OptimizedMeshGeometry[SetIndex].ApproximateMeshLODs)
+			for (FDynamicMesh3& ApproximateLODMesh : Assembly.OptimizedMeshGeometry[SetIndex].ApproximateMeshLODs)
 			{
-				CostInfo[SetIndex].LODChainMeshes.Add(&SimplifiedLODMesh);
+				CostInfo[SetIndex].LODChainMeshes.Add(&ApproximateLODMesh);
 			}
 		}
 
@@ -3346,8 +3349,9 @@ void BuildCombinedMesh(
 
 			for ( const FMeshPartInstance& Instance : Part->Instances )
 			{
-				bool bIsDecorativePart = (Instance.DetailLevel == EMeshDetailLevel::Decorative);
+				const FDynamicMesh3* InstanceAppendMesh = UseAppendMesh;
 
+				bool bIsDecorativePart = (Instance.DetailLevel == EMeshDetailLevel::Decorative);
 				if (bIsDecorativePart)
 				{
 					// filter out detail parts at higher LODs, or if we are doing VoxWrap LOD
@@ -3358,12 +3362,21 @@ void BuildCombinedMesh(
 					// at last detail part LOD, switch to approximate mesh
 					if (LODLevel >= (CombineOptions.FilterDecorativePartsLODLevel - CombineOptions.ApproximateDecorativePartLODs) )
 					{
-						UseAppendMesh = ApproximateAppendMesh;
+						InstanceAppendMesh = ApproximateAppendMesh;
 					}
 				}
 
+				// If approximation is disabled for this instance, fall back to last available simplified LOD.
+				// TODO: if part budget was applied, the mesh in this slot might actually be an approximation.
+				//   This is difficult to fix because the mesh LOD chains are per-part and not per-instance,
+				//   need some way to keep the original copied & simplified LOD chain around...
+				if (Instance.bAllowApproximation == false && LevelLODType == ECombinedLODType::Approximated)
+				{
+					InstanceAppendMesh = &OptimizedGeometry.SimplifiedMeshLODs.Last();
+				}
+
 				// need to make a copy to run pre-process func
-				FDynamicMesh3 TempAppendMesh(*UseAppendMesh);
+				FDynamicMesh3 TempAppendMesh(*InstanceAppendMesh);
 				if (Assembly.PreProcessInstanceMeshFunc)
 				{
 					Assembly.PreProcessInstanceMeshFunc(TempAppendMesh, Instance);
@@ -3394,14 +3407,15 @@ void BuildCombinedMesh(
 					}
 				}
 
-				// could precompute these indexes for each instance?
-				// also for source mesh we could transfer material IDs correctly...
-				UMaterialInterface* UseMaterial = Instance.Materials[0];
-				const int32* FoundMaterialIndex = Assembly.MaterialMap.Find(UseMaterial);
-				int32 AssignMaterialIndex = (FoundMaterialIndex != nullptr) ? *FoundMaterialIndex : 0;
-
+				// transfer Material IDs if part mesh has them
+				FDynamicMeshMaterialAttribute* AppendMaterialAttrib = TempAppendMesh.HasAttributes() ? TempAppendMesh.Attributes()->GetMaterialID() : nullptr;
 				for (int32 tid : TempAppendMesh.TriangleIndicesItr())
 				{
+					int32 SourceMaterialID = (AppendMaterialAttrib != nullptr) ? AppendMaterialAttrib->GetValue(tid) : 0;
+					UMaterialInterface* UseMaterial = Instance.Materials[SourceMaterialID];
+					const int32* FoundMaterialIndex = Assembly.MaterialMap.Find(UseMaterial);
+					int32 AssignMaterialIndex = (FoundMaterialIndex != nullptr) ? *FoundMaterialIndex : 0;
+
 					CombinedMeshLODData.MaterialIDs->SetValue( Mappings.GetNewTriangle(tid), AssignMaterialIndex );
 				}
 			}
diff --git a/Engine/Source/Developer/GeometryProcessingInterfaces/Public/GeometryProcessingInterfaces/CombineMeshInstances.h b/Engine/Source/Developer/GeometryProcessingInterfaces/Public/GeometryProcessingInterfaces/CombineMeshInstances.h
--- a/Engine/Source/Developer/GeometryProcessingInterfaces/Public/GeometryProcessingInterfaces/CombineMeshInstances.h
+++ b/Engine/Source/Developer/GeometryProcessingInterfaces/Public/GeometryProcessingInterfaces/CombineMeshInstances.h
@@ -58,6 +58,8 @@ public:
 		TArray<FTransform3d> TransformSequence;		// set of transforms on this instance. Often just a single transform.
 		int32 GroupDataIndex = -1;					// index into FSourceInstanceList::InstanceGroupDatas
 
+		bool bAllowApproximation = true;			// if false, only Copied or Simplified LODs will be used for this part instance
+
 		// in some cases it may be desirable to have "groups" of instances which should be output as separate meshes, but
 		// be jointly processed in terms of (eg) the part LODs. If any InstanceSubsetID is non-zero, then instance subsets
 		// are grouped/extracted by integer ID and will be returned as separate FOutputMesh's in the FResults. 
@@ -224,9 +226,10 @@ public:
 		// Can significantly reduce triangle count, but attributes on the interiors of polygonal areas will be completely discarded
 		int32 PlanarPolygonRetriangulationStartLOD = -1;
 
-		// If enabled, attempt to retriangulate planar areas of Source LODs to remove redundant coplanar geometry
+		// If enabled, attempt to retriangulate planar areas of Source LODs to remove redundant coplanar geometry. 
+		// This option affects individual parts and not the combined prefab.
 		bool bRetriangulateSourceLODs = true;
-		// which Source LOD to start retriangulating at
+		// which Source LOD to start planar retriangulation at
 		int32 StartRetriangulateSourceLOD = 1;
 
 		//
