CombineMeshInstances: add support for skipping face merging on specified Material IDs

diff --git a/Engine/Plugins/Experimental/MeshModelingToolsetExp/Source/GeometryProcessingAdapters/Private/CombineMeshInstancesImpl.cpp b/Engine/Plugins/Experimental/MeshModelingToolsetExp/Source/GeometryProcessingAdapters/Private/CombineMeshInstancesImpl.cpp
--- a/Engine/Plugins/Experimental/MeshModelingToolsetExp/Source/GeometryProcessingAdapters/Private/CombineMeshInstancesImpl.cpp
+++ b/Engine/Plugins/Experimental/MeshModelingToolsetExp/Source/GeometryProcessingAdapters/Private/CombineMeshInstancesImpl.cpp
@@ -2193,7 +2193,8 @@ static void PostProcessHiddenFaceRemovedMesh(
 	double BaseGeometricTolerance,
 	bool bTryToMergeFaces,
 	bool bApplyPlanarRetriangulation,
-	TFunctionRef<FIndex3i(const FDynamicMesh3& Mesh, int32 TriangleID)> GetTriangleGroupingIDFunc)
+	TFunctionRef<FIndex3i(const FDynamicMesh3& Mesh, int32 TriangleID)> GetTriangleGroupingIDFunc,
+	TSet<int32>* SkipMaterialIDs = nullptr)
 {
 	bool bVerbose = CVarGeometryCombineMeshInstancesVerbose.GetValueOnAnyThread();
 
@@ -2230,6 +2231,7 @@ static void PostProcessHiddenFaceRemovedMesh(
 		FMergeTriInfo TriInfo;
 		TriInfo.MaterialID = (MaterialIDs) ? MaterialIDs->GetValue(tid) : -1;
 		TriInfo.ExternalGroupingID = GetTriangleGroupingIDFunc(TargetMesh, tid);
+
 		int32* Found = UniqueMatIndices.Find(TriInfo);
 		if (Found == nullptr)
 		{
@@ -2256,6 +2258,16 @@ static void PostProcessHiddenFaceRemovedMesh(
 
 	for (FDynamicMesh3& SubRegionMesh : SplitMeshes)
 	{
+		// we split by MaterialID, so we can check any triangle against the SkipMaterialID list
+		if (SkipMaterialIDs != nullptr && SubRegionMesh.HasAttributes() && SubRegionMesh.Attributes()->HasMaterialID())
+		{
+			int32 MaterialID = SubRegionMesh.Attributes()->GetMaterialID()->GetValue(0);
+			if (SkipMaterialIDs->Contains(MaterialID))
+			{
+				continue;
+			}
+		}
+
 		// resolving T-junctions tends to make things worse...
 		//FMeshResolveTJunctions Resolver(&SubRegionMesh);
 		//Resolver.DistanceTolerance = 0.01;
@@ -2957,7 +2969,8 @@ void OptimizeLODMeshTriangulation(
 	FDynamicMesh3& MeshLOD,
 	int32 LODIndex,
 	IGeometryProcessing_CombineMeshInstances::FOptions CombineOptions,
-	double BaseGeometricTolerance)
+	double BaseGeometricTolerance,
+	TSet<int32>* SkipMaterialIDs = nullptr)
 {
 	bool bVerbose = CVarGeometryCombineMeshInstancesVerbose.GetValueOnAnyThread();
 
@@ -2970,7 +2983,7 @@ void OptimizeLODMeshTriangulation(
 	PostProcessHiddenFaceRemovedMesh(MeshLOD, BaseGeometricTolerance,
 		CombineOptions.bMergeCoplanarFaces && LODIndex >= CombineOptions.MergeCoplanarFacesStartLOD,
 		CombineOptions.PlanarPolygonRetriangulationStartLOD >= 0 && LODIndex >= CombineOptions.PlanarPolygonRetriangulationStartLOD,
-		GroupingIDFunc);
+		GroupingIDFunc, SkipMaterialIDs);
 }
 
 
@@ -2981,7 +2994,8 @@ void ProcessCombinedLODChain(
 	TArray<FCombinedMeshLOD>& MeshLODs,
 	const TArray<double>& OptimizationTolerances,
 	int32 FirstVoxWrappedIndex,
-	IGeometryProcessing_CombineMeshInstances::FOptions CombineOptions
+	IGeometryProcessing_CombineMeshInstances::FOptions CombineOptions,
+	TSet<int32>* PreserveTopologyMaterialIDs = nullptr
 )
 {
 	using namespace UE::Geometry;
@@ -3003,10 +3017,10 @@ void ProcessCombinedLODChain(
 			}
 
 			double UseTolerance = OptimizationTolerances[LODIndex];
-			UE::Tasks::FTask RemoveHiddenTask = UE::Tasks::Launch(UE_SOURCE_LOCATION, [&MeshLODs, &CombineOptions, LODIndex, UseTolerance]()
+			UE::Tasks::FTask RemoveHiddenTask = UE::Tasks::Launch(UE_SOURCE_LOCATION, [&MeshLODs, &CombineOptions, LODIndex, UseTolerance, PreserveTopologyMaterialIDs]()
 			{
 				ComputeHiddenRemovalForLOD(MeshLODs[LODIndex].Mesh, LODIndex, CombineOptions);
-				OptimizeLODMeshTriangulation(MeshLODs[LODIndex].Mesh, LODIndex, CombineOptions, UseTolerance);
+				OptimizeLODMeshTriangulation(MeshLODs[LODIndex].Mesh, LODIndex, CombineOptions, UseTolerance, PreserveTopologyMaterialIDs);
 			});
 			PendingRemoveHiddenTasks.Add(RemoveHiddenTask);
 
@@ -3422,6 +3436,16 @@ void BuildCombinedMesh(
 		}
 	}
 
+	// Some Material regions may need to be explicitly preserved, this set will be passed on later
+	TSet<int32> PreserveTopologyMaterialIDSet;
+	for (UMaterialInterface* Material : CombineOptions.PreventMergingMaterialSet)
+	{
+		if (const int32* FoundMaterialIndex = Assembly.MaterialMap.Find(Material))
+		{
+			PreserveTopologyMaterialIDSet.Add(*FoundMaterialIndex);
+		}
+	}
+
 	// Above we copied the last approximation LOD into the first VoxWrap LOD slot. However 
 	// approximate LOD may be very coarse, it's better to use a known-good starting mesh?
 	// Need to do this before starting hidden-removal task...
@@ -3486,9 +3510,9 @@ void BuildCombinedMesh(
 		for (int32 SubsetIndex = 0; SubsetIndex < NumSubsets; ++SubsetIndex)
 		{
 			TArray<FCombinedMeshLOD>& LODChain = SubsetMeshLODChains[SubsetIndex];
-			UE::Tasks::FTask ProcessSubsetTask = UE::Tasks::Launch(UE_SOURCE_LOCATION, [&LODChain, &CombineOptions, &OptimizationTolerances, FirstVoxWrappedIndex]()
+			UE::Tasks::FTask ProcessSubsetTask = UE::Tasks::Launch(UE_SOURCE_LOCATION, [&LODChain, &CombineOptions, &OptimizationTolerances, &PreserveTopologyMaterialIDSet, FirstVoxWrappedIndex]()
 			{
-				ProcessCombinedLODChain(LODChain, OptimizationTolerances, FirstVoxWrappedIndex, CombineOptions);
+				ProcessCombinedLODChain(LODChain, OptimizationTolerances, FirstVoxWrappedIndex, CombineOptions, &PreserveTopologyMaterialIDSet);
 			});
 			PendingSubsetTasks.Add(ProcessSubsetTask);
 			if (bVerbose)
diff --git a/Engine/Source/Developer/GeometryProcessingInterfaces/Public/GeometryProcessingInterfaces/CombineMeshInstances.h b/Engine/Source/Developer/GeometryProcessingInterfaces/Public/GeometryProcessingInterfaces/CombineMeshInstances.h
--- a/Engine/Source/Developer/GeometryProcessingInterfaces/Public/GeometryProcessingInterfaces/CombineMeshInstances.h
+++ b/Engine/Source/Developer/GeometryProcessingInterfaces/Public/GeometryProcessingInterfaces/CombineMeshInstances.h
@@ -226,6 +226,10 @@ public:
 		// Can significantly reduce triangle count, but attributes on the interiors of polygonal areas will be completely discarded
 		int32 PlanarPolygonRetriangulationStartLOD = -1;
 
+		// Triangles with Materials assigned that are also in this set will not be allowed to be combined/retriangulated
+		// with adjacent triangles. This can be used to preserve topology/UVs/etc for specific material areas.
+		TArray<UMaterialInterface*> PreventMergingMaterialSet;
+
 		// If enabled, attempt to retriangulate planar areas of Source LODs to remove redundant coplanar geometry. 
 		// This option affects individual parts and not the combined prefab.
 		bool bRetriangulateSourceLODs = true;
