CombineMeshInstances: fix paste-o

diff --git a/Engine/Plugins/Experimental/MeshModelingToolsetExp/Source/GeometryProcessingAdapters/Private/CombineMeshInstancesImpl.cpp b/Engine/Plugins/Experimental/MeshModelingToolsetExp/Source/GeometryProcessingAdapters/Private/CombineMeshInstancesImpl.cpp
--- a/Engine/Plugins/Experimental/MeshModelingToolsetExp/Source/GeometryProcessingAdapters/Private/CombineMeshInstancesImpl.cpp
+++ b/Engine/Plugins/Experimental/MeshModelingToolsetExp/Source/GeometryProcessingAdapters/Private/CombineMeshInstancesImpl.cpp
@@ -3534,7 +3534,7 @@ void BuildCombinedMesh(
 	}
 	else
 	{
-		ProcessCombinedLODChain(MeshLODs, OptimizationTolerances, FirstVoxWrappedIndex, CombineOptions);
+		ProcessCombinedLODChain(MeshLODs, OptimizationTolerances, FirstVoxWrappedIndex, CombineOptions, &PreserveTopologyMaterialIDSet);
 
 		CombinedResults.SetNum(1);
 		BuildOutputSubAssembly(MeshLODs, 0, FirstVoxWrappedIndex, CombinedResults[0]);
