Skeleton Modifier: ensure that the mirrored names are unique and fixed missbehaviors with multiple patterns found

    eg. 'spine_04_latissimus_r' should return 'spine_04_latissimus_l' and not 'spine_04_ratissimus_r' or 'spine_04_ratissimus_l'
    make sure that the last found pattern takes priority to avoid replacing the wrong one

+ fixed recursive get children

diff --git a/Engine/Plugins/Experimental/MeshModelingToolsetExp/Source/SkeletalMeshModifiers/Private/SkeletonModifier.cpp b/Engine/Plugins/Experimental/MeshModelingToolsetExp/Source/SkeletalMeshModifiers/Private/SkeletonModifier.cpp
--- a/Engine/Plugins/Experimental/MeshModelingToolsetExp/Source/SkeletalMeshModifiers/Private/SkeletonModifier.cpp
+++ b/Engine/Plugins/Experimental/MeshModelingToolsetExp/Source/SkeletalMeshModifiers/Private/SkeletonModifier.cpp
@@ -15,6 +15,7 @@
 #include "AssetRegistry/AssetData.h"
 #include "Engine/SkeletalMesh.h"
 #include "RenderingThread.h"
+#include "Animation/MirrorDataTable.h"
 
 
 namespace USkeletonModifierLocals
@@ -578,34 +579,45 @@ void USkeletonModifier::GetMirroredNames(
 	{
 		return;
 	}
+
+	const FName Left(InOptions.LeftString), Right(InOptions.RightString);
 	
+	const TArray MirrorFindReplaceExpressions({
+		FMirrorFindReplaceExpression(Left, Right, EMirrorFindReplaceMethod::Suffix),
+		FMirrorFindReplaceExpression(Right, Left, EMirrorFindReplaceMethod::Suffix)
+		});
+		
 	const TArray<FMeshBoneInfo>& BoneInfos = ReferenceSkeleton->GetRawRefBoneInfo();
 	
 	OutBonesName.Reserve(InBonesToMirror.Num());
 	Algo::Transform(InBonesToMirror, OutBonesName, [&](const int32 BoneIndex)
 	{
-		FString BoneStr = BoneInfos[BoneIndex].Name.ToString();
-		if (BoneStr.Contains(InOptions.LeftString, ESearchCase::CaseSensitive, ESearchDir::FromEnd))
-		{
-			BoneStr = BoneStr.Replace(*InOptions.LeftString, *InOptions.RightString, ESearchCase::CaseSensitive);
-		}
-		else if (BoneStr.Contains(InOptions.RightString, ESearchCase::CaseSensitive, ESearchDir::FromEnd))
-		{
-			BoneStr = BoneStr.Replace(*InOptions.RightString, *InOptions.LeftString, ESearchCase::CaseSensitive);
-		}
-		else
+		const FName BoneName = BoneInfos[BoneIndex].Name;
+		const FName MirrorName = UMirrorDataTable::GetMirrorName(BoneName, MirrorFindReplaceExpressions);
+
+		if (MirrorName.IsNone() || MirrorName == BoneName)
 		{
-			return GetUniqueName(BoneInfos[BoneIndex].Name, OutBonesName);
+			return GetUniqueName(BoneName, OutBonesName);
 		}
-		return FName(*BoneStr);
+
+		return MirrorName;
 	});
 }
 
 void USkeletonModifier::GetMirroredBones(
 	const TArray<int32>& InBonesToMirror, const TArray<FName>& InMirroredNames, TArray<int32>& OutMirroredBones)
 {
+	const int32 NumBones = InBonesToMirror.Num();
+	
 	OutMirroredBones.Reset();
-	if (InBonesToMirror.IsEmpty() || InBonesToMirror.Num() != InMirroredNames.Num())
+	if (InBonesToMirror.IsEmpty() || NumBones != InMirroredNames.Num())
+	{
+		return;
+	}
+
+	// check mirrored names uniqueness 
+	const TSet<FName> UniqueMirroredNames(InMirroredNames);
+	if (UniqueMirroredNames.Num() != NumBones)
 	{
 		return;
 	}
@@ -615,7 +627,7 @@ void USkeletonModifier::GetMirroredBones(
 	
 	TArray<FName> BonesToAdd, ParentNames;
 	TArray<FTransform> Transforms;
-	const int32 NumBones = InBonesToMirror.Num();
+
 	for (int32 Index = 0; Index < NumBones; Index++)
 	{
 		const FName& MirroredName = InMirroredNames[Index];
@@ -1237,44 +1249,128 @@ FName USkeletonModifier::GetUniqueName(const FName InBoneName, const TArray<FNam
 		return NAME_None;
 	}
 
-	// remove unwanted characters
-	auto GetSanitizedName = [](const FName InName)
+	static constexpr TCHAR Underscore = '_';
+	static constexpr TCHAR Hashtag = '#';
+
+	int32 LastHashtag = INDEX_NONE, LastDigit = INDEX_NONE;
+	
+	FString InBoneNameStr = InBoneName.ToString();
+
+	// 1. Remove white spaces from start / end
+	InBoneNameStr.TrimStartAndEndInline();
+	const int32 NameLength = InBoneNameStr.Len();
+
+	// 2. Sanitize name: remove unwanted characters and get padding info from hashtags or digits
+	auto SanitizedName = [&]()
 	{
 		bool bHasAnyGoodChar = false;
-		FString Name = InName.ToString(); 
-		Name.TrimStartAndEndInline();
-		for (int32 i = 0; i < Name.Len(); ++i)
+		for (int32 Index = 0; Index < NameLength; ++Index)
 		{
-			TCHAR& C = Name[i];
-			const bool bGoodChar = FChar::IsAlpha(C) || (C == '_') || ((i > 0) && FChar::IsDigit(C));
+			TCHAR& C = InBoneNameStr[Index];
+
+			const bool bIsAlphaOrUnderscore = FChar::IsAlpha(C) || (C == Underscore);
+			
+			const bool bIsDigit = FChar::IsDigit(C);
+			if (bIsDigit)
+			{
+				LastDigit = Index;
+			}
+			
+			const bool bIsHashtag = C == Hashtag;
+			if (bIsHashtag)
+			{
+				LastHashtag = Index;
+			}
+			
+			const bool bGoodChar =	bIsAlphaOrUnderscore || bIsHashtag || ((Index > 0) && bIsDigit);
 			bHasAnyGoodChar |= bGoodChar;
 			if (!bGoodChar)
 			{
-				C = '_';
+				C = Underscore;
 			}
 		}
-		return bHasAnyGoodChar ? FName(*Name) : NAME_None;
+		return bHasAnyGoodChar;
 	};
 
-	// remove digits at the end
-	auto GetPrefix = [](const FName InName) -> FName
+	// 3. Early exit if none of the character is good to use
+	const bool bHasAnyGoodChar = SanitizedName();
+	if (!bHasAnyGoodChar)
 	{
-		const FString Name = InName.ToString();
-		const int32 LastNonDigit = Name.FindLastCharByPredicate([](TCHAR C) { return !FChar::IsDigit(C); });
-		if (LastNonDigit == INDEX_NONE)
+		return NAME_None;
+	}
+
+	// 4. if we found a padding, check if there are digits before/after to grow it if needed (ei. joint_##10_left)
+
+	// note that # takes priority here
+	int32 EndPadding = LastHashtag != INDEX_NONE ? LastHashtag : LastDigit != INDEX_NONE ? LastDigit : INDEX_NONE;
+	int32 StartPadding = EndPadding;
+	FString PaddingStr;
+
+	if (EndPadding != INDEX_NONE)
+	{
+		// find # or digit before EndPadding
+		bool bIsHashtag = InBoneNameStr[StartPadding] == Hashtag, bIsDigit = FChar::IsDigit(InBoneNameStr[StartPadding]);
+
+		while (StartPadding > INDEX_NONE && (bIsHashtag || bIsDigit))
 		{
-			return NAME_None;
+			--StartPadding;
+			if (StartPadding > INDEX_NONE)
+			{
+				bIsHashtag = InBoneNameStr[StartPadding] == Hashtag;
+				bIsDigit = FChar::IsDigit(InBoneNameStr[StartPadding]);
+			}
 		}
-		if (LastNonDigit == Name.Len()-1)
+		
+		// find # or digit after StartPadding
+		StartPadding = StartPadding+1;
+		EndPadding = StartPadding;
+
+		bIsHashtag = InBoneNameStr[EndPadding] == Hashtag, bIsDigit = FChar::IsDigit(InBoneNameStr[EndPadding]);
+
+		while (EndPadding < NameLength && (bIsHashtag || bIsDigit))
 		{
-			return InName;
+			++EndPadding;
+			if (EndPadding < NameLength)
+			{
+				bIsHashtag = InBoneNameStr[EndPadding] == Hashtag;
+				bIsDigit = FChar::IsDigit(InBoneNameStr[EndPadding]);
+			}
 		}
-		return FName(*Name.Left(LastNonDigit+1));
-	};
+		EndPadding -= 1;
 
-	// check for availability
-	auto IsNameAvailable = [&](const FName Name)
+		// store the padding string
+		if (EndPadding >= StartPadding)
+		{
+			PaddingStr = InBoneNameStr.Mid(StartPadding, EndPadding-StartPadding+1);
+		}
+
+		// replace any # with zeros
+		static constexpr TCHAR Zero = '0';
+		InBoneNameStr.ReplaceInline(&Hashtag, &Zero);
+		PaddingStr.ReplaceInline(&Hashtag, &Zero);
+	}
+
+	// 5. prepare prefix, suffix and padding
+	FString Prefix = InBoneNameStr;
+	FString Suffix;
+	int32 CurrentIndex = 1;
+
+	if (StartPadding != INDEX_NONE)
 	{
+		CurrentIndex = FCString::Atoi(*PaddingStr);
+		if (CurrentIndex == 0)
+		{
+			PaddingStr[PaddingStr.Len()-1] = '1';
+			CurrentIndex = 1;
+		}
+		Prefix = InBoneNameStr.Left(StartPadding);
+		Suffix = InBoneNameStr.RightChop(EndPadding+1);
+	}
+
+	// check for availability in both the reference skeleton and the names that are going to be added 
+	auto IsNameAvailable = [&](const FString& InNameStr)
+	{
+		const FName Name(*InNameStr);
 		const int32 Index = ReferenceSkeleton->FindRawBoneIndex(Name);
 		if (Index != INDEX_NONE)
 		{
@@ -1289,26 +1385,35 @@ FName USkeletonModifier::GetUniqueName(const FName InBoneName, const TArray<FNam
 		return true;
 	};
 
-	const FName SanitizedName = GetSanitizedName(InBoneName);
-	if (SanitizedName == NAME_None)
-	{
-		return NAME_None;
-	}
-	
-	const FName Prefix = GetPrefix(SanitizedName);
-	if (Prefix == NAME_None)
+	// 6. build the new unique name 
+	FString OutBoneNameStr = Prefix + PaddingStr + Suffix;
+	while (!IsNameAvailable(OutBoneNameStr))
 	{
-		return NAME_None;
-	}
+		// increment the index
+		FString NewIncrement = PaddingStr.FromInt(CurrentIndex++);
 
-	FName OutName = Prefix;
-	int32 CurrentIndex = 0;
-	while (!IsNameAvailable(OutName))
-	{
-		OutName = FName(*FString::Printf(TEXT("%s%d"), *Prefix.ToString(), CurrentIndex++));
+		// switch this new index into a padding str
+		const int32 IncrementLen = PaddingStr.Len();
+		const int32 NewIncrementLen = NewIncrement.Len();
+		if (NewIncrementLen < IncrementLen)
+		{
+			int32 Index = 0;
+			while (Index < NewIncrementLen)
+			{
+				PaddingStr[IncrementLen-1-Index] = NewIncrement[NewIncrementLen-1-Index];
+				Index++;
+			}
+		}
+		else
+		{
+			PaddingStr = NewIncrement;
+		}
+
+		// form the new name
+		OutBoneNameStr = Prefix + PaddingStr + Suffix;
 	}
 
-	return OutName;
+	return FName(*OutBoneNameStr);
 }
 
 const FReferenceSkeleton& USkeletonModifier::GetReferenceSkeleton() const
@@ -1369,9 +1474,9 @@ TArray<FName> USkeletonModifier::GetChildrenNames(const FName InBoneName, const
 		auto GetChildren = [&](const int32 InBoneIndex, auto&& GetChildren) -> void
 		{
 			TArray<int32> Children; ReferenceSkeleton->GetRawDirectChildBones(InBoneIndex, Children);
+			ChildrenIndices.Append(Children);
 			for (int32 ChildIndex: Children)
 			{
-				ChildrenIndices.Add(InBoneIndex);
 				GetChildren(ChildIndex, GetChildren);
 			}
 		};
