Fix modeling mode vertex paint tool using incorrect transform for camera to paint mesh

diff --git a/Engine/Plugins/Experimental/MeshModelingToolsetExp/Source/MeshModelingToolsExp/Private/MeshVertexPaintTool.cpp b/Engine/Plugins/Experimental/MeshModelingToolsetExp/Source/MeshModelingToolsExp/Private/MeshVertexPaintTool.cpp
--- a/Engine/Plugins/Experimental/MeshModelingToolsetExp/Source/MeshModelingToolsExp/Private/MeshVertexPaintTool.cpp
+++ b/Engine/Plugins/Experimental/MeshModelingToolsetExp/Source/MeshModelingToolsExp/Private/MeshVertexPaintTool.cpp
@@ -1049,8 +1049,7 @@ void UMeshVertexPaintTool::ApplyVisibilityFilter(const TArray<int32>& Triangles,
 
 	FViewCameraState StateOut;
 	GetToolManager()->GetContextQueriesAPI()->GetCurrentViewState(StateOut);
-	FTransform3d LocalToWorld = UE::ToolTarget::GetLocalToWorldTransform(Target);
-	FVector3d LocalEyePosition(LocalToWorld.InverseTransformPosition(StateOut.Position));
+	FVector3d LocalEyePosition(CurTargetTransform.InverseTransformPosition(StateOut.Position));
 
 	const FDynamicMesh3* Mesh = GetSculptMesh();
 
