Skin Binding: fixed pos to grid/cell center indexation

+ removed occupancy grid computation when useless (this is only used for debug purposes and slows down the binding calculation)

diff --git a/Engine/Plugins/Experimental/MeshModelingToolsetExp/Source/MeshModelingToolsEditorOnlyExp/Private/SkeletalMesh/SkinWeightsBindingTool.cpp b/Engine/Plugins/Experimental/MeshModelingToolsetExp/Source/MeshModelingToolsEditorOnlyExp/Private/SkeletalMesh/SkinWeightsBindingTool.cpp
--- a/Engine/Plugins/Experimental/MeshModelingToolsetExp/Source/MeshModelingToolsEditorOnlyExp/Private/SkeletalMesh/SkinWeightsBindingTool.cpp
+++ b/Engine/Plugins/Experimental/MeshModelingToolsetExp/Source/MeshModelingToolsEditorOnlyExp/Private/SkeletalMesh/SkinWeightsBindingTool.cpp
@@ -105,12 +105,22 @@ void USkinWeightsBindingTool::Setup()
 								 ToolSetupUtil::GetDefaultWorkingMaterial(GetToolManager())
 	);
 
+	auto UpdateOccupancy = [this]()
+	{
+		Occupancy.Reset();
+		if (Properties->BindingType == ESkinWeightsBindType::GeodesicVoxel && Properties->bDebugDraw)
+		{
+			Occupancy = MakeShared<UE::Geometry::FOccupancyGrid3>(*OriginalMesh, Properties->VoxelResolution);
+		}
+	};
+	
 	// setup watchers
-	auto HandlePropertyChange = [this](const bool bUpdateOperator)
+	auto HandlePropertyChange = [this, UpdateOccupancy](const bool bUpdateOperator)
 	{
 		if (bUpdateOperator)
 		{
 			Preview->InvalidateResult();
+			UpdateOccupancy();
 		}
 		UpdateVisualization();
 	};
@@ -124,12 +134,7 @@ void USkinWeightsBindingTool::Setup()
 	Properties->WatchProperty(Properties->MaxInfluences,
 							  [HandlePropertyChange](int32) { HandlePropertyChange(true); });
 	Properties->WatchProperty(Properties->VoxelResolution,
-							  [this](int32)
-							  {
-							  	Occupancy = MakeShared<UE::Geometry::FOccupancyGrid3>(*OriginalMesh, Properties->VoxelResolution);
-								Preview->InvalidateResult();
-								UpdateVisualization();
-							  });
+							  [HandlePropertyChange](int32) { HandlePropertyChange(true); });
 
 	OriginalMesh = MakeShared<FDynamicMesh3, ESPMode::ThreadSafe>();
 	FMeshDescriptionToDynamicMesh Converter;
@@ -147,7 +152,7 @@ void USkinWeightsBindingTool::Setup()
 	Preview->PreviewMesh->SetShadowsEnabled(false);
 	Preview->PreviewMesh->UpdatePreview(OriginalMesh.Get());
 
-	Occupancy = MakeShared<UE::Geometry::FOccupancyGrid3>(*OriginalMesh, Properties->VoxelResolution);
+	UpdateOccupancy();
 	
 	UpdateVisualization(/*bForce=*/true);
 	
@@ -256,9 +261,8 @@ void USkinWeightsBindingTool::Render(IToolsContextRenderAPI* RenderAPI)
 			if (bShowBoundary && Domain == FOccupancyGrid3::EDomain::Boundary)
 			{
 				FBox Box{Occupancy->GetCellBoxFromIndex(OccupancyIndex)};
-				DrawBox(RenderAPI, Transform, Box, FLinearColor(1.0, 1.0, 0.0, 0.5), 0.5f);
+				DrawBox(RenderAPI, Transform, Box, FLinearColor(1.0, 1.0, 0.0, 0.5), 0.0f);
 			}
-			
 		}
 	}
 	
diff --git a/Engine/Plugins/Runtime/MeshModelingToolset/Source/ModelingOperators/Private/SkinningOps/SkinBindingOp.cpp b/Engine/Plugins/Runtime/MeshModelingToolset/Source/ModelingOperators/Private/SkinningOps/SkinBindingOp.cpp
--- a/Engine/Plugins/Runtime/MeshModelingToolset/Source/ModelingOperators/Private/SkinningOps/SkinBindingOp.cpp
+++ b/Engine/Plugins/Runtime/MeshModelingToolset/Source/ModelingOperators/Private/SkinningOps/SkinBindingOp.cpp
@@ -556,21 +556,24 @@ void FSkinBindingOp::CreateSkinWeights_GeodesicVoxel(
 			const FVector3i CellIndex = Occupancy.GetCellIndexFromPoint(Pos);
 			const FVector3d CellCenter{Occupancy.GetCellCenterFromIndex(CellIndex)};
 
-			float Distance = BoneDistance[CellIndex];
-			Distance += FVector3d::Distance(CellCenter, Pos);
+			if (ensure(BoneDistance.IsValidIndex(CellIndex)))
+			{
+				float Distance = BoneDistance[CellIndex];
+				Distance += FVector3d::Distance(CellCenter, Pos);
 				
-			// Normalize the distance by the diagonal size of the bbox to maintain scale invariance.
-			float Weight = Distance / DiagonalBounds;
+				// Normalize the distance by the diagonal size of the bbox to maintain scale invariance.
+				float Weight = Distance / DiagonalBounds;
 
-			// Avoid div-by-zero but allow for the possibility that multiple bones may
-			// touch this vertex.
-			Weight = FMath::Max(Weight, UE_KINDA_SMALL_NUMBER);
+				// Avoid div-by-zero but allow for the possibility that multiple bones may
+				// touch this vertex.
+				Weight = FMath::Max(Weight, UE_KINDA_SMALL_NUMBER);
 
-			// Compute the actual weight, factoring in the stiffness value. W = (1/S(D))^2
-			// Where S(x) is the stiffness function.
-			Weight = FMath::Square(1.0f / ComputeWeightStiffness(Weight, InStiffness));
+				// Compute the actual weight, factoring in the stiffness value. W = (1/S(D))^2
+				// Where S(x) is the stiffness function.
+				Weight = FMath::Square(1.0f / ComputeWeightStiffness(Weight, InStiffness));
 
-			Weights[VertexIdx * TransformHierarchy.Num() + BoneIndex] = Weight;
+				Weights[VertexIdx * TransformHierarchy.Num() + BoneIndex] = Weight;
+			}
 		}
 	});	
 
diff --git a/Engine/Source/Runtime/GeometryCore/Public/Spatial/OccupancyGrid3.h b/Engine/Source/Runtime/GeometryCore/Public/Spatial/OccupancyGrid3.h
--- a/Engine/Source/Runtime/GeometryCore/Public/Spatial/OccupancyGrid3.h
+++ b/Engine/Source/Runtime/GeometryCore/Public/Spatial/OccupancyGrid3.h
@@ -59,8 +59,8 @@ struct FOccupancyGrid3
 			FVector3i(1, 1, 0),
 			FVector3i(1, 1, 1),
 		};
-		
-		for (int32 OccupancyId = 0; OccupancyId < Occupancy.Size(); OccupancyId++)
+
+		ParallelFor(Occupancy.Size(), [&](const int32 OccupancyId)
 		{
 			const FVector3i OccupancyIndex(Occupancy.ToIndex(OccupancyId));
 			const int32 Count = Algo::CountIf(CornerOffsets, [&WindingGrid, OccupancyIndex](FVector3i CornerOffset)
@@ -77,7 +77,7 @@ struct FOccupancyGrid3
 			{
 				Occupancy[OccupancyIndex] = EDomain::Boundary;
 			}
-		}
+		});
 
 		// Make sure we include all the vertices of the mesh as a part of the boundary, if
 		// the vertex areas are marked as being exterior.
@@ -85,9 +85,13 @@ struct FOccupancyGrid3
 		{
 			const FVector3d& Pos = InMesh.GetVertex(VertexIdx);
 			const FVector3i OccupancyIndex = GetCellIndexFromPoint(FVector(Pos));
-			if (Occupancy[OccupancyIndex] == EDomain::Exterior)
+
+			if (ensure(Occupancy.IsValidIndex(OccupancyIndex)))
 			{
-				Occupancy[OccupancyIndex] = EDomain::Boundary;
+				if (Occupancy[OccupancyIndex] == EDomain::Exterior)
+				{
+					Occupancy[OccupancyIndex] = EDomain::Boundary;
+				}
 			}
 		}
 	}
@@ -99,8 +103,7 @@ struct FOccupancyGrid3
 	{
 		FVector3f PP(InPoint);
 		PP -= GridOrigin;
-		PP += CellMidPoint;
-
+		
 		return { FMath::FloorToInt(PP.X / CellSize),
 				FMath::FloorToInt(PP.Y / CellSize),
 				FMath::FloorToInt(PP.Z / CellSize) };
@@ -110,7 +113,9 @@ struct FOccupancyGrid3
 	/// in the winding number grid, from which the occupancy got computed.
 	FVector3f GetCellCenterFromIndex(const FVector3i &Index) const
 	{
-		return {float(Index.X) * CellSize + GridOrigin.X, float(Index.Y) * CellSize + GridOrigin.Y, float(Index.Z) * CellSize + GridOrigin.Z};
+		return {float(Index.X) * CellSize + GridOrigin.X + CellMidPoint.X,
+				float(Index.Y) * CellSize + GridOrigin.Y + CellMidPoint.Y,
+				float(Index.Z) * CellSize + GridOrigin.Z + CellMidPoint.Z};
 	}
 
 	/// Given an index into the occupancy grid, returns the bbox that represents the cell in the winding number
