Fix crash in paint maps tool if user attempts to paint when there are no attributes on the mesh, then tries to undo

diff --git a/Engine/Plugins/Experimental/MeshModelingToolsetExp/Source/MeshModelingToolsExp/Private/MeshAttributePaintTool.cpp b/Engine/Plugins/Experimental/MeshModelingToolsetExp/Source/MeshModelingToolsExp/Private/MeshAttributePaintTool.cpp
--- a/Engine/Plugins/Experimental/MeshModelingToolsetExp/Source/MeshModelingToolsExp/Private/MeshAttributePaintTool.cpp
+++ b/Engine/Plugins/Experimental/MeshModelingToolsetExp/Source/MeshModelingToolsExp/Private/MeshAttributePaintTool.cpp
@@ -362,9 +362,12 @@ void UMeshAttributePaintTool::OnEndDrag(const FRay& Ray)
 
 	// close change record
 	TUniquePtr<FMeshAttributePaintChange> Change = EndChange();
-	GetToolManager()->BeginUndoTransaction(LOCTEXT("AttributeValuesChange", "Paint"));
-	GetToolManager()->EmitObjectChange(this, MoveTemp(Change), LOCTEXT("AttributeValuesChange", "Paint"));
-	GetToolManager()->EndUndoTransaction();
+	if (Change)
+	{
+		GetToolManager()->BeginUndoTransaction(LOCTEXT("AttributeValuesChange", "Paint"));
+		GetToolManager()->EmitObjectChange(this, MoveTemp(Change), LOCTEXT("AttributeValuesChange", "Paint"));
+		GetToolManager()->EndUndoTransaction();
+	}
 }
 
 
@@ -728,6 +731,10 @@ void UMeshAttributePaintTool::OnShutdown(EToolShutdownType ShutdownType)
 
 void UMeshAttributePaintTool::BeginChange()
 {
+	if (CurrentAttributeIndex < 0)
+	{
+		return;
+	}
 	if (! ActiveChangeBuilder)
 	{
 		ActiveChangeBuilder = MakeUnique<TIndexedValuesChangeBuilder<float, FMeshAttributePaintChange>>();
@@ -739,6 +746,11 @@ void UMeshAttributePaintTool::BeginChange()
 
 TUniquePtr<FMeshAttributePaintChange> UMeshAttributePaintTool::EndChange()
 {
+	if (!ActiveChangeBuilder)
+	{
+		return nullptr;
+	}
+
 	TUniquePtr<FMeshAttributePaintChange> Result = ActiveChangeBuilder->ExtractResult();
 	
 	Result->ApplyFunction = [](UObject* Object, const int32& AttribIndex, const TArray<int32>& Indices, const TArray<float>& Values)
@@ -759,7 +771,10 @@ TUniquePtr<FMeshAttributePaintChange> UMeshAttributePaintTool::EndChange()
 
 void UMeshAttributePaintTool::ExternalUpdateValues(int32 AttribIndex, const TArray<int32>& VertexIndices, const TArray<float>& NewValues)
 {
-	check(Attributes.IsValidIndex(AttribIndex));
+	if (!ensure(Attributes.IsValidIndex(AttribIndex)))
+	{
+		return;
+	}
 	FAttributeData& AttribData = Attributes[AttribIndex];
 
 	int32 NumV = VertexIndices.Num();
