Modeling Mode: Fix issue where weightmaps are loaded with invalid values for a target when tools' settings are loaded with a preset.

diff --git a/Engine/Plugins/Experimental/MeshModelingToolsetExp/Source/MeshModelingToolsExp/Private/DisplaceMeshTool.cpp b/Engine/Plugins/Experimental/MeshModelingToolsetExp/Source/MeshModelingToolsExp/Private/DisplaceMeshTool.cpp
--- a/Engine/Plugins/Experimental/MeshModelingToolsetExp/Source/MeshModelingToolsExp/Private/DisplaceMeshTool.cpp
+++ b/Engine/Plugins/Experimental/MeshModelingToolsetExp/Source/MeshModelingToolsExp/Private/DisplaceMeshTool.cpp
@@ -1467,12 +1467,22 @@ void UDisplaceMeshTool::UpdateActiveWeightMap()
 	{
 		TSharedPtr<FIndexedWeightMap, ESPMode::ThreadSafe> NewWeightMap = MakeShared<FIndexedWeightMap, ESPMode::ThreadSafe>();
 		const FMeshDescription* MeshDescription = UE::ToolTarget::GetMeshDescription(Target);
-		UE::WeightMaps::GetVertexWeightMap(MeshDescription, CommonProperties->WeightMap, *NewWeightMap, 1.0f);
+		bool bFound = UE::WeightMaps::GetVertexWeightMap(MeshDescription, CommonProperties->WeightMap, *NewWeightMap, 1.0f);
 		if (CommonProperties->bInvertWeightMap)
 		{
 			NewWeightMap->InvertWeightMap();
 		}
-		ActiveWeightMap = NewWeightMap;
+		// We'll check for invalid weightmaps here and reset to None if they aren't valid. This helps in cases where values are set externally,
+		// for example from preset loading, where the weightmap values incoming might not have any relation to the current target mesh.
+		if (bFound)
+		{
+			ActiveWeightMap = NewWeightMap;
+		}
+		else
+		{
+			CommonProperties->WeightMap = FName(CommonProperties->WeightMapsList[0]);
+			ActiveWeightMap = nullptr;
+		}
 	}
 }
 
diff --git a/Engine/Plugins/Runtime/MeshModelingToolset/Source/ModelingComponents/Private/BaseTools/BaseMeshProcessingTool.cpp b/Engine/Plugins/Runtime/MeshModelingToolset/Source/ModelingComponents/Private/BaseTools/BaseMeshProcessingTool.cpp
--- a/Engine/Plugins/Runtime/MeshModelingToolset/Source/ModelingComponents/Private/BaseTools/BaseMeshProcessingTool.cpp
+++ b/Engine/Plugins/Runtime/MeshModelingToolset/Source/ModelingComponents/Private/BaseTools/BaseMeshProcessingTool.cpp
@@ -384,6 +384,14 @@ void UBaseMeshProcessingTool::OnSelectedWeightMapChanged(bool bInvalidate)
 	{
 		NewWeightMap->InvertWeightMap();
 	}
+
+	// This will automatically reset the weightmap to None if the new value doesn't exist. This could easily occur during a preset load, for instance, and we 
+	// don't want to pollute the value with garbage from a preset value that doesn't apply to the active target.
+	if (!bFound)
+	{
+		WeightMapPropertySet->SetSelectedFromWeightMapIndex(-1);
+	}
+
 	ActiveWeightMap = NewWeightMap;
 
 	if (bInvalidate)
