CombineMeshInstances: add support for >1 UV channel in the combined LOD meshes

diff --git a/Engine/Plugins/Experimental/MeshModelingToolsetExp/Source/GeometryProcessingAdapters/Private/CombineMeshInstancesImpl.cpp b/Engine/Plugins/Experimental/MeshModelingToolsetExp/Source/GeometryProcessingAdapters/Private/CombineMeshInstancesImpl.cpp
--- a/Engine/Plugins/Experimental/MeshModelingToolsetExp/Source/GeometryProcessingAdapters/Private/CombineMeshInstancesImpl.cpp
+++ b/Engine/Plugins/Experimental/MeshModelingToolsetExp/Source/GeometryProcessingAdapters/Private/CombineMeshInstancesImpl.cpp
@@ -3280,16 +3280,22 @@ void BuildCombinedMesh(
 		}
 	}
 
-	//CombinedLOD0.Attributes()->SetNumPolygroupLayers(2);
-	//FDynamicMeshPolygroupAttribute* PartIDAttrib = AccumMesh.Attributes()->GetPolygroupLayer(0);
-	//FDynamicMeshPolygroupAttribute* PartInstanceMapAttrib = AccumMesh.Attributes()->GetPolygroupLayer(1);
-
 	int32 NumParts = Assembly.Parts.Num();
 
-	//for ( int32 SetIndex = 0; SetIndex < NumParts; ++SetIndex )
-	//{
-	//	CombinedLOD0.EnableMatchingAttributes( Assembly.SourceMeshGeometry[Index].OriginalMesh, false, false );
-	//}
+	// determine maximum number of UV channels on any Part LOD0, and configure the output combined
+	// LOD meshes to have that many UV channels (clamping to at least 1). Triangles from Parts that
+	// have fewer UV channels will end up with (0,0) UVs in the extra channels.
+	int32 MaxNumUVChannels = 1;
+	for (int32 SetIndex = 0; SetIndex < NumParts; ++SetIndex)
+	{
+		const FDynamicMesh3& SourceMesh = Assembly.SourceMeshGeometry[SetIndex].SourceMeshLODs[0];
+		MaxNumUVChannels = FMath::Max(MaxNumUVChannels, (SourceMesh.HasAttributes() ? SourceMesh.Attributes()->NumUVLayers() : 0));
+	}
+	for (FCombinedMeshLOD& LODMeshData : MeshLODs)
+	{
+		LODMeshData.Mesh.Attributes()->SetNumUVLayers(MaxNumUVChannels);
+	}
+
 
 	// determine if we have multiple part subsets. In this case we need to be able to split the mesh
 	// by part later, which we will do by appending a polygroup layer
