Fixing crash mirroring weights

diff --git a/Engine/Plugins/Experimental/MeshModelingToolsetExp/Source/MeshModelingToolsEditorOnlyExp/Private/SkeletalMesh/SkinWeightsPaintTool.cpp b/Engine/Plugins/Experimental/MeshModelingToolsetExp/Source/MeshModelingToolsEditorOnlyExp/Private/SkeletalMesh/SkinWeightsPaintTool.cpp
--- a/Engine/Plugins/Experimental/MeshModelingToolsetExp/Source/MeshModelingToolsEditorOnlyExp/Private/SkeletalMesh/SkinWeightsPaintTool.cpp
+++ b/Engine/Plugins/Experimental/MeshModelingToolsetExp/Source/MeshModelingToolsEditorOnlyExp/Private/SkeletalMesh/SkinWeightsPaintTool.cpp
@@ -1456,8 +1456,7 @@ void FSkinMirrorData::RegenerateMirrorData(
 	const FReferenceSkeleton& RefSkeleton,
 	const TArray<FVector>& RefPoseVertices,
 	EAxis::Type InMirrorAxis,
-	EMirrorDirection InMirrorDirection,
-	const float MaxSearchRadius)
+	EMirrorDirection InMirrorDirection)
 {
 	if (bIsInitialized && InMirrorAxis == Axis && InMirrorDirection==Direction)
 	{
@@ -1533,7 +1532,7 @@ void FSkinMirrorData::RegenerateMirrorData(
 
 			// forcibly break out if search radius gets bigger than the mesh bounds.
 			// this could potentially happen if mesh is highly non-symmetrical along mirror axis
-			if (SearchRadius >= MaxSearchRadius)
+			if (SearchRadius >= HashGridCellSize)
 			{
 				break;
 			}
@@ -1582,8 +1581,7 @@ void USkinWeightsPaintTool::MirrorWeights(EAxis::Type Axis, EMirrorDirection Dir
 		RefSkeleton,
 		RefPoseVertices,
 		Axis,
-		Direction,
-		Weights.Deformer.Component->Bounds.SphereRadius);
+		Direction);
 
 	// get a reference to the mirror tables
 	const TMap<int32, int32>& BoneMap = MirrorData.GetBoneMap();
diff --git a/Engine/Plugins/Experimental/MeshModelingToolsetExp/Source/MeshModelingToolsEditorOnlyExp/Public/SkeletalMesh/SkinWeightsPaintTool.h b/Engine/Plugins/Experimental/MeshModelingToolsetExp/Source/MeshModelingToolsEditorOnlyExp/Public/SkeletalMesh/SkinWeightsPaintTool.h
--- a/Engine/Plugins/Experimental/MeshModelingToolsetExp/Source/MeshModelingToolsEditorOnlyExp/Public/SkeletalMesh/SkinWeightsPaintTool.h
+++ b/Engine/Plugins/Experimental/MeshModelingToolsetExp/Source/MeshModelingToolsEditorOnlyExp/Public/SkeletalMesh/SkinWeightsPaintTool.h
@@ -212,8 +212,7 @@ namespace SkinPaintTool
 			const FReferenceSkeleton& RefSkeleton,
 			const TArray<FVector>& RefPoseVertices,
 			EAxis::Type InMirrorAxis,
-			EMirrorDirection InMirrorDirection,
-			const float MaxSearchRadius);
+			EMirrorDirection InMirrorDirection);
 
 		const TMap<int32, int32>& GetBoneMap() const { return BoneMap; };
 		const TMap<int32, int32>& GetVertexMap() const { return VertexMap; };
