Skeleton Modifier: allow renaming bones with names starting with an integer

#rnx

diff --git a/Engine/Plugins/Experimental/MeshModelingToolsetExp/Source/SkeletalMeshModifiers/Private/SkeletonModifier.cpp b/Engine/Plugins/Experimental/MeshModelingToolsetExp/Source/SkeletalMeshModifiers/Private/SkeletonModifier.cpp
--- a/Engine/Plugins/Experimental/MeshModelingToolsetExp/Source/SkeletalMeshModifiers/Private/SkeletonModifier.cpp
+++ b/Engine/Plugins/Experimental/MeshModelingToolsetExp/Source/SkeletalMeshModifiers/Private/SkeletonModifier.cpp
@@ -1282,7 +1282,7 @@ FName USkeletonModifier::GetUniqueName(const FName InBoneName, const TArray<FNam
 				LastHashtag = Index;
 			}
 			
-			const bool bGoodChar =	bIsAlphaOrUnderscore || bIsHashtag || ((Index > 0) && bIsDigit);
+			const bool bGoodChar = bIsAlphaOrUnderscore || bIsHashtag || bIsDigit;
 			bHasAnyGoodChar |= bGoodChar;
 			if (!bGoodChar)
 			{
@@ -1357,7 +1357,7 @@ FName USkeletonModifier::GetUniqueName(const FName InBoneName, const TArray<FNam
 
 	if (StartPadding != INDEX_NONE)
 	{
-		CurrentIndex = FCString::Atoi(*PaddingStr);
+		CurrentIndex = PaddingStr.IsEmpty() ? INDEX_NONE : FCString::Atoi(*PaddingStr);
 		if (CurrentIndex == 0)
 		{
 			PaddingStr[PaddingStr.Len()-1] = '1';
