Modeling Mode: Fix issue with the Revolve Spline tool where changes to the revolution axis weren't being updated live in the preview geometry.

diff --git a/Engine/Plugins/Experimental/MeshModelingToolsetExp/Source/MeshModelingToolsExp/Private/RevolveSplineTool.cpp b/Engine/Plugins/Experimental/MeshModelingToolsetExp/Source/MeshModelingToolsExp/Private/RevolveSplineTool.cpp
--- a/Engine/Plugins/Experimental/MeshModelingToolsetExp/Source/MeshModelingToolsExp/Private/RevolveSplineTool.cpp
+++ b/Engine/Plugins/Experimental/MeshModelingToolsetExp/Source/MeshModelingToolsExp/Private/RevolveSplineTool.cpp
@@ -286,13 +286,12 @@ void URevolveSplineTool::OnPropertyModified(UObject* PropertySet, FProperty* Pro
 		{
 			UpdatePointsFromSpline();
 		}
-		else if (Property->GetFName() == GET_MEMBER_NAME_CHECKED(URevolveSplineToolProperties, AxisOrigin)
-			|| Property->GetFName() == GET_MEMBER_NAME_CHECKED(URevolveSplineToolProperties, AxisOrientation))
-		{
-			PlaneMechanic->SetPlaneWithoutBroadcast(FFrame3d(Settings->AxisOrigin,
-				FRotator(Settings->AxisOrientation.X, Settings->AxisOrientation.Y, 0).Quaternion()));
-			UpdateRevolutionAxis();
-		}
+
+		// Checking the name for these settings doesn't work, since the reported names are the low level components, like "X" or "Y"
+		// So we'll simply update the axis whenever any property changes. It's overkill but probably not too bad.
+		PlaneMechanic->SetPlaneWithoutBroadcast(FFrame3d(Settings->AxisOrigin,
+			FRotator(Settings->AxisOrientation.X, Settings->AxisOrientation.Y, 0).Quaternion()));
+		UpdateRevolutionAxis();		
 	}
 
 	Preview->PreviewMesh->EnableWireframe(MaterialProperties->bShowWireframe);
