Fixed skin weight color not updating in viewport.

diff --git a/Engine/Plugins/Experimental/MeshModelingToolsetExp/Source/MeshModelingToolsEditorOnlyExp/Private/SkeletalMesh/SkinWeightsPaintTool.cpp b/Engine/Plugins/Experimental/MeshModelingToolsetExp/Source/MeshModelingToolsEditorOnlyExp/Private/SkeletalMesh/SkinWeightsPaintTool.cpp
--- a/Engine/Plugins/Experimental/MeshModelingToolsetExp/Source/MeshModelingToolsEditorOnlyExp/Private/SkeletalMesh/SkinWeightsPaintTool.cpp
+++ b/Engine/Plugins/Experimental/MeshModelingToolsetExp/Source/MeshModelingToolsEditorOnlyExp/Private/SkeletalMesh/SkinWeightsPaintTool.cpp
@@ -1027,7 +1027,7 @@ FVector4f USkinWeightsPaintTool::WeightToColor(float Value) const
 	// optional greyscale mode
 	if (WeightToolProperties->ColorMode == EWeightColorMode::MinMax)
 	{
-		return FLinearColor::LerpUsingHSV(WeightToolProperties->MinColor, WeightToolProperties->MaxColor, Value);	
+		return FMath::Lerp(WeightToolProperties->MinColor, WeightToolProperties->MaxColor, Value);
 	}
 	
 	// early out zero weights to min color
@@ -1048,7 +1048,7 @@ FVector4f USkinWeightsPaintTool::WeightToColor(float Value) const
 	// revert back to simple Lerp(min,max) if user supplied color ramp doesn't have enough colors
 	if (Colors.Num() < 2)
 	{
-		const FLinearColor FinalColor = FLinearColor::LerpUsingHSV(WeightToolProperties->MinColor, WeightToolProperties->MaxColor, Value);
+		const FLinearColor FinalColor = FMath::Lerp(WeightToolProperties->MinColor, WeightToolProperties->MaxColor, Value);
 		return UE::Geometry::ToVector4<float>(FinalColor);
 	}
 
@@ -1061,7 +1061,7 @@ FVector4f USkinWeightsPaintTool::WeightToColor(float Value) const
 	const float Param = (Value - RangeStart) / (RangeEnd - RangeStart);
 	const FLinearColor& StartColor = Colors[ColorIndex];
 	const FLinearColor& EndColor = Colors[ColorIndex+1];
-	const FLinearColor FinalColor = FLinearColor::LerpUsingHSV(StartColor, EndColor, Param);
+	const FLinearColor FinalColor = FMath::Lerp(StartColor, EndColor, Param);
 	return UE::Geometry::ToVector4<float>(FinalColor);
 }
 
