Mesh Modeling.  Clamp values in the UI for simplification percentage.

diff --git a/Engine/Plugins/Experimental/MeshModelingToolsetExp/Source/MeshModelingToolsEditorOnlyExp/Public/SimplifyMeshTool.h b/Engine/Plugins/Experimental/MeshModelingToolsetExp/Source/MeshModelingToolsEditorOnlyExp/Public/SimplifyMeshTool.h
--- a/Engine/Plugins/Experimental/MeshModelingToolsetExp/Source/MeshModelingToolsEditorOnlyExp/Public/SimplifyMeshTool.h
+++ b/Engine/Plugins/Experimental/MeshModelingToolsetExp/Source/MeshModelingToolsEditorOnlyExp/Public/SimplifyMeshTool.h
@@ -48,7 +48,7 @@ public:
 	ESimplifyTargetType TargetMode;
 
 	/** Target percentage of original triangle count */
-	UPROPERTY(EditAnywhere, Category = Options, meta = (UIMin = "0", UIMax = "100",
+	UPROPERTY(EditAnywhere, Category = Options, meta = (UIMin = "0", UIMax = "100", ClampMin = "0.0", ClampMax = "100.0", 
 				EditCondition = "SimplifierType != ESimplifyType::MinimalPolygroup && SimplifierType != ESimplifyType::MinimalPlanar && TargetMode == ESimplifyTargetType::Percentage"))
 	int TargetPercentage;
 
