ModelingMode: add new Harvest Instances Tool, which can be used to gather static mesh components/instances from multiple Actors into new InstancedStaticMeshComponents under one or more new Actors.

diff --git a/Engine/Plugins/Editor/ModelingToolsEditorMode/Source/ModelingToolsEditorMode/Private/ModelingToolsEditorMode.cpp b/Engine/Plugins/Editor/ModelingToolsEditorMode/Source/ModelingToolsEditorMode/Private/ModelingToolsEditorMode.cpp
--- a/Engine/Plugins/Editor/ModelingToolsEditorMode/Source/ModelingToolsEditorMode/Private/ModelingToolsEditorMode.cpp
+++ b/Engine/Plugins/Editor/ModelingToolsEditorMode/Source/ModelingToolsEditorMode/Private/ModelingToolsEditorMode.cpp
@@ -93,6 +93,7 @@
 #include "ConvertMeshesTool.h"
 #include "SplitMeshesTool.h"
 #include "PatternTool.h"
+#include "HarvestInstancesTool.h"
 #include "TriangulateSplinesTool.h"
 
 #include "Polymodeling/ExtrudeMeshSelectionTool.h"
@@ -650,6 +651,7 @@ void UModelingToolsEditorMode::Enter()
 	UPatternToolBuilder* PatternToolBuilder = NewObject<UPatternToolBuilder>();
 	PatternToolBuilder->bEnableCreateISMCs = true;
 	RegisterTool(ToolManagerCommands.BeginPatternTool, TEXT("BeginPatternTool"), PatternToolBuilder);
+	RegisterTool(ToolManagerCommands.BeginHarvestInstancesTool, TEXT("BeginHarvestInstancesTool"), NewObject<UHarvestInstancesToolBuilder>());
 
 	UCombineMeshesToolBuilder* CombineMeshesToolBuilder = NewObject<UCombineMeshesToolBuilder>();
 	RegisterTool(ToolManagerCommands.BeginCombineMeshesTool, TEXT("BeginCombineMeshesTool"), CombineMeshesToolBuilder);
diff --git a/Engine/Plugins/Editor/ModelingToolsEditorMode/Source/ModelingToolsEditorMode/Private/ModelingToolsEditorModeToolkit.cpp b/Engine/Plugins/Editor/ModelingToolsEditorMode/Source/ModelingToolsEditorMode/Private/ModelingToolsEditorModeToolkit.cpp
--- a/Engine/Plugins/Editor/ModelingToolsEditorMode/Source/ModelingToolsEditorMode/Private/ModelingToolsEditorModeToolkit.cpp
+++ b/Engine/Plugins/Editor/ModelingToolsEditorMode/Source/ModelingToolsEditorMode/Private/ModelingToolsEditorModeToolkit.cpp
@@ -309,6 +309,7 @@ void FModelingToolsEditorModeToolkit::RegisterPalettes()
 		Commands.BeginSplitMeshesTool,
 		Commands.BeginPatternTool,
 		
+		Commands.BeginHarvestInstancesTool,
 		Commands.BeginISMEditorTool
 	});
 	ToolkitBuilder->AddPalette(
diff --git a/Engine/Plugins/Editor/ModelingToolsEditorMode/Source/ModelingToolsEditorMode/Private/ModelingToolsManagerActions.cpp b/Engine/Plugins/Editor/ModelingToolsEditorMode/Source/ModelingToolsEditorMode/Private/ModelingToolsManagerActions.cpp
--- a/Engine/Plugins/Editor/ModelingToolsEditorMode/Source/ModelingToolsEditorMode/Private/ModelingToolsManagerActions.cpp
+++ b/Engine/Plugins/Editor/ModelingToolsEditorMode/Source/ModelingToolsEditorMode/Private/ModelingToolsManagerActions.cpp
@@ -176,6 +176,8 @@ void FModelingToolsManagerCommands::RegisterCommands()
 	REGISTER_MODELING_TOOL_COMMAND(BeginConvertMeshesTool, "Convert", "Convert the selected Meshes to a different type of Mesh Object");
 	REGISTER_MODELING_TOOL_COMMAND(BeginSplitMeshesTool, "Split", "Split the selected Meshes into separate parts based on connectivity");
 	REGISTER_MODELING_TOOL_COMMAND(BeginPatternTool, "Pattern", "Create patterns of Meshes");
+	REGISTER_MODELING_TOOL_COMMAND(BeginHarvestInstancesTool, "Harvest Instances", "Extract a set of InstancedStaticMeshComponents from a set of Actors");
+	
 
 	// MeshOps
 	if ( bUseShortNames )
diff --git a/Engine/Plugins/Editor/ModelingToolsEditorMode/Source/ModelingToolsEditorMode/Public/ModelingToolsManagerActions.h b/Engine/Plugins/Editor/ModelingToolsEditorMode/Source/ModelingToolsEditorMode/Public/ModelingToolsManagerActions.h
--- a/Engine/Plugins/Editor/ModelingToolsEditorMode/Source/ModelingToolsEditorMode/Public/ModelingToolsManagerActions.h
+++ b/Engine/Plugins/Editor/ModelingToolsEditorMode/Source/ModelingToolsEditorMode/Public/ModelingToolsManagerActions.h
@@ -79,6 +79,7 @@ public:
 	TSharedPtr<FUICommandInfo> BeginConvertMeshesTool;
 	TSharedPtr<FUICommandInfo> BeginSplitMeshesTool;
 	TSharedPtr<FUICommandInfo> BeginPatternTool;
+	TSharedPtr<FUICommandInfo> BeginHarvestInstancesTool;
 
 
 	TSharedPtr<FUICommandInfo> BeginRemeshSculptMeshTool;
diff --git a/Engine/Plugins/Experimental/MeshModelingToolsetExp/Source/MeshModelingToolsEditorOnlyExp/Private/HarvestInstancesTool.cpp b/Engine/Plugins/Experimental/MeshModelingToolsetExp/Source/MeshModelingToolsEditorOnlyExp/Private/HarvestInstancesTool.cpp
new file mode 100644
--- /dev/null
+++ b/Engine/Plugins/Experimental/MeshModelingToolsetExp/Source/MeshModelingToolsEditorOnlyExp/Private/HarvestInstancesTool.cpp
@@ -0,0 +1,495 @@
+// Copyright Epic Games, Inc. All Rights Reserved.
+
+#include "HarvestInstancesTool.h"
+#include "InteractiveToolManager.h"
+#include "ToolTargetManager.h"
+#include "TargetInterfaces/PrimitiveComponentBackedTarget.h"
+#include "TargetInterfaces/AssetBackedTarget.h"
+#include "ToolBuilderUtil.h"
+#include "ToolSetupUtil.h"
+#include "ModelingToolTargetUtil.h"
+#include "Selection/ToolSelectionUtil.h"
+#include "ModelingObjectsCreationAPI.h"
+#include "ToolDataVisualizer.h"
+#include "Util/ColorConstants.h"
+
+
+#include "Components/PrimitiveComponent.h"
+#include "Components/StaticMeshComponent.h"
+#include "Components/InstancedStaticMeshComponent.h"
+#include "Components/HierarchicalInstancedStaticMeshComponent.h"
+#include "Engine/StaticMesh.h"
+#include "Engine/World.h"
+#include "Editor/EditorEngine.h"		// for FActorLabelUtilities
+
+
+using namespace UE::Geometry;
+
+#define LOCTEXT_NAMESPACE "UHarvestInstancesTool"
+
+
+/*
+ * ToolBuilder
+ */
+
+bool UHarvestInstancesToolBuilder::CanBuildTool(const FToolBuilderState& SceneState) const
+{
+	int32 ValidTargets = 0;
+
+	// only static mesh components currently supported
+	SceneState.TargetManager->EnumerateSelectedAndTargetableComponents(SceneState, GetTargetRequirements(), [&ValidTargets](UActorComponent* Component)
+	{
+		if (Cast<UStaticMeshComponent>(Component))
+		{
+			ValidTargets++;
+		}
+	});
+	
+	return ValidTargets > 0;
+}
+
+UMultiSelectionMeshEditingTool* UHarvestInstancesToolBuilder::CreateNewTool(const FToolBuilderState& SceneState) const
+{
+	UHarvestInstancesTool* NewTool = NewObject<UHarvestInstancesTool>(SceneState.ToolManager);
+	return NewTool;
+}
+
+void UHarvestInstancesToolBuilder::InitializeNewTool(UMultiSelectionMeshEditingTool* NewTool, const FToolBuilderState& SceneState) const
+{
+	TArray<TObjectPtr<UToolTarget>> Targets;
+	SceneState.TargetManager->EnumerateSelectedAndTargetableComponents(SceneState, GetTargetRequirements(), [this, &Targets, &SceneState](UActorComponent* Component)
+	{
+		if (Cast<UStaticMeshComponent>(Component))
+		{
+			Targets.Add(SceneState.TargetManager->BuildTarget(Component, GetTargetRequirements()));
+		}
+	});
+
+	NewTool->SetTargets(Targets);
+	NewTool->SetWorld(SceneState.World);
+}
+
+const FToolTargetTypeRequirements& UHarvestInstancesToolBuilder::GetTargetRequirements() const
+{
+	static FToolTargetTypeRequirements TypeRequirements({
+		UPrimitiveComponentBackedTarget::StaticClass(),
+		UAssetBackedTarget::StaticClass()
+		});
+	return TypeRequirements;
+}
+
+
+/*
+ * 
+ */
+
+UHarvestInstancesTool::UHarvestInstancesTool()
+{
+}
+
+
+void UHarvestInstancesTool::Setup()
+{
+	UInteractiveTool::Setup();
+
+	OutputSettings = NewObject<UHarvestInstancesTool_OutputSettings>();
+	AddToolPropertySource(OutputSettings);
+	OutputSettings->RestoreProperties(this);
+	OutputSettings->WatchProperty(OutputSettings->bIncludeSingleInstances, [this](bool bNewValue) { UpdateInstanceSets(); });
+
+	Settings = NewObject<UHarvestInstancesToolSettings>();
+	AddToolPropertySource(Settings);
+	Settings->RestoreProperties(this);
+
+	InitializeInstanceSets();
+	UpdateInstanceSets();
+
+	SetToolDisplayName(LOCTEXT("ToolName", "Harvest Instances"));
+	GetToolManager()->DisplayMessage(
+		LOCTEXT("OnStartHarvestInstancesTool", "Harvest repeated StaticMesh Components from selected Actors into new InstancedStaticMeshComponents"),
+		EToolMessageLevel::UserNotification);
+
+	UpdateWarningMessages();
+}
+
+
+
+void UHarvestInstancesTool::OnShutdown(EToolShutdownType ShutdownType)
+{
+	Settings->SaveProperties(this);
+	OutputSettings->SaveProperties(this);
+
+	if (ShutdownType == EToolShutdownType::Accept)
+	{
+		EmitResults();
+	}
+}
+
+
+
+
+void UHarvestInstancesTool::Render(IToolsContextRenderAPI* RenderAPI)
+{
+	FToolDataVisualizer LineRenderer;
+	LineRenderer.BeginFrame(RenderAPI);
+
+	if (Settings->bDrawBounds)
+	{
+		int Counter = 0;
+		for (const FInstanceSet& InstanceSet : InstanceSets)
+		{
+			FColor SetColor = LinearColors::SelectFColor(Counter++);
+			LineRenderer.SetLineParameters(FLinearColor(0.95, 0.05, 0.05), 4.0);
+			FBox Bounds = InstanceSet.Bounds;
+			Bounds.Min -= 5.0 * FVector::One();
+			Bounds.Max += 5.0 * FVector::One();
+
+			for (const FTransform& Transform : InstanceSet.WorldInstanceTransforms)
+			{
+				LineRenderer.PushTransform(Transform);
+				LineRenderer.DrawWireBox(Bounds, SetColor, 2.0, true);
+				LineRenderer.PopTransform();
+			}
+		}
+	}
+
+	LineRenderer.EndFrame();
+}
+
+
+
+void UHarvestInstancesTool::InitializeInstanceSets()
+{
+	// collect up source meshes
+	this->SourceMeshes.Reset();
+
+	int32 NumTargets = Targets.Num();
+	for (int32 TargetIdx = 0; TargetIdx < NumTargets; TargetIdx++)
+	{
+		UPrimitiveComponent* TargetComponent = UE::ToolTarget::GetTargetComponent(Targets[TargetIdx]);
+		if ( UStaticMeshComponent* StaticMeshComponent = Cast<UStaticMeshComponent>(TargetComponent) )
+		{
+			UStaticMesh* StaticMesh = StaticMeshComponent->GetStaticMesh();
+			if (!StaticMesh) continue;
+			TArray<UMaterialInterface*> InstanceMaterials = UE::ToolTarget::GetMaterialSet(Targets[TargetIdx], false).Materials;
+
+			FSourceMesh* UseSourceMesh = nullptr;
+			for (int32 k = 0; k < SourceMeshes.Num() && UseSourceMesh == nullptr; ++k)
+			{
+				if (SourceMeshes[k].Mesh == StaticMesh && SourceMeshes[k].InstanceMaterials == InstanceMaterials)
+				{
+					UseSourceMesh = &SourceMeshes[k];
+				}
+			}
+			if (UseSourceMesh == nullptr)
+			{
+				FSourceMesh NewSourceMesh;
+				NewSourceMesh.Mesh = StaticMesh;
+				NewSourceMesh.InstanceMaterials = InstanceMaterials;
+				SourceMeshes.Add(NewSourceMesh);
+				UseSourceMesh = &SourceMeshes.Last();
+			}
+
+			UseSourceMesh->SourceComponents.Add(StaticMeshComponent);
+
+			if (UInstancedStaticMeshComponent* ISMComponent = Cast<UInstancedStaticMeshComponent>(StaticMeshComponent))
+			{
+				int32 NumInstances = ISMComponent->GetInstanceCount();
+				for (int32 i = 0; i < NumInstances; ++i)
+				{
+					if (ISMComponent->IsValidInstance(i))
+					{
+						FTransform InstanceWorldTransform;
+						if (ensure(ISMComponent->GetInstanceTransform(i, InstanceWorldTransform, true)))
+						{
+							UseSourceMesh->InstanceTransformsWorld.Add(InstanceWorldTransform);
+						}
+					}
+				}
+			}
+			else
+			{
+				FTransform WorldTransform = UE::ToolTarget::GetLocalToWorldTransform(Targets[TargetIdx]);
+				UseSourceMesh->InstanceTransformsWorld.Add(WorldTransform);
+			}
+		}
+	}
+
+	// determine if we have 'loner' instances
+	for (int32 k = 0; k < SourceMeshes.Num(); ++k)
+	{
+		if (SourceMeshes[k].InstanceTransformsWorld.Num() == 1)
+		{
+			OutputSettings->bHaveLonerInstances = true;
+		}
+	}
+}
+
+
+
+void UHarvestInstancesTool::UpdateInstanceSets()
+{
+	TArray<FSourceMesh> FilteredSourceMeshes = SourceMeshes;
+
+	// remove any meshes that only have one instance, if desired
+	if (OutputSettings->bIncludeSingleInstances == false)
+	{
+		for (int32 k = 0; k < FilteredSourceMeshes.Num(); ++k)
+		{
+			if (FilteredSourceMeshes[k].InstanceTransformsWorld.Num() <= 1)
+			{
+				FilteredSourceMeshes.RemoveAt(k);
+				k--;
+			}
+		}
+	}
+
+	InstanceSets.Reset();
+	SourceActors.Reset();
+	SourceComponents.Reset();
+
+	int32 NumInstanceSets = FilteredSourceMeshes.Num();
+	OutputSettings->bHaveSingleInstanceSet = (NumInstanceSets == 1);
+
+	InstanceSets.SetNum(NumInstanceSets);
+	for (int32 Index = 0; Index < NumInstanceSets; Index++)
+	{
+		FSourceMesh& SourceMesh = FilteredSourceMeshes[Index];
+
+		FInstanceSet& Element = InstanceSets[Index];
+
+		Element.SourceComponents = SourceMesh.SourceComponents;
+		for (UStaticMeshComponent* Component : Element.SourceComponents)
+		{
+			Element.SourceActors.AddUnique(Component->GetOwner());
+			SourceActors.AddUnique(Component->GetOwner());
+			SourceComponents.Add(Component);
+		}
+
+		Element.StaticMesh = SourceMesh.Mesh;
+		Element.Bounds = SourceMesh.Mesh->GetBoundingBox();
+		Element.InstanceMaterials = SourceMesh.InstanceMaterials;
+		Element.WorldInstanceTransforms = SourceMesh.InstanceTransformsWorld;
+
+		// do this?
+		//bHaveNonUniformScaleElements = bHaveNonUniformScaleElements || Element.BaseRotateScale.HasNonUniformScale();
+	}
+
+
+	// check that all Components of Source Actors are included in instance set. If not, we cannot delete
+	// the Source Actors
+	bool bFoundUnhandledComponent = false;
+	for (AActor* SourceActor : SourceActors)
+	{
+		SourceActor->ForEachComponent(true, [&](UActorComponent* Component) 
+		{
+			if (bFoundUnhandledComponent) return;
+
+			if (UPrimitiveComponent* PrimComponent = Cast<UPrimitiveComponent>(Component))
+			{
+				if (SourceComponents.Contains(PrimComponent))
+				{
+					return;
+				}
+			}
+			bFoundUnhandledComponent = true;
+		});
+	}
+	OutputSettings->bCanDeleteInputs = (bFoundUnhandledComponent == false);
+	UpdateWarningMessages();
+}
+
+
+
+void UHarvestInstancesTool::UpdateWarningMessages()
+{
+	FTextBuilder WarningTextBuilder;
+	if (OutputSettings->bCanDeleteInputs == false)
+	{
+		WarningTextBuilder.AppendLine(
+			LOCTEXT("CannotDeleteWarning", "Source Actors cannot be deleted because they contain additional Components that cannot be Instanced. "));
+	}
+	if (bHaveNonUniformScaleElements)
+	{
+		WarningTextBuilder.AppendLine(
+			LOCTEXT("NonUniformScaleWarning", "Source Objects have Non-Uniform Scaling, which may prevent Instance Transforms from working correctly."));
+	}
+
+	GetToolManager()->DisplayMessage(WarningTextBuilder.ToText(), EToolMessageLevel::UserWarning);
+}
+
+
+void UHarvestInstancesTool::OnPropertyModified(UObject* PropertySet, FProperty* Property)
+{
+	if (PropertySet == Settings || PropertySet == OutputSettings)
+	{
+		return;
+	}
+
+	OnParametersUpdated();
+}
+
+
+void UHarvestInstancesTool::OnParametersUpdated()
+{
+	//MarkPatternDirty();
+}
+
+
+
+
+void UHarvestInstancesTool::EmitResults()
+{
+	GetToolManager()->BeginUndoTransaction(LOCTEXT("HarvestInstances", "Harvest Instances"));
+
+	bool bCreateHISMs = (OutputSettings->ComponentType == EHarvestInstancesToolOutputType::HISMC);
+	int32 NumInstanceSets = InstanceSets.Num();
+
+	TArray<AActor*> NewActors;		// set of new actors created by operation
+
+	// TODO: investigate use of CopyPropertiesForUnrelatedObjects to transfer settings from source to target Components/Actors
+
+	auto ComputeComponentOriginFunc = [](FInstanceSet& Element) -> FVector
+	{
+		FVector ComponentPosition = FVector::Zero();
+		for (const FTransform& WorldTransform : Element.WorldInstanceTransforms)
+		{
+			ComponentPosition += WorldTransform.GetLocation();
+		}
+		ComponentPosition /= Element.WorldInstanceTransforms.Num();
+		return ComponentPosition;
+	};
+
+	auto CreateNewComponentFunc = [bCreateHISMs](FInstanceSet& Element, AActor* ParentActor, FString BaseName) -> UInstancedStaticMeshComponent*
+	{
+		TSubclassOf<UInstancedStaticMeshComponent> ComponentClass = (bCreateHISMs) ?
+			UHierarchicalInstancedStaticMeshComponent::StaticClass() : UInstancedStaticMeshComponent::StaticClass();
+		FString Suffix = (bCreateHISMs) ? TEXT("_HISM") : TEXT("_ISM");
+
+		FName UseName = (BaseName.Len() > 0) ? FName(BaseName + Suffix) : FName();
+
+		UInstancedStaticMeshComponent* ISMComponent = NewObject<UInstancedStaticMeshComponent>(ParentActor, ComponentClass,
+			MakeUniqueObjectName(ParentActor, ComponentClass, UseName));
+
+		ISMComponent->SetFlags(RF_Transactional);
+		ISMComponent->bHasPerInstanceHitProxies = true;
+
+		ISMComponent->SetStaticMesh(Element.StaticMesh);
+		for (int32 j = 0; j < Element.InstanceMaterials.Num(); ++j)
+		{
+			ISMComponent->SetMaterial(j, Element.InstanceMaterials[j]);
+		}
+		return ISMComponent;
+	};
+
+
+	AActor* SingleActor = nullptr;
+	if (OutputSettings->bSingleActor || NumInstanceSets == 1)
+	{
+		FActorSpawnParameters SpawnInfo;
+		SingleActor = GetTargetWorld()->SpawnActor<AActor>(SpawnInfo);
+		if ( SingleActor == nullptr )
+		{
+			return;		// ??
+		}
+
+		FString UseBaseName = (OutputSettings->ActorName.Len() > 0) ? OutputSettings->ActorName : TEXT("Instances");
+		FActorLabelUtilities::SetActorLabelUnique(SingleActor, UseBaseName);
+
+		NewActors.Add(SingleActor);
+		bool bFirst = true;
+
+		for (int32 ElemIdx = 0; ElemIdx < NumInstanceSets; ++ElemIdx)
+		{
+			FInstanceSet& Element = InstanceSets[ElemIdx];
+			FString AssetBaseName = FPaths::GetBaseFilename(Element.StaticMesh->GetName());
+			AssetBaseName = UE::Modeling::StripGeneratedAssetSuffixFromName(AssetBaseName);
+
+			FVector ComponentPosition = ComputeComponentOriginFunc(Element);
+			UInstancedStaticMeshComponent* ISMComponent = CreateNewComponentFunc(Element, SingleActor, AssetBaseName);
+
+			if (bFirst)
+			{
+				SingleActor->SetRootComponent(ISMComponent);
+				ISMComponent->OnComponentCreated();
+				SingleActor->AddInstanceComponent(ISMComponent);
+				SingleActor->SetActorTransform(FTransform(ComponentPosition));
+
+				bFirst = false;
+			}
+			else
+			{
+				ISMComponent->OnComponentCreated();
+				ISMComponent->SetWorldTransform(FTransform(ComponentPosition));
+				ISMComponent->AttachToComponent(SingleActor->GetRootComponent(), FAttachmentTransformRules::KeepRelativeTransform);
+				SingleActor->AddInstanceComponent(ISMComponent);
+			}
+
+			for (FTransform WorldTransform : Element.WorldInstanceTransforms)
+			{
+				ISMComponent->AddInstance(WorldTransform, /*bTransformInWorldSpace=*/true);
+			}
+
+			ISMComponent->RegisterComponent();
+		}
+	}
+	else
+	{
+		for (int32 ElemIdx = 0; ElemIdx < NumInstanceSets; ++ElemIdx)
+		{
+			FInstanceSet& Element = InstanceSets[ElemIdx];
+			FVector ComponentPosition = ComputeComponentOriginFunc(Element);
+
+			FActorSpawnParameters SpawnInfo;
+			AActor* NewActor = GetTargetWorld()->SpawnActor<AActor>(SpawnInfo);
+			if (NewActor != nullptr)
+			{
+				FString AssetName = FPaths::GetBaseFilename(Element.StaticMesh->GetName());
+				FString UseBaseName = UE::Modeling::StripGeneratedAssetSuffixFromName(AssetName);
+				UseBaseName += TEXT("Instances");
+				FActorLabelUtilities::SetActorLabelUnique(NewActor, UseBaseName);
+
+				NewActors.Add(NewActor);
+
+				FString AssetBaseName = FPaths::GetBaseFilename(Element.StaticMesh->GetName());
+				AssetBaseName = UE::Modeling::StripGeneratedAssetSuffixFromName(AssetBaseName);
+				UInstancedStaticMeshComponent* ISMComponent = CreateNewComponentFunc(Element, NewActor, AssetBaseName);
+
+				NewActor->SetRootComponent(ISMComponent);
+				ISMComponent->OnComponentCreated();
+				NewActor->AddInstanceComponent(ISMComponent);
+
+				NewActor->SetActorTransform(FTransform(ComponentPosition));
+
+				for (FTransform WorldTransform : Element.WorldInstanceTransforms)
+				{
+					ISMComponent->AddInstance(WorldTransform, /*bTransformInWorldSpace=*/true);
+				}
+
+				ISMComponent->RegisterComponent();
+			}
+		}
+	}
+
+	// delete all source actors if desired
+	if (OutputSettings->bCanDeleteInputs && OutputSettings->bDeleteInputs)
+	{
+		// destroy any Actors that won't have any Components left
+		for (AActor* Actor : SourceActors)
+		{
+			Actor->Modify();
+			Actor->Destroy();
+		}
+	}
+
+	if (NewActors.Num() > 0)
+	{
+		ToolSelectionUtil::SetNewActorSelection(GetToolManager(), NewActors);
+	}
+
+	GetToolManager()->EndUndoTransaction();
+}
+
+
+
+#undef LOCTEXT_NAMESPACE
diff --git a/Engine/Plugins/Experimental/MeshModelingToolsetExp/Source/MeshModelingToolsEditorOnlyExp/Public/HarvestInstancesTool.h b/Engine/Plugins/Experimental/MeshModelingToolsetExp/Source/MeshModelingToolsEditorOnlyExp/Public/HarvestInstancesTool.h
new file mode 100644
--- /dev/null
+++ b/Engine/Plugins/Experimental/MeshModelingToolsetExp/Source/MeshModelingToolsEditorOnlyExp/Public/HarvestInstancesTool.h
@@ -0,0 +1,173 @@
+// Copyright Epic Games, Inc. All Rights Reserved.
+
+#pragma once
+
+#include "CoreMinimal.h"
+#include "BaseTools/MultiSelectionMeshEditingTool.h"
+#include "InteractiveToolBuilder.h"
+#include "BaseBehaviors/BehaviorTargetInterfaces.h"
+#include "FrameTypes.h"
+#include "BoxTypes.h"
+#include "TransformTypes.h"
+#include "HarvestInstancesTool.generated.h"
+
+class UStaticMesh;
+class UStaticMeshComponent;
+class UInstancedStaticMeshComponent;
+class UMaterialInterface;
+
+
+/**
+ *
+ */
+UCLASS()
+class MESHMODELINGTOOLSEDITORONLYEXP_API UHarvestInstancesToolBuilder : public UMultiSelectionMeshEditingToolBuilder
+{
+	GENERATED_BODY()
+
+public:
+	virtual bool CanBuildTool(const FToolBuilderState& SceneState) const override;
+	
+	virtual UMultiSelectionMeshEditingTool* CreateNewTool(const FToolBuilderState& SceneState) const override;
+
+	virtual void InitializeNewTool(UMultiSelectionMeshEditingTool* NewTool, const FToolBuilderState& SceneState) const override;
+	
+protected:
+	virtual const FToolTargetTypeRequirements& GetTargetRequirements() const override;
+};
+
+
+/**
+ * Settings for the Pattern Tool
+ */
+UCLASS()
+class MESHMODELINGTOOLSEDITORONLYEXP_API UHarvestInstancesToolSettings : public UInteractiveToolPropertySet
+{
+	GENERATED_BODY()
+public:
+
+	/** Display the bounding box of each mesh that will be moved to instances. Bounding box colors indicate instance groupings. */
+	UPROPERTY(EditAnywhere, Category = Visualization)
+	bool bDrawBounds = true;
+};
+
+
+
+UENUM()
+enum class EHarvestInstancesToolOutputType
+{
+	HISMC UMETA(DisplayName="Hierarchical (HISM)"),
+	ISMC UMETA(DisplayName="Non-Hierarchical (ISM)")
+};
+
+
+
+/**
+ * Output Settings for the Pattern Tool
+ */
+UCLASS()
+class MESHMODELINGTOOLSEDITORONLYEXP_API UHarvestInstancesTool_OutputSettings : public UInteractiveToolPropertySet
+{
+	GENERATED_BODY()
+public:
+	/** Which type of Instanced Static Mesh Component to Create */
+	UPROPERTY(EditAnywhere, Category = Output)
+	EHarvestInstancesToolOutputType ComponentType = EHarvestInstancesToolOutputType::HISMC;
+
+	/** Group all output Instanced Components under a single Actor. By default, multiple Actors will be created. */
+	UPROPERTY(EditAnywhere, Category = Output, meta = (EditCondition="bHaveSingleInstanceSet == false", HideEditConditionToggle) )
+	bool bSingleActor = false;
+
+	/** Base Name to use for the emitted single Actor */
+	UPROPERTY(EditAnywhere, Category = Output, meta=(TransientToolProperty, DisplayName="New Actor Name", EditCondition = "bHaveSingleInstanceSet || bSingleActor == true"))
+	FString ActorName = TEXT("Instances");
+
+	UPROPERTY(EditAnywhere, AdvancedDisplay, Category = Output, meta = (EditCondition = "bSingleActor == true && bHaveLonerInstances == true", HideEditConditionToggle))
+	bool bIncludeSingleInstances = false;
+
+	/** Delete Actors that have had their Components harvested */
+	UPROPERTY(EditAnywhere, Category = Output, meta = (EditCondition = "bCanDeleteInputs == true", HideEditConditionToggle))
+	bool bDeleteInputs = true;
+
+	UPROPERTY(meta=(TransientToolProperty))
+	bool bHaveSingleInstanceSet = false;
+
+	UPROPERTY(meta = (TransientToolProperty))
+	bool bHaveLonerInstances = false;
+
+	UPROPERTY(meta = (TransientToolProperty))
+	bool bCanDeleteInputs = true;
+};
+
+
+
+/**
+ * UHarvestInstancesTool 
+ */
+UCLASS()
+class MESHMODELINGTOOLSEDITORONLYEXP_API UHarvestInstancesTool : public UMultiSelectionMeshEditingTool
+{
+	GENERATED_BODY()
+
+public:
+	UHarvestInstancesTool();
+
+	virtual void Setup() override;
+	virtual void OnShutdown(EToolShutdownType ShutdownType) override;
+
+	virtual void Render(IToolsContextRenderAPI* RenderAPI) override;
+
+	virtual void OnPropertyModified(UObject* PropertySet, FProperty* Property) override;
+
+	virtual bool HasCancel() const override { return true; }
+	virtual bool HasAccept() const override { return true; }
+	virtual bool CanAccept() const override { return true; }
+
+public:
+	UPROPERTY()
+	TObjectPtr<UHarvestInstancesToolSettings> Settings;
+	
+	UPROPERTY()
+	TObjectPtr<UHarvestInstancesTool_OutputSettings> OutputSettings;
+
+protected:
+
+	void OnParametersUpdated();
+
+
+	struct FSourceMesh
+	{
+		TArray<UStaticMeshComponent*> SourceComponents;
+		UStaticMesh* Mesh;
+		TArray<UMaterialInterface*> InstanceMaterials;
+		TArray<FTransform> InstanceTransformsWorld;
+	};
+	TArray<FSourceMesh> SourceMeshes;
+
+	struct FInstanceSet
+	{
+		TArray<UStaticMeshComponent*> SourceComponents;
+		TArray<AActor*> SourceActors;
+
+		UStaticMesh* StaticMesh = nullptr;
+		FBox Bounds;
+		TArray<UMaterialInterface*> InstanceMaterials;
+
+		TArray<FTransform> WorldInstanceTransforms;
+	};
+	TArray<FInstanceSet> InstanceSets;
+
+	TArray<AActor*> SourceActors;
+	TArray<UPrimitiveComponent*> SourceComponents;
+
+	bool bHaveNonUniformScaleElements = false;
+
+protected:
+
+	void InitializeInstanceSets();
+	void UpdateInstanceSets();
+
+	void EmitResults();
+
+	void UpdateWarningMessages();
+};
