Make the revolve boundary tool always reset the axis origin to the mesh location, rather than using the origin / the last-used location.

diff --git a/Engine/Plugins/Experimental/MeshModelingToolsetExp/Source/MeshModelingToolsExp/Private/RevolveBoundaryTool.cpp b/Engine/Plugins/Experimental/MeshModelingToolsetExp/Source/MeshModelingToolsExp/Private/RevolveBoundaryTool.cpp
--- a/Engine/Plugins/Experimental/MeshModelingToolsetExp/Source/MeshModelingToolsExp/Private/RevolveBoundaryTool.cpp
+++ b/Engine/Plugins/Experimental/MeshModelingToolsetExp/Source/MeshModelingToolsExp/Private/RevolveBoundaryTool.cpp
@@ -112,13 +112,7 @@ void URevolveBoundaryTool::Setup()
 
 
 	FTransform LocalToWorld = UE::ToolTarget::GetLocalToWorldTransform(Target);
-	// Assume an axis that is > 1000 meters away is typically not desired, and can be snapped closer
-	// This works around bad behavior if the tool is started at LWC
-	constexpr double AxisVeryFarThreshold = 100 * 1000;
-	if (FVector::DistSquared(Settings->AxisOrigin, LocalToWorld.GetTranslation()) > AxisVeryFarThreshold * AxisVeryFarThreshold)
-	{
-		Settings->AxisOrigin = LocalToWorld.GetTranslation();
-	}
+	Settings->AxisOrigin = LocalToWorld.GetTranslation();
 
 	UpdateRevolutionAxis();
 
