Forcing skin weight colors to have alpha value of 1.0

diff --git a/Engine/Plugins/Experimental/MeshModelingToolsetExp/Source/MeshModelingToolsEditorOnlyExp/Private/SkeletalMesh/SkinWeightsPaintTool.cpp b/Engine/Plugins/Experimental/MeshModelingToolsetExp/Source/MeshModelingToolsEditorOnlyExp/Private/SkeletalMesh/SkinWeightsPaintTool.cpp
--- a/Engine/Plugins/Experimental/MeshModelingToolsetExp/Source/MeshModelingToolsEditorOnlyExp/Private/SkeletalMesh/SkinWeightsPaintTool.cpp
+++ b/Engine/Plugins/Experimental/MeshModelingToolsetExp/Source/MeshModelingToolsEditorOnlyExp/Private/SkeletalMesh/SkinWeightsPaintTool.cpp
@@ -1867,14 +1867,23 @@ void USkinWeightsPaintTool::OnPropertyModified(UObject* ModifiedObject, FPropert
 {
 	Super::OnPropertyModified(ModifiedObject, ModifiedProperty);
 	
-	// invalidate vertex color cache when weight color properties are modified
 	const bool bColorModeModified = ModifiedProperty->GetNameCPP() == GET_MEMBER_NAME_STRING_CHECKED(USkinWeightsPaintToolProperties, ColorMode);
 	const bool bColorRampModified = ModifiedProperty->GetNameCPP() == GET_MEMBER_NAME_STRING_CHECKED(USkinWeightsPaintToolProperties, ColorRamp);
 	const bool bMinColorModified = ModifiedProperty->GetNameCPP() == GET_MEMBER_NAME_STRING_CHECKED(USkinWeightsPaintToolProperties, MinColor);
 	const bool bMaxColorModified = ModifiedProperty->GetNameCPP() == GET_MEMBER_NAME_STRING_CHECKED(USkinWeightsPaintToolProperties, MaxColor);
+
+	// invalidate vertex color cache when any weight color properties are modified
 	if (bColorModeModified || bColorRampModified || bMinColorModified || bMaxColorModified)
 	{
 		bVisibleWeightsValid = false;
+
+		// force all colors to have Alpha = 1
+		WeightToolProperties->MinColor.A = 1.f;
+		WeightToolProperties->MaxColor.A = 1.f;
+		for (FLinearColor& Color : WeightToolProperties->ColorRamp)
+		{
+			Color.A = 1.f;
+		}
 	}
 }
 
