ModelingMode: modify logic in HarvestInstancesTool around detecting additional non-instanceable Components on input Actors, because Nav system seems to sometimes add hidden extra Editor-only Components on SM Actors.

diff --git a/Engine/Plugins/Experimental/MeshModelingToolsetExp/Source/MeshModelingToolsEditorOnlyExp/Private/HarvestInstancesTool.cpp b/Engine/Plugins/Experimental/MeshModelingToolsetExp/Source/MeshModelingToolsEditorOnlyExp/Private/HarvestInstancesTool.cpp
--- a/Engine/Plugins/Experimental/MeshModelingToolsetExp/Source/MeshModelingToolsEditorOnlyExp/Private/HarvestInstancesTool.cpp
+++ b/Engine/Plugins/Experimental/MeshModelingToolsetExp/Source/MeshModelingToolsEditorOnlyExp/Private/HarvestInstancesTool.cpp
@@ -276,8 +276,8 @@ void UHarvestInstancesTool::UpdateInstanceSets()
 	}
 
 
-	// check that all Components of Source Actors are included in instance set. If not, we cannot delete
-	// the Source Actors
+	// Check that all SceneComponents of Source Actors are included in instance set. 
+	// If not, we cannot delete the Source Actors
 	bool bFoundUnhandledComponent = false;
 	for (AActor* SourceActor : SourceActors)
 	{
@@ -285,6 +285,9 @@ void UHarvestInstancesTool::UpdateInstanceSets()
 		{
 			if (bFoundUnhandledComponent) return;
 
+			// only going to consider SceneComponents
+			if (Cast<USceneComponent>(Component) == nullptr) return;
+
 			if (UPrimitiveComponent* PrimComponent = Cast<UPrimitiveComponent>(Component))
 			{
 				if (SourceComponents.Contains(PrimComponent))
@@ -292,7 +295,12 @@ void UHarvestInstancesTool::UpdateInstanceSets()
 					return;
 				}
 			}
-			bFoundUnhandledComponent = true;
+
+			// for now skip editor-only Components as some editor systems (like navigation) apparently attach additional hidden Components to static meshes...
+			if (Component->IsEditorOnly() == false)
+			{
+				bFoundUnhandledComponent = true;
+			}
 		});
 	}
 	OutputSettings->bCanDeleteInputs = (bFoundUnhandledComponent == false);
