Fix issue, skin paint colors not updating while editing individual color channels in details panel.

diff --git a/Engine/Plugins/Experimental/MeshModelingToolsetExp/Source/MeshModelingToolsEditorOnlyExp/Private/SkeletalMesh/SkinWeightsPaintTool.cpp b/Engine/Plugins/Experimental/MeshModelingToolsetExp/Source/MeshModelingToolsEditorOnlyExp/Private/SkeletalMesh/SkinWeightsPaintTool.cpp
--- a/Engine/Plugins/Experimental/MeshModelingToolsetExp/Source/MeshModelingToolsEditorOnlyExp/Private/SkeletalMesh/SkinWeightsPaintTool.cpp
+++ b/Engine/Plugins/Experimental/MeshModelingToolsetExp/Source/MeshModelingToolsEditorOnlyExp/Private/SkeletalMesh/SkinWeightsPaintTool.cpp
@@ -1866,14 +1866,20 @@ void USkinWeightsPaintTool::ToggleEditingMode()
 void USkinWeightsPaintTool::OnPropertyModified(UObject* ModifiedObject, FProperty* ModifiedProperty)
 {
 	Super::OnPropertyModified(ModifiedObject, ModifiedProperty);
-	
-	const bool bColorModeModified = ModifiedProperty->GetNameCPP() == GET_MEMBER_NAME_STRING_CHECKED(USkinWeightsPaintToolProperties, ColorMode);
-	const bool bColorRampModified = ModifiedProperty->GetNameCPP() == GET_MEMBER_NAME_STRING_CHECKED(USkinWeightsPaintToolProperties, ColorRamp);
-	const bool bMinColorModified = ModifiedProperty->GetNameCPP() == GET_MEMBER_NAME_STRING_CHECKED(USkinWeightsPaintToolProperties, MinColor);
-	const bool bMaxColorModified = ModifiedProperty->GetNameCPP() == GET_MEMBER_NAME_STRING_CHECKED(USkinWeightsPaintToolProperties, MaxColor);
+
+	const FString NameOfModifiedProperty = ModifiedProperty->GetNameCPP();
 
 	// invalidate vertex color cache when any weight color properties are modified
-	if (bColorModeModified || bColorRampModified || bMinColorModified || bMaxColorModified)
+	const TArray<FString> ColorPropertyNames = {
+		GET_MEMBER_NAME_STRING_CHECKED(USkinWeightsPaintToolProperties, ColorMode),
+		GET_MEMBER_NAME_STRING_CHECKED(USkinWeightsPaintToolProperties, ColorRamp),
+		GET_MEMBER_NAME_STRING_CHECKED(USkinWeightsPaintToolProperties, MinColor),
+		GET_MEMBER_NAME_STRING_CHECKED(USkinWeightsPaintToolProperties, MaxColor),
+		GET_MEMBER_NAME_STRING_CHECKED(FLinearColor, R),
+		GET_MEMBER_NAME_STRING_CHECKED(FLinearColor, G),
+		GET_MEMBER_NAME_STRING_CHECKED(FLinearColor, B),
+		GET_MEMBER_NAME_STRING_CHECKED(FLinearColor, A)};
+	if (ColorPropertyNames.Contains(NameOfModifiedProperty))
 	{
 		bVisibleWeightsValid = false;
 
