Add ability to frame selected bone(s) in skin paint tool.

#swarm https://p4-swarm.epicgames.net/reviews/26750179

diff --git a/Engine/Plugins/Experimental/MeshModelingToolsetExp/Source/MeshModelingToolsEditorOnlyExp/Private/SkeletalMesh/SkinWeightsPaintTool.cpp b/Engine/Plugins/Experimental/MeshModelingToolsetExp/Source/MeshModelingToolsEditorOnlyExp/Private/SkeletalMesh/SkinWeightsPaintTool.cpp
--- a/Engine/Plugins/Experimental/MeshModelingToolsetExp/Source/MeshModelingToolsEditorOnlyExp/Private/SkeletalMesh/SkinWeightsPaintTool.cpp
+++ b/Engine/Plugins/Experimental/MeshModelingToolsetExp/Source/MeshModelingToolsEditorOnlyExp/Private/SkeletalMesh/SkinWeightsPaintTool.cpp
@@ -777,9 +777,7 @@ void USkinWeightsPaintTool::Render(IToolsContextRenderAPI* RenderAPI)
 
 FBox USkinWeightsPaintTool::GetWorldSpaceFocusBox()
 {
-	// TODO add support for selected bones
-
-	// frame vertex selection if there is one
+	// 1. Prioritize framing vertex selection if vertices are selected
 	if (PolygonSelectionMechanic)
 	{
 		const FGroupTopologySelection& Selection = PolygonSelectionMechanic->GetActiveSelection();
@@ -799,7 +797,44 @@ FBox USkinWeightsPaintTool::GetWorldSpaceFocusBox()
 			}
 		}
 	}
-	
+
+	// 2. Fallback on framing selected bones (if there are any)
+	// TODO, there are several places in the engine that frame bone selections. Let's consolidate this logic.
+	if (!SelectedBoneIndices.IsEmpty())
+	{
+		const USkeletalMeshComponent* MeshComponent = Weights.Deformer.Component;
+		const FReferenceSkeleton& RefSkeleton = MeshComponent->GetSkeletalMeshAsset()->GetRefSkeleton();
+		const TArray<FTransform>& CurrentBoneTransforms = MeshComponent->GetComponentSpaceTransforms();
+		FAxisAlignedBox3d Bounds = FAxisAlignedBox3d::Empty();
+		for (const int32 BoneIndex : SelectedBoneIndices)
+		{
+			// add bone position and position of all direct children to the frame bounds
+			const FVector BonePosition = CurrentBoneTransforms[BoneIndex].GetLocation();
+			Bounds.Contain(BonePosition);
+			TArray<int32> ChildrenIndices;
+			RefSkeleton.GetDirectChildBones(BoneIndex, ChildrenIndices);
+			if (ChildrenIndices.IsEmpty())
+			{
+				constexpr float SingleBoneSize = 10.f;
+				FVector BoneOffset = FVector(SingleBoneSize, SingleBoneSize, SingleBoneSize);
+				Bounds.Contain(BonePosition + BoneOffset);
+				Bounds.Contain(BonePosition - BoneOffset);
+			}
+			else
+			{
+				for (const int32 ChildIndex : ChildrenIndices)
+				{
+					Bounds.Contain(CurrentBoneTransforms[ChildIndex].GetLocation());
+				}
+			}	
+		}
+		if (Bounds.MaxDim() > FMathf::ZeroTolerance)
+		{
+			return static_cast<FBox>(Bounds);
+		}
+	}
+
+	// 3. Finally, fallback on component bounds if nothing else is selected
 	return PreviewMesh->GetActor()->GetComponentsBoundingBox();
 }
 
