Fixing crash when calling when the number of materials in the instance does not match the number of Material section in the MeshDescription.

diff --git a/Engine/Plugins/Experimental/MeshModelingToolsetExp/Source/GeometryProcessingAdapters/Private/CombineMeshInstancesImpl.cpp b/Engine/Plugins/Experimental/MeshModelingToolsetExp/Source/GeometryProcessingAdapters/Private/CombineMeshInstancesImpl.cpp
--- a/Engine/Plugins/Experimental/MeshModelingToolsetExp/Source/GeometryProcessingAdapters/Private/CombineMeshInstancesImpl.cpp
+++ b/Engine/Plugins/Experimental/MeshModelingToolsetExp/Source/GeometryProcessingAdapters/Private/CombineMeshInstancesImpl.cpp
@@ -3432,7 +3432,7 @@ void BuildCombinedMesh(
 				for (int32 tid : TempAppendMesh.TriangleIndicesItr())
 				{
 					int32 SourceMaterialID = (AppendMaterialAttrib != nullptr) ? AppendMaterialAttrib->GetValue(tid) : 0;
-					UMaterialInterface* UseMaterial = Instance.Materials[SourceMaterialID];
+					UMaterialInterface* UseMaterial = Instance.Materials.IsValidIndex(SourceMaterialID) ? Instance.Materials[SourceMaterialID] : nullptr;
 					const int32* FoundMaterialIndex = Assembly.MaterialMap.Find(UseMaterial);
 					int32 AssignMaterialIndex = (FoundMaterialIndex != nullptr) ? *FoundMaterialIndex : 0;
 
