CombineMeshInstances: add more support for preserving UVs
#rnx

diff --git a/Engine/Plugins/Experimental/MeshModelingToolsetExp/Source/GeometryProcessingAdapters/Private/CombineMeshInstancesImpl.cpp b/Engine/Plugins/Experimental/MeshModelingToolsetExp/Source/GeometryProcessingAdapters/Private/CombineMeshInstancesImpl.cpp
--- a/Engine/Plugins/Experimental/MeshModelingToolsetExp/Source/GeometryProcessingAdapters/Private/CombineMeshInstancesImpl.cpp
+++ b/Engine/Plugins/Experimental/MeshModelingToolsetExp/Source/GeometryProcessingAdapters/Private/CombineMeshInstancesImpl.cpp
@@ -86,6 +86,16 @@ enum class EMeshDetailLevel
 };
 
 
+enum class ECombinedLODType
+{
+	Copied = 0,
+	Simplified = 1,
+	Approximated = 2,
+	VoxWrapped = 3
+};
+
+
+
 // FMeshPartInstance represents a single instance of a FMeshPart
 struct FMeshPartInstance
 {
@@ -761,6 +771,8 @@ static void SimplifyPartMesh(
 	double Tolerance, 
 	double RecomputeNormalsAngleThreshold,
 	bool bTryToPreserveSalientCorners = false,
+	bool bPreserveUVs = false,
+	bool bPreserveVertexColors = false,
 	double PreserveCornersAngleThreshold = 44,
 	double MinSalientPartDimension = 1.0)
 {
@@ -781,17 +793,26 @@ static void SimplifyPartMesh(
 		return;
 	}
 
-	using SimplifierType = FVolPresMeshSimplification;
-	//using SimplifierType = FQEMSimplification;
-	SimplifierType Simplifier(&EditMesh);
-
 	// clear out attributes so it doesn't affect simplification
-	//EditMesh.DiscardAttributes();
-	EditMesh.Attributes()->SetNumUVLayers(0);
+	bool bAllAttributesCleared = (!bPreserveUVs) && (!bPreserveVertexColors);
+	if (bPreserveUVs == false)
+	{
+		EditMesh.Attributes()->SetNumUVLayers(0);
+	}
+	if (bPreserveVertexColors == false)
+	{
+		EditMesh.Attributes()->DisablePrimaryColors();
+	}
 	EditMesh.Attributes()->DisableTangents();
-	EditMesh.Attributes()->DisablePrimaryColors();
 	FMeshNormals::InitializeOverlayToPerVertexNormals(EditMesh.Attributes()->PrimaryNormals(), false);
 
+
+	// todo: if preserving UVs or Vertex Colors, might prefer the Attribute simplifier here? 
+	// Unclear how to do that conditionally as it's a template, though...
+	using SimplifierType = FVolPresMeshSimplification;
+	//using SimplifierType = FQEMSimplification;
+	SimplifierType Simplifier(&EditMesh);
+
 	Simplifier.ProjectionMode = SimplifierType::ETargetProjectionMode::NoProjection;
 
 	FColliderMesh ColliderMesh;
@@ -805,19 +826,16 @@ static void SimplifyPartMesh(
 	// This should perhaps be based on some heuristics about 'part type'
 	Simplifier.bRetainQuadricMemory = true;
 												
-	// currenly no need for this path, as seam attributes have been cleared. 
-	//if ( bNoSplitAttributes == false )
-	//{
-	//	Simplifier.bAllowSeamCollapse = true;
-	//	Simplifier.SetEdgeFlipTolerance(1.e-5);
-	//	if (EditMesh.HasAttributes())
-	//	{
-	//		EditMesh.Attributes()->SplitAllBowties();	// eliminate any bowties that might have formed on attribute seams.
-	//	}
-	//}
-
-	// this should preserve part shape better but it completely fails currently =\
-	//Simplifier.CollapseMode = FVolPresMeshSimplification::ESimplificationCollapseModes::MinimalExistingVertexError;
+	// if preserving any attributes, have to clean up seams
+	if (bAllAttributesCleared == false )
+	{
+		Simplifier.bAllowSeamCollapse = true;
+		Simplifier.SetEdgeFlipTolerance(1.e-5);
+		if (EditMesh.HasAttributes())
+		{
+			EditMesh.Attributes()->SplitAllBowties();	// eliminate any bowties that might have formed on attribute seams.
+		}
+	}
 
 	// do these flags matter here since we are not flipping??
 	EEdgeRefineFlags MeshBoundaryConstraints = EEdgeRefineFlags::NoFlip;
@@ -1329,8 +1347,15 @@ void ComputeMeshApproximations(
 		for (int32 k = 0; k < NumSimplifiedLODs; ++k)
 		{
 			ApproxGeo.SimplifiedMeshLODs[k] = *OptimizationSourceMesh;
-			SimplifyPartMesh(ApproxGeo.SimplifiedMeshLODs[k], InitialTolerance, AngleThresholdDeg,
-				CombineOptions.bSimplifyPreserveCorners, CombineOptions.SimplifySharpEdgeAngleDeg, CombineOptions.SimplifyMinSalientDimension);
+			int32 SimplifiedLODIndex = NumSourceLODs+k;
+			SimplifyPartMesh(ApproxGeo.SimplifiedMeshLODs[k], 
+				InitialTolerance, 
+				AngleThresholdDeg,
+				CombineOptions.bSimplifyPreserveCorners,
+				Part->bPreserveUVs || CombineOptions.bSimplifyPreserveUVs || (SimplifiedLODIndex <= CombineOptions.PreserveUVLODLevel),
+				CombineOptions.bSimplifyPreserveVertexColors,
+				CombineOptions.SimplifySharpEdgeAngleDeg, 
+				CombineOptions.SimplifyMinSalientDimension);
 			InitialTolerance *= CombineOptions.SimplifyLODLevelToleranceScale;
 		}
 
@@ -1405,14 +1430,17 @@ void ComputeMeshApproximations(
 			// these fields are precomputed
 			int32 PartIndex = 0;			// index into Assembly.Parts
 			int32 NumInstances = 0;			// number of instances of this part in the final mesh, currently *excluding* decorative parts
-			double ReplacedWeight = 1.0;	// reduce this as we shift LODs up, to try to let other comparable parts take the hit...
 
 			TArray<FDynamicMesh3*> LODChainMeshes;	// flattened list of mesh pointers for [Source LODs][Simplified LODs][Approximate LODs]
+			TArray<ECombinedLODType> LODChainMeshTypes;
+			TArray<bool> LODHasUVs;
 
 			// these fields are temporary storage updated during the algorilthm below
 			int32 PartTriCount = 0;			// triangle count of the part for the active LOD
 			int32 TotalTriCount = 0;		// total estimated triangle count for this part in the combined mesh 
 
+			double ReplacedWeight = 1.0;	// reduce this as we shift LODs up, to try to let other comparable parts take the hit...
+
 			double PartCostWeight() const 
 			{
 				if (NumInstances == 0 || PartTriCount <= 12) return 0.0;		// assume a box is min-cost and cannot be improved
@@ -1438,17 +1466,28 @@ void ComputeMeshApproximations(
 				}
 			}
 
+			int32 LODIndex = 0;
 			for (FDynamicMesh3& SourceLODMesh : Assembly.SourceMeshGeometry[SetIndex].SourceMeshLODs)
 			{
 				CostInfo[SetIndex].LODChainMeshes.Add(&SourceLODMesh);
+				CostInfo[SetIndex].LODChainMeshTypes.Add(ECombinedLODType::Copied);
+				CostInfo[SetIndex].LODHasUVs.Add(true);
+				LODIndex++;
 			}
 			for (FDynamicMesh3& SimplifiedLODMesh : Assembly.OptimizedMeshGeometry[SetIndex].SimplifiedMeshLODs)
 			{
 				CostInfo[SetIndex].LODChainMeshes.Add(&SimplifiedLODMesh);
+				CostInfo[SetIndex].LODChainMeshTypes.Add(ECombinedLODType::Simplified);
+				CostInfo[SetIndex].LODHasUVs.Add(
+					Part->bPreserveUVs || CombineOptions.bSimplifyPreserveUVs || LODIndex <= CombineOptions.PreserveUVLODLevel);
+				LODIndex++;
 			}
 			for (FDynamicMesh3& ApproximateLODMesh : Assembly.OptimizedMeshGeometry[SetIndex].ApproximateMeshLODs)
 			{
 				CostInfo[SetIndex].LODChainMeshes.Add(&ApproximateLODMesh);
+				CostInfo[SetIndex].LODChainMeshTypes.Add(ECombinedLODType::Approximated);
+				CostInfo[SetIndex].LODHasUVs.Add(false);
+				LODIndex++;
 			}
 		}
 
@@ -1460,6 +1499,12 @@ void ComputeMeshApproximations(
 			int32 LODTriangleBudget = CombineOptions.HardLODBudgets[LODIndex] * CombineOptions.PartLODPromotionBudgetMultiplier;
 			if (LODTriangleBudget <= 0) continue;
 
+			// reset replaced weights that are incrementally updated below
+			for (FPartCostInfo& PartCostInfo : CostInfo)
+			{
+				PartCostInfo.ReplacedWeight = 1.0;
+			}
+
 			int32 LastTotalCurLODTriCount = 999999;
 			int32 MaxIters = 1000;
 			int32 NoProgressIters = 0;
@@ -1514,26 +1559,72 @@ void ComputeMeshApproximations(
 					}
 				}
 				int32 SetIndex = MaxSetIndex;
+
 				// if our worst part is a box (CostWeight == 0), there is no point in replacing it
 				if ( (CostInfo[SetIndex].PartCostWeight() > 0) 
 					&& (LODIndex < CostInfo[SetIndex].LODChainMeshes.Num()-2) )	
 				{
 					FPartCostInfo& ReplaceInfo = CostInfo[SetIndex];
 
-					if (bVerbose)
+					ECombinedLODType PartCurLODType = ReplaceInfo.LODChainMeshTypes[LODIndex];
+					ECombinedLODType PartNextLODType = ReplaceInfo.LODChainMeshTypes[LODIndex+1];
+
+					// if we want to preserve UVs for a part, or for a LOD level, we cannot allow
+					// a part with no UVs to be swapped in for a part that does have UVs. 
+					bool bReplacementIsAllowed = true;
+					bool bPartPreserveUVs = Assembly.Parts[ReplaceInfo.PartIndex]->bPreserveUVs;
+					if (bPartPreserveUVs || LODIndex <= CombineOptions.PreserveUVLODLevel)
 					{
-						UE_LOG(LogGeometry, Log, TEXT("    PartPromotion LOD %1d: Iter %4d  Promoting Part %4d, from %5d tris to %5d tris"), 
-							LODIndex, NumIter, SetIndex, ReplaceInfo.LODChainMeshes[LODIndex]->TriangleCount(), ReplaceInfo.LODChainMeshes[LODIndex+1]->TriangleCount());
+						bool bCurHasUVs = ReplaceInfo.LODHasUVs[LODIndex];
+						bool bNextHasUVs = ReplaceInfo.LODHasUVs[LODIndex+1];
+						if (bCurHasUVs && bNextHasUVs == false)
+						{
+							bReplacementIsAllowed = false;
+						}
 					}
 
-					int32 NumAllLODs = ReplaceInfo.LODChainMeshes.Num();
-					for (int32 k = LODIndex; k < NumAllLODs - 1; ++k)
+					if (bReplacementIsAllowed == false)
 					{
-						*ReplaceInfo.LODChainMeshes[k] = *ReplaceInfo.LODChainMeshes[k + 1];
+						if (bVerbose)
+						{
+							UE_LOG(LogGeometry, Log, TEXT("    PartPromotion LOD %1d: Iter %4d  Disallowed Promoting Part %4d"),
+								LODIndex, NumIter, SetIndex);
+						}
+
+						// if replacement is not allowed at this LOD, set a large negative weight so that we do not
+						// consider this part again
+						ReplaceInfo.ReplacedWeight = -9999.0;
 					}
+					else
+					{
+						if (bVerbose)
+						{
+							UE_LOG(LogGeometry, Log, TEXT("    PartPromotion LOD %1d: Iter %4d  Promoting Part %4d, from %5d tris to %5d tris (replacing type %d with type %d)"), 
+								LODIndex, NumIter, SetIndex, ReplaceInfo.LODChainMeshes[LODIndex]->TriangleCount(), ReplaceInfo.LODChainMeshes[LODIndex+1]->TriangleCount(), 
+								(int32)PartCurLODType, (int32)PartNextLODType);
+						}
+
+						// shift all meshes in the LOD chain down one slot
+						int32 NumAllLODs = ReplaceInfo.LODChainMeshes.Num();
+						for (int32 k = LODIndex; k < NumAllLODs - 1; ++k)
+						{
+							// if we want to preserve UVs for this part, or up to some LOD level, and the
+							// next LOD has no UVs, we have to stop shifting
+							if ((bPartPreserveUVs || k == CombineOptions.PreserveUVLODLevel)
+								&& ReplaceInfo.LODHasUVs[k + 1] == false)
+							{
+								break;
+							}
 
-					ReplaceInfo.ReplacedWeight *= 0.5;
+							*ReplaceInfo.LODChainMeshes[k] = *ReplaceInfo.LODChainMeshes[k + 1];
+							ReplaceInfo.LODChainMeshTypes[k] = ReplaceInfo.LODChainMeshTypes[k + 1];
+							ReplaceInfo.LODHasUVs[k] = ReplaceInfo.LODHasUVs[k + 1];
+						}
+
+						ReplaceInfo.ReplacedWeight *= 0.5;
+					}
 				}
+
 				// slowly increase weights of parts   (should this be modulated by tri count?)
 				for (int32 k = 0; k < NumParts; ++k)
 				{
@@ -2854,15 +2945,6 @@ struct FCombinedMeshLOD
 }
 
 
-enum class ECombinedLODType
-{
-	Copied = 0,
-	Simplified = 1,
-	Approximated = 2,
-	VoxWrapped = 3
-};
-
-
 
 static void SortMesh(FDynamicMesh3& Mesh)
 {
@@ -2980,10 +3062,14 @@ void OptimizeLODMeshTriangulation(
 		GroupingIDFunc = [](const FDynamicMesh3&, int32) { return FIndex3i::Zero(); };
 	}
 
-	PostProcessHiddenFaceRemovedMesh(MeshLOD, BaseGeometricTolerance,
-		CombineOptions.bMergeCoplanarFaces && LODIndex >= CombineOptions.MergeCoplanarFacesStartLOD,
-		CombineOptions.PlanarPolygonRetriangulationStartLOD >= 0 && LODIndex >= CombineOptions.PlanarPolygonRetriangulationStartLOD,
-		GroupingIDFunc, SkipMaterialIDs);
+	bool bWantCoplanarMerging = (CombineOptions.bMergeCoplanarFaces)
+		&& (LODIndex >= CombineOptions.MergeCoplanarFacesStartLOD)
+		&& (LODIndex > CombineOptions.PreserveUVLODLevel);
+	bool bWantPlanarRetriangulation = (bWantCoplanarMerging)
+		&& (CombineOptions.PlanarPolygonRetriangulationStartLOD >= 0)
+		&& (LODIndex >= CombineOptions.PlanarPolygonRetriangulationStartLOD);
+
+	PostProcessHiddenFaceRemovedMesh(MeshLOD, BaseGeometricTolerance, bWantCoplanarMerging, bWantPlanarRetriangulation,	GroupingIDFunc, SkipMaterialIDs);
 }
 
 
@@ -3402,8 +3488,15 @@ void BuildCombinedMesh(
 					Assembly.PreProcessInstanceMeshFunc(TempAppendMesh, Instance);
 				}
 
-				// if part does not require UVs, get rid of them here to encourage merging downstream
-				if (LODLevel > 0 && Part->bPreserveUVs == false && TempAppendMesh.HasAttributes())
+				// determine if we should be keeping UVs around for this Part
+				bool bPreserveUVs = (LODLevel == 0)
+					|| (LODLevel <= CombineOptions.PreserveUVLODLevel)
+					|| ((int)LevelLODType <= (int)ECombinedLODType::Simplified && CombineOptions.bSimplifyPreserveUVs)
+					|| Part->bPreserveUVs;
+
+				// if part does not require UVs, but still has them, discard them here to encourage merging downstream
+				// (todo: is this ever possible now? should we skip this if we are not going to do the merging?)
+				if (bPreserveUVs == false && TempAppendMesh.HasAttributes())
 				{
 					for (int32 UVLayer = 0; UVLayer < TempAppendMesh.Attributes()->NumUVLayers(); ++UVLayer)
 					{
diff --git a/Engine/Source/Developer/GeometryProcessingInterfaces/Public/GeometryProcessingInterfaces/CombineMeshInstances.h b/Engine/Source/Developer/GeometryProcessingInterfaces/Public/GeometryProcessingInterfaces/CombineMeshInstances.h
--- a/Engine/Source/Developer/GeometryProcessingInterfaces/Public/GeometryProcessingInterfaces/CombineMeshInstances.h
+++ b/Engine/Source/Developer/GeometryProcessingInterfaces/Public/GeometryProcessingInterfaces/CombineMeshInstances.h
@@ -155,6 +155,11 @@ public:
 		int32 NumSimplifiedLODs = 3;
 		double SimplifyBaseTolerance = 1.0;
 		double SimplifyLODLevelToleranceScale = 2.0;
+		// if true, UVs will be preserved in Simplified LODs. This generally will result in lower-quality geometric shape approximation.
+		bool bSimplifyPreserveUVs = false;
+		// if true, vertex colors will be preserved in Simplified LODs. This generally will result in lower-quality geometric shape approximation.
+		bool bSimplifyPreserveVertexColors = false;
+		// if true, geometrically-detected "sharp" corners (eg like the corners of a box) will be preserved with hard constraints in Simplified LODs. This can be desirable on mechanical/geometric shapes.
 		bool bSimplifyPreserveCorners = true;
 		double SimplifySharpEdgeAngleDeg = 44.0;
 		double SimplifyMinSalientDimension = 1.0;
@@ -211,6 +216,12 @@ public:
 		// opening angle used to detect/assign sharp edges
 		double HardNormalAngleDeg = 15.0;
 
+		// UVs on input geometry will be preserved up to this LOD level (inclusive). 
+		// Note that this setting will severely constrain and/or fully disable many other optimizations. 
+		// In particular, Coplanar merging and retriangulation cannot be applied if UVs are to be preserved. 
+		// WARNING: this LOD level must be <= (NumCopiedLODs+NumSimplifiedLODs)
+		int PreserveUVLODLevel = -1;
+
 		// Attempt to merge/weld coplanar areas after hidden removal, and then further simplify those merged areas
 		// Coplanar merging is never applied between areas with different Materials. 
 		// TriangleGroupingIDFunc below can be used to further demarcate important internal shape boundaries.
