CombineMeshInstances: fix some crashes that can happen if a 0-triangle LOD mesh makes it through to the LOD post-processing. Catch case of no input parts in top-level ::CombineMeshInstances().
#rnx

diff --git a/Engine/Plugins/Experimental/MeshModelingToolsetExp/Source/GeometryProcessingAdapters/Private/CombineMeshInstancesImpl.cpp b/Engine/Plugins/Experimental/MeshModelingToolsetExp/Source/GeometryProcessingAdapters/Private/CombineMeshInstancesImpl.cpp
--- a/Engine/Plugins/Experimental/MeshModelingToolsetExp/Source/GeometryProcessingAdapters/Private/CombineMeshInstancesImpl.cpp
+++ b/Engine/Plugins/Experimental/MeshModelingToolsetExp/Source/GeometryProcessingAdapters/Private/CombineMeshInstancesImpl.cpp
@@ -3097,6 +3097,8 @@ void ProcessCombinedLODChain(
 	{
 		for (int32 LODIndex = CombineOptions.RemoveHiddenStartLOD; LODIndex < NumLODs && LODIndex < FirstVoxWrappedIndex; ++LODIndex)
 		{
+			if (MeshLODs[LODIndex].Mesh.TriangleCount() == 0) continue;
+
 			if (bVerbose)
 			{
 				UE_LOG(LogGeometry, Log, TEXT("  Optimizing LOD%d - Tris %6d Verts %6d"), LODIndex, MeshLODs[LODIndex].Mesh.TriangleCount(), MeshLODs[LODIndex].Mesh.VertexCount());
@@ -3122,7 +3124,7 @@ void ProcessCombinedLODChain(
 	// Process VoxWrapped LODs 
 	//
 	bool bUsingCoarseSweepApproximation = false;
-	if ( FirstVoxWrappedIndex < 9999 )
+	if ( FirstVoxWrappedIndex < 9999 &&  (MeshLODs[FirstVoxWrappedIndex].Mesh.TriangleCount() > 0) )
 	{
 		FDynamicMesh3 SourceVoxWrapMesh = MoveTemp(MeshLODs[FirstVoxWrappedIndex].Mesh);
 		FDynamicMeshAABBTree3 SourceSpatial(&SourceVoxWrapMesh, true);
@@ -3259,6 +3261,8 @@ void ProcessCombinedLODChain(
 	{
 		for (int32 LODIndex = 0; LODIndex < NumLODs && LODIndex < FirstVoxWrappedIndex; ++LODIndex)
 		{
+			if (MeshLODs[LODIndex].Mesh.TriangleCount() == 0) continue;
+
 			UE::Tasks::FTask AutoUVTask = UE::Tasks::Launch(UE_SOURCE_LOCATION, [&MeshLODs, &CombineOptions, LODIndex, bComputeTangents]()
 			{
 				ComputeMissingUVs(MeshLODs[LODIndex].Mesh);
@@ -3277,6 +3281,8 @@ void ProcessCombinedLODChain(
 		TRACE_CPUPROFILER_EVENT_SCOPE(RemoveHidden);
 		ParallelFor(NumLODs, [&](int32 LODIndex)
 		{
+			if (MeshLODs[LODIndex].Mesh.TriangleCount() == 0) return;
+
 			if ( LODIndex >= FirstVoxWrappedIndex )
 			{ 
 				if ( bVerbose )
@@ -3901,6 +3907,12 @@ void FCombineMeshInstancesImpl::CombineMeshInstances(
 				PartAssembly.Parts.Num(), PartAssembly.UniqueMaterials.Num());
 		}
 
+		if (PartAssembly.Parts.Num() == 0)
+		{
+			// todo: set some kind of error code in ResultsOut...
+			return;
+		}
+
 		InitializeAssemblySourceMeshesFromLOD(PartAssembly, Options.BaseCopiedLOD, Options.NumCopiedLODs);
 		InitializePartAssemblySpatials(PartAssembly);
 	}
