Fix bug where BakeRC tool does not respect the "New Asset Location" option in Modeling Mode Quick Settings

diff --git a/Engine/Plugins/Experimental/MeshModelingToolsetExp/Source/MeshModelingToolsEditorOnlyExp/Private/BakeRenderCaptureTool.cpp b/Engine/Plugins/Experimental/MeshModelingToolsetExp/Source/MeshModelingToolsEditorOnlyExp/Private/BakeRenderCaptureTool.cpp
--- a/Engine/Plugins/Experimental/MeshModelingToolsetExp/Source/MeshModelingToolsEditorOnlyExp/Private/BakeRenderCaptureTool.cpp
+++ b/Engine/Plugins/Experimental/MeshModelingToolsetExp/Source/MeshModelingToolsEditorOnlyExp/Private/BakeRenderCaptureTool.cpp
@@ -479,10 +479,8 @@ void UBakeRenderCaptureTool::OnShutdown(EToolShutdownType ShutdownType)
 	if (ShutdownType == EToolShutdownType::Accept && ResultSettings->IsEmpty() == false)
 	{
 		// TODO Support skeletal meshes here---see BakeMeshAttributeMapsTool::OnShutdown
-		IStaticMeshBackedTarget* StaticMeshTarget = Cast<IStaticMeshBackedTarget>(Targets[0]);
-		UObject* SourceAsset = StaticMeshTarget ? StaticMeshTarget->GetStaticMesh() : nullptr;
 		const UPrimitiveComponent* SourceComponent = UE::ToolTarget::GetTargetComponent(Targets[0]);
-		CreateAssets(SourceComponent->GetWorld(), SourceAsset);
+		CreateAssets(SourceComponent->GetWorld());
 	}
 
 	// Clear actors on shutdown so that their lifetime is not tied to the lifetime of the tool
@@ -499,7 +497,7 @@ void UBakeRenderCaptureTool::OnShutdown(EToolShutdownType ShutdownType)
 	}
 }
 
-void UBakeRenderCaptureTool::CreateAssets(UWorld* SourceWorld, UObject* SourceAsset)
+void UBakeRenderCaptureTool::CreateAssets(UWorld* SourceWorld)
 {
 	const FString BaseName = UE::ToolTarget::GetTargetActor(Targets[0])->GetActorNameOrLabel();
 	const bool bPackedMRS = ResultSettings->PackedMRSMap != nullptr;
@@ -511,7 +509,7 @@ void UBakeRenderCaptureTool::CreateAssets(UWorld* SourceWorld, UObject* SourceAs
 		bool bIsFallbackTexture;
 	};
 
-	auto CreateTextureAsset = [this, &BaseName, &SourceWorld, &SourceAsset] (
+	auto CreateTextureAsset = [this, &BaseName, &SourceWorld] (
 		const FString& TexParamName,
 		FTexture2DBuilder::ETextureType Type,
 		TObjectPtr<UTexture2D> Texture,
@@ -547,7 +545,6 @@ void UBakeRenderCaptureTool::CreateAssets(UWorld* SourceWorld, UObject* SourceAs
 		// We need to save the Fallback textures as well so that the generated Materials can reference them
 		FCreateTextureObjectParams TexParams;
 		TexParams.TargetWorld = SourceWorld;
-		TexParams.StoreRelativeToObject = SourceAsset;
 		TexParams.BaseName = FString::Printf(TEXT("%s_%s"), *BaseName, *TexParamName);
 		TexParams.GeneratedTransientTexture = bUseTexture ? Texture : Fallback; 
 
@@ -662,7 +659,6 @@ void UBakeRenderCaptureTool::CreateAssets(UWorld* SourceWorld, UObject* SourceAs
 
 		FCreateMaterialObjectParams MaterialParams;
 		MaterialParams.TargetWorld = SourceWorld;
-		MaterialParams.StoreRelativeToObject = SourceAsset;
 		MaterialParams.BaseName = FString::Printf(TEXT("%s_Material"), *BaseName);
 		MaterialParams.MaterialToDuplicate = Material;
 		const FCreateMaterialObjectResult MaterialResult = UE::Modeling::CreateMaterialObject(GetToolManager(), MoveTemp(MaterialParams));
diff --git a/Engine/Plugins/Experimental/MeshModelingToolsetExp/Source/MeshModelingToolsEditorOnlyExp/Public/BakeRenderCaptureTool.h b/Engine/Plugins/Experimental/MeshModelingToolsetExp/Source/MeshModelingToolsEditorOnlyExp/Public/BakeRenderCaptureTool.h
--- a/Engine/Plugins/Experimental/MeshModelingToolsetExp/Source/MeshModelingToolsEditorOnlyExp/Public/BakeRenderCaptureTool.h
+++ b/Engine/Plugins/Experimental/MeshModelingToolsetExp/Source/MeshModelingToolsEditorOnlyExp/Public/BakeRenderCaptureTool.h
@@ -346,9 +346,8 @@ protected:
 	/**
 	 * Create texture and material assets from our result map of Texture2D
 	 * @param SourceWorld the source world to define where the assets will be stored.
-	 * @param SourceAsset if not null, result textures/material will be stored adjacent to this asset.
 	 */
-	void CreateAssets(UWorld* SourceWorld, UObject* SourceAsset);
+	void CreateAssets(UWorld* SourceWorld);
 
 	// The baking background compute operation
 	TUniquePtr<TGenericDataBackgroundCompute<UE::Geometry::FMeshMapBaker>> Compute = nullptr;
