Fixes crash switching from "Edit Skin Weights" to other tool.

diff --git a/Engine/Plugins/Experimental/MeshModelingToolsetExp/Source/MeshModelingToolsEditorOnlyExp/Private/SkeletalMesh/SkinWeightsPaintTool.cpp b/Engine/Plugins/Experimental/MeshModelingToolsetExp/Source/MeshModelingToolsEditorOnlyExp/Private/SkeletalMesh/SkinWeightsPaintTool.cpp
--- a/Engine/Plugins/Experimental/MeshModelingToolsetExp/Source/MeshModelingToolsEditorOnlyExp/Private/SkeletalMesh/SkinWeightsPaintTool.cpp
+++ b/Engine/Plugins/Experimental/MeshModelingToolsetExp/Source/MeshModelingToolsEditorOnlyExp/Private/SkeletalMesh/SkinWeightsPaintTool.cpp
@@ -805,32 +805,35 @@ FBox USkinWeightsPaintTool::GetWorldSpaceFocusBox()
 		const USkeletalMeshComponent* MeshComponent = Weights.Deformer.Component;
 		const FReferenceSkeleton& RefSkeleton = MeshComponent->GetSkeletalMeshAsset()->GetRefSkeleton();
 		const TArray<FTransform>& CurrentBoneTransforms = MeshComponent->GetComponentSpaceTransforms();
-		FAxisAlignedBox3d Bounds = FAxisAlignedBox3d::Empty();
-		for (const int32 BoneIndex : SelectedBoneIndices)
+		if (!CurrentBoneTransforms.IsEmpty())
 		{
-			// add bone position and position of all direct children to the frame bounds
-			const FVector BonePosition = CurrentBoneTransforms[BoneIndex].GetLocation();
-			Bounds.Contain(BonePosition);
-			TArray<int32> ChildrenIndices;
-			RefSkeleton.GetDirectChildBones(BoneIndex, ChildrenIndices);
-			if (ChildrenIndices.IsEmpty())
-			{
-				constexpr float SingleBoneSize = 10.f;
-				FVector BoneOffset = FVector(SingleBoneSize, SingleBoneSize, SingleBoneSize);
-				Bounds.Contain(BonePosition + BoneOffset);
-				Bounds.Contain(BonePosition - BoneOffset);
-			}
-			else
+			FAxisAlignedBox3d Bounds = FAxisAlignedBox3d::Empty();
+			for (const int32 BoneIndex : SelectedBoneIndices)
 			{
-				for (const int32 ChildIndex : ChildrenIndices)
+				// add bone position and position of all direct children to the frame bounds
+				const FVector BonePosition = CurrentBoneTransforms[BoneIndex].GetLocation();
+				Bounds.Contain(BonePosition);
+				TArray<int32> ChildrenIndices;
+				RefSkeleton.GetDirectChildBones(BoneIndex, ChildrenIndices);
+				if (ChildrenIndices.IsEmpty())
 				{
-					Bounds.Contain(CurrentBoneTransforms[ChildIndex].GetLocation());
+					constexpr float SingleBoneSize = 10.f;
+					FVector BoneOffset = FVector(SingleBoneSize, SingleBoneSize, SingleBoneSize);
+					Bounds.Contain(BonePosition + BoneOffset);
+					Bounds.Contain(BonePosition - BoneOffset);
 				}
-			}	
-		}
-		if (Bounds.MaxDim() > FMathf::ZeroTolerance)
-		{
-			return static_cast<FBox>(Bounds);
+				else
+				{
+					for (const int32 ChildIndex : ChildrenIndices)
+					{
+						Bounds.Contain(CurrentBoneTransforms[ChildIndex].GetLocation());
+					}
+				}	
+			}
+			if (Bounds.MaxDim() > FMathf::ZeroTolerance)
+			{
+				return static_cast<FBox>(Bounds);
+			}
 		}
 	}
 
