CombineMeshInstances: add workaround for SimplifyToMinimalPlanar seam collapse issue
#rnx

diff --git a/Engine/Plugins/Experimental/MeshModelingToolsetExp/Source/GeometryProcessingAdapters/Private/CombineMeshInstancesImpl.cpp b/Engine/Plugins/Experimental/MeshModelingToolsetExp/Source/GeometryProcessingAdapters/Private/CombineMeshInstancesImpl.cpp
--- a/Engine/Plugins/Experimental/MeshModelingToolsetExp/Source/GeometryProcessingAdapters/Private/CombineMeshInstancesImpl.cpp
+++ b/Engine/Plugins/Experimental/MeshModelingToolsetExp/Source/GeometryProcessingAdapters/Private/CombineMeshInstancesImpl.cpp
@@ -2371,9 +2371,10 @@ static void PostProcessHiddenFaceRemovedMesh(
 		Welder.bWeldAttrsOnMergedEdges = true;
 		Welder.Apply();
 
-		// todo: could discard UVs here to improve merging. Possibly also do planar remesh, it 
-		// seems like the current method still will not remove boundary vertices...
-
+		// although we split bowties above, we have now pulled out submeshes which may have created more bowties
+		FDynamicMeshEditor BowtieSplitter(&SubRegionMesh);
+		FDynamicMeshEditResult TmpEditResult;
+		BowtieSplitter.SplitBowties(TmpEditResult);
 		if (SubRegionMesh.HasAttributes())
 		{
 			SubRegionMesh.Attributes()->SplitAllBowties();
@@ -2382,14 +2383,19 @@ static void PostProcessHiddenFaceRemovedMesh(
 
 		// simplify to planar
 		FQEMSimplification Simplifier(&SubRegionMesh);
+		Simplifier.bAllowSeamCollapse = false;		// workaround for UV seam collapse issue
+		//Simplifier.bAllowSeamCollapse = true;
+		Simplifier.bRetainQuadricMemory = false;
+
+		EEdgeRefineFlags BoundaryConstraint = EEdgeRefineFlags::NoFlip;
 
 		// set up constraints, necessary to avoid crashing in presence of attributes
 		FMeshConstraints Constraints;
 		FMeshConstraintsUtil::ConstrainAllBoundariesAndSeams(Constraints, SubRegionMesh,
-			EEdgeRefineFlags::NoFlip, EEdgeRefineFlags::NoConstraint, EEdgeRefineFlags::NoConstraint, true, false, true);
+			BoundaryConstraint, EEdgeRefineFlags::NoConstraint, EEdgeRefineFlags::NoConstraint, true, false, Simplifier.bAllowSeamCollapse);
 		Simplifier.SetExternalConstraints(Constraints);
 		// need to transfer constraint setting to the simplifier, these are used to update the constraints as edges collapse.
-		Simplifier.MeshBoundaryConstraint = EEdgeRefineFlags::NoFlip;
+		Simplifier.MeshBoundaryConstraint = BoundaryConstraint;
 		Simplifier.GroupBoundaryConstraint = EEdgeRefineFlags::NoConstraint;
 		Simplifier.MaterialBoundaryConstraint = EEdgeRefineFlags::NoConstraint;
 
