MeshModelingTools: Fix crash in dynamic sculpt tool due to preview component being attached to the target actor and being destroyed by blueprint recompiles.

diff --git a/Engine/Plugins/Experimental/MeshModelingToolsetExp/Source/MeshModelingToolsExp/Private/DynamicMeshSculptTool.cpp b/Engine/Plugins/Experimental/MeshModelingToolsetExp/Source/MeshModelingToolsExp/Private/DynamicMeshSculptTool.cpp
--- a/Engine/Plugins/Experimental/MeshModelingToolsetExp/Source/MeshModelingToolsExp/Private/DynamicMeshSculptTool.cpp
+++ b/Engine/Plugins/Experimental/MeshModelingToolsetExp/Source/MeshModelingToolsExp/Private/DynamicMeshSculptTool.cpp
@@ -3,6 +3,7 @@
 #include "DynamicMeshSculptTool.h"
 #include "Containers/Map.h"
 #include "Async/Async.h"
+#include "Engine/World.h"
 #include "InteractiveToolManager.h"
 #include "InteractiveGizmoManager.h"
 #include "ToolBuilderUtil.h"
@@ -101,9 +102,12 @@ void UDynamicMeshSculptTool::Setup()
 	SetToolDisplayName(LOCTEXT("ToolName", "DynaSculpt"));
 
 	// create dynamic mesh component to use for live preview
-	DynamicMeshComponent = NewObject<UOctreeDynamicMeshComponent>(UE::ToolTarget::GetTargetActor(Target));
+	FActorSpawnParameters SpawnInfo;
+	PreviewMeshActor = TargetWorld->SpawnActor<AInternalToolFrameworkActor>(FVector::ZeroVector, FRotator::ZeroRotator, SpawnInfo);
+	DynamicMeshComponent = NewObject<UOctreeDynamicMeshComponent>(PreviewMeshActor);
 	DynamicMeshComponent->SetShadowsEnabled(false);
-	DynamicMeshComponent->SetupAttachment(UE::ToolTarget::GetTargetActor(Target)->GetRootComponent());
+	DynamicMeshComponent->SetupAttachment(PreviewMeshActor->GetRootComponent());
+	DynamicMeshComponent->SetWorldTransform((FTransform)UE::ToolTarget::GetLocalToWorldTransform(Target));
 	DynamicMeshComponent->RegisterComponent();
 	ToolSetupUtil::ApplyRenderingConfigurationToPreview(DynamicMeshComponent, Target); 
 
@@ -327,6 +331,12 @@ void UDynamicMeshSculptTool::Shutdown(EToolShutdownType ShutdownType)
 		DynamicMeshComponent = nullptr;
 	}
 
+	if (PreviewMeshActor != nullptr)
+	{
+		PreviewMeshActor->Destroy();
+		PreviewMeshActor = nullptr;
+	}
+
 	BrushProperties->SaveProperties(this);
 	SculptProperties->SaveProperties(this);
 	KelvinBrushProperties->SaveProperties(this);
diff --git a/Engine/Plugins/Experimental/MeshModelingToolsetExp/Source/MeshModelingToolsExp/Public/DynamicMeshSculptTool.h b/Engine/Plugins/Experimental/MeshModelingToolsetExp/Source/MeshModelingToolsExp/Public/DynamicMeshSculptTool.h
--- a/Engine/Plugins/Experimental/MeshModelingToolsetExp/Source/MeshModelingToolsExp/Public/DynamicMeshSculptTool.h
+++ b/Engine/Plugins/Experimental/MeshModelingToolsetExp/Source/MeshModelingToolsExp/Public/DynamicMeshSculptTool.h
@@ -347,6 +347,9 @@ private:
 	UPROPERTY()
 	TObjectPtr<UPreviewMesh> BrushIndicatorMesh;
 
+	UPROPERTY()
+	TObjectPtr<AInternalToolFrameworkActor> PreviewMeshActor = nullptr;
+
 	UPROPERTY()
 	TObjectPtr<UOctreeDynamicMeshComponent> DynamicMeshComponent;
 
