CombineMeshInstances: Add support for specifying constraints on allowed Approximation types for specific Parts
#rnx

diff --git a/Engine/Plugins/Experimental/MeshModelingToolsetExp/Source/GeometryProcessingAdapters/Private/CombineMeshInstancesImpl.cpp b/Engine/Plugins/Experimental/MeshModelingToolsetExp/Source/GeometryProcessingAdapters/Private/CombineMeshInstancesImpl.cpp
--- a/Engine/Plugins/Experimental/MeshModelingToolsetExp/Source/GeometryProcessingAdapters/Private/CombineMeshInstancesImpl.cpp
+++ b/Engine/Plugins/Experimental/MeshModelingToolsetExp/Source/GeometryProcessingAdapters/Private/CombineMeshInstancesImpl.cpp
@@ -127,6 +127,9 @@ struct FMeshPart
 
 	bool bPreserveUVs = false;
 	bool bAllowMerging = true;
+	bool bAllowApproximation = true;
+
+	IGeometryProcessing_CombineMeshInstances::EApproximationType ApproxFilter = IGeometryProcessing_CombineMeshInstances::EApproximationType::NoConstraint;
 
 	int32 GetNumTriangles() const 
 	{
@@ -229,6 +232,8 @@ void InitializeMeshPartAssembly(
 
 			(*FoundPart)->bPreserveUVs = GroupData.bPreserveUVs;
 			(*FoundPart)->bAllowMerging = GroupData.bAllowMerging;
+			(*FoundPart)->bAllowApproximation = GroupData.bAllowApproximation;
+			(*FoundPart)->ApproxFilter = GroupData.ApproximationConstraint;
 		}
 
 		NewInstance.SourceComponent = SourceMeshInstance.SourceComponent;
@@ -285,6 +290,8 @@ void InitializeMeshPartAssembly(
 
 			(*FoundPart)->bPreserveUVs = GroupData.bPreserveUVs;
 			(*FoundPart)->bAllowMerging = GroupData.bAllowMerging;
+			(*FoundPart)->bAllowApproximation = GroupData.bAllowApproximation;
+			(*FoundPart)->ApproxFilter = GroupData.ApproximationConstraint;
 		}
 
 		NewInstance.SourceComponent = nullptr;
@@ -1223,6 +1230,7 @@ static void ComputeSweptSolidApproximation(
 static void SelectBestFittingMeshApproximation(
 	const FDynamicMesh3& OriginalMesh, 
 	const FDynamicMeshAABBTree3& OriginalMeshSpatial,
+	IGeometryProcessing_CombineMeshInstances::EApproximationType ApproxTypes,
 	FDynamicMesh3& ResultMesh,
 	EApproximatePartMethod& BestMethodIDOut,
 	double AcceptableDeviationTol,
@@ -1235,25 +1243,39 @@ static void SelectBestFittingMeshApproximation(
 	ApproxSelector.Initialize(&OriginalMesh, &OriginalMeshSpatial);
 	ApproxSelector.TriangleCost = TriangleCost;
 	ApproxSelector.MaxAllowableDeviation = MaxDeviation;
+	bool bNoApproxFilter = (ApproxTypes == IGeometryProcessing_CombineMeshInstances::EApproximationType::NoConstraint);
 
-	ApproxSelector.AddGeneratedMesh( [&](FDynamicMesh3& PartMeshInOut) {
-		ComputeSimplePartApproximation(PartMeshInOut, PartMeshInOut, EApproximatePartMethod::AxisAlignedBox);
-	}, (int32)EApproximatePartMethod::AxisAlignedBox);
+	if ( bNoApproxFilter || ((int)ApproxTypes & (int)IGeometryProcessing_CombineMeshInstances::EApproximationType::AxisAlignedBox) > 0 )
+	{
+		ApproxSelector.AddGeneratedMesh( [&](FDynamicMesh3& PartMeshInOut) {
+			ComputeSimplePartApproximation(PartMeshInOut, PartMeshInOut, EApproximatePartMethod::AxisAlignedBox);
+		}, (int32)EApproximatePartMethod::AxisAlignedBox);
+	}
 
-	ApproxSelector.AddGeneratedMesh( [&](FDynamicMesh3& PartMeshInOut) {
-		ComputeSimplePartApproximation(PartMeshInOut, PartMeshInOut, EApproximatePartMethod::OrientedBox);
-	}, (int32)EApproximatePartMethod::OrientedBox );
+	if (bNoApproxFilter || ((int)ApproxTypes & (int)IGeometryProcessing_CombineMeshInstances::EApproximationType::OrientedBox) > 0)
+	{
+		ApproxSelector.AddGeneratedMesh( [&](FDynamicMesh3& PartMeshInOut) {
+			ComputeSimplePartApproximation(PartMeshInOut, PartMeshInOut, EApproximatePartMethod::OrientedBox);
+		}, (int32)EApproximatePartMethod::OrientedBox );
+	}
 
-	ApproxSelector.AddGeneratedMesh( [&](FDynamicMesh3& PartMeshInOut) {
-		ComputeSimplePartApproximation(PartMeshInOut, PartMeshInOut, EApproximatePartMethod::MinVolumeSweptHull);
-	}, (int32)EApproximatePartMethod::MinVolumeSweptHull );
+	if (bNoApproxFilter || ((int)ApproxTypes & (int)IGeometryProcessing_CombineMeshInstances::EApproximationType::SweptHull) > 0)
+	{
+		ApproxSelector.AddGeneratedMesh( [&](FDynamicMesh3& PartMeshInOut) {
+			ComputeSimplePartApproximation(PartMeshInOut, PartMeshInOut, EApproximatePartMethod::MinVolumeSweptHull);
+		}, (int32)EApproximatePartMethod::MinVolumeSweptHull );
+	}
 
-	ApproxSelector.AddGeneratedMesh( [&](FDynamicMesh3& PartMeshInOut) {
-		ComputeSimplePartApproximation(PartMeshInOut, PartMeshInOut, EApproximatePartMethod::ConvexHull);
-	}, (int32)EApproximatePartMethod::ConvexHull );
+	if (bNoApproxFilter || ((int)ApproxTypes & (int)IGeometryProcessing_CombineMeshInstances::EApproximationType::ConvexHull) > 0)
+	{
+		ApproxSelector.AddGeneratedMesh( [&](FDynamicMesh3& PartMeshInOut) {
+			ComputeSimplePartApproximation(PartMeshInOut, PartMeshInOut, EApproximatePartMethod::ConvexHull);
+		}, (int32)EApproximatePartMethod::ConvexHull );
+	}
 
 	// Add swept-solid approximations
 	// Currently this is a bit hardcoded and some of these numbers should be exposed as parameters
+	if (bNoApproxFilter || ((int)ApproxTypes & (int)IGeometryProcessing_CombineMeshInstances::EApproximationType::SweptProjection) > 0)
 	{
 		const double MinHoleSize = 10.0;		// very aggressive, should be exposed as a parameter
 		const double MinHoleArea = MinHoleSize * MinHoleSize;
@@ -1292,15 +1314,18 @@ static void SelectBestFittingMeshApproximation(
 	// If Axis-Aligned box volume is less than (100+k%) larger than best option, just use that instead.
 	// Default is 10%, but if approximation is likely to also be a box, double it.
 	// (todo should be configurable)
-	FVector2d ApproxMeshVolArea = TMeshQueries<FDynamicMesh3>::GetVolumeArea(ResultMesh);
-	FAxisAlignedBox3d AlignedBox = OriginalMesh.GetBounds(false);
-	double BoxVolume = AlignedBox.Volume();
-	double VolRatio = AlignedBox.Volume() / ApproxMeshVolArea.X;
-	double BoxPreferenceVolumeRatioPercent = (ResultMesh.TriangleCount() <= 12) ? 20.0 : 10.0;
-	if (VolRatio < (1.0 + BoxPreferenceVolumeRatioPercent/100.0) )
+	if (bNoApproxFilter || ((int)ApproxTypes & (int)IGeometryProcessing_CombineMeshInstances::EApproximationType::AxisAlignedBox) > 0)
 	{
-		ComputeSimplePartApproximation(OriginalMesh, ResultMesh, EApproximatePartMethod::AxisAlignedBox);
-		BestMethodIDOut = EApproximatePartMethod::OverrideAxisBox;
+		FVector2d ApproxMeshVolArea = TMeshQueries<FDynamicMesh3>::GetVolumeArea(ResultMesh);
+		FAxisAlignedBox3d AlignedBox = OriginalMesh.GetBounds(false);
+		double BoxVolume = AlignedBox.Volume();
+		double VolRatio = AlignedBox.Volume() / ApproxMeshVolArea.X;
+		double BoxPreferenceVolumeRatioPercent = (ResultMesh.TriangleCount() <= 12) ? 20.0 : 10.0;
+		if (VolRatio < (1.0 + BoxPreferenceVolumeRatioPercent/100.0) )
+		{
+			ComputeSimplePartApproximation(OriginalMesh, ResultMesh, EApproximatePartMethod::AxisAlignedBox);
+			BestMethodIDOut = EApproximatePartMethod::OverrideAxisBox;
+		}
 	}
 
 }
@@ -1377,6 +1402,7 @@ void ComputeMeshApproximations(
 		for (int32 k = 0; k < UseNumApproxLODs; ++k)
 		{
 			SelectBestFittingMeshApproximation(*OptimizationSourceMesh, OptimizationSourceMeshSpatial, 
+				Part->ApproxFilter,
 				ApproxGeo.ApproximateMeshLODs[k], SelectedMethodID[k], 
 				CombineOptions.SimplifyBaseTolerance, InitialTriCost, CombineOptions.MaxAllowableApproximationDeviation);
 			if (k < NumApproxLODs)
@@ -3488,7 +3514,8 @@ void BuildCombinedMesh(
 				// TODO: if part budget was applied, the mesh in this slot might actually be an approximation.
 				//   This is difficult to fix because the mesh LOD chains are per-part and not per-instance,
 				//   need some way to keep the original copied & simplified LOD chain around...
-				if (Instance.bAllowApproximation == false && LevelLODType == ECombinedLODType::Approximated)
+				bool bAllowApproximation = (Part->bAllowApproximation && Instance.bAllowApproximation);
+				if (bAllowApproximation == false && LevelLODType == ECombinedLODType::Approximated)
 				{
 					InstanceAppendMesh = &OptimizedGeometry.SimplifiedMeshLODs.Last();
 				}
diff --git a/Engine/Source/Developer/GeometryProcessingInterfaces/Public/GeometryProcessingInterfaces/CombineMeshInstances.h b/Engine/Source/Developer/GeometryProcessingInterfaces/Public/GeometryProcessingInterfaces/CombineMeshInstances.h
--- a/Engine/Source/Developer/GeometryProcessingInterfaces/Public/GeometryProcessingInterfaces/CombineMeshInstances.h
+++ b/Engine/Source/Developer/GeometryProcessingInterfaces/Public/GeometryProcessingInterfaces/CombineMeshInstances.h
@@ -33,6 +33,18 @@ public:
 	};
 
 
+	enum class EApproximationType
+	{
+		NoConstraint = 0,
+		AxisAlignedBox = 1 << 0,
+		OrientedBox = 1 << 1,
+		SweptHull = 1 << 2,
+		ConvexHull = 1 << 3,
+		SweptProjection = 1 << 4,
+		All = 0xFFFF
+	};
+
+
 	/**
 	 * FMeshInstanceGroupData is data shared among one or more FBaseMeshInstances.
 	 * For example all instances in an ISMC can share a single FMeshInstanceGroupData.
@@ -46,6 +58,13 @@ public:
 
 		bool bPreserveUVs = false;
 		bool bAllowMerging = true;		// if false, cannot merge the geometry from this mesh with adjacent meshes to reduce triangle count
+
+		bool bAllowApproximation = true;			// if false, only Copied or Simplified LODs will be used for this Part. This flag will be combined w/ the Instance-level flag.
+
+		// ApproximationConstraint can be used to control which types of Approximation are used for LODs of this Part.
+		// This is a bitmask, any unset EApproximationType bits should be ignored by the CombineMeshInstances implementation
+		// Note however that 0 (all bits unset) is 'NoConstraint', implementations are intended to treat this as 'Allow All Types'
+		EApproximationType ApproximationConstraint = EApproximationType::NoConstraint;
 	};
 
 
@@ -58,7 +77,7 @@ public:
 		TArray<FTransform3d> TransformSequence;		// set of transforms on this instance. Often just a single transform.
 		int32 GroupDataIndex = -1;					// index into FSourceInstanceList::InstanceGroupDatas
 
-		bool bAllowApproximation = true;			// if false, only Copied or Simplified LODs will be used for this part instance
+		bool bAllowApproximation = true;			// if false, only Copied or Simplified LODs will be used for this part Instance. Will be combined w/ the GroupData flag.
 
 		// in some cases it may be desirable to have "groups" of instances which should be output as separate meshes, but
 		// be jointly processed in terms of (eg) the part LODs. If any InstanceSubsetID is non-zero, then instance subsets
