CombineMeshInstances: revert workaround that was avoiding crash in SimplifyToMinimalPlanar
#rnx

diff --git a/Engine/Plugins/Experimental/MeshModelingToolsetExp/Source/GeometryProcessingAdapters/Private/CombineMeshInstancesImpl.cpp b/Engine/Plugins/Experimental/MeshModelingToolsetExp/Source/GeometryProcessingAdapters/Private/CombineMeshInstancesImpl.cpp
--- a/Engine/Plugins/Experimental/MeshModelingToolsetExp/Source/GeometryProcessingAdapters/Private/CombineMeshInstancesImpl.cpp
+++ b/Engine/Plugins/Experimental/MeshModelingToolsetExp/Source/GeometryProcessingAdapters/Private/CombineMeshInstancesImpl.cpp
@@ -2409,8 +2409,7 @@ static void PostProcessHiddenFaceRemovedMesh(
 
 		// simplify to planar
 		FQEMSimplification Simplifier(&SubRegionMesh);
-		Simplifier.bAllowSeamCollapse = false;		// workaround for UV seam collapse issue
-		//Simplifier.bAllowSeamCollapse = true;
+		Simplifier.bAllowSeamCollapse = true;
 		Simplifier.bRetainQuadricMemory = false;
 
 		EEdgeRefineFlags BoundaryConstraint = EEdgeRefineFlags::NoFlip;
