Removing accidental button style in skin weight detail panel.

diff --git a/Engine/Plugins/Experimental/MeshModelingToolsetExp/Source/ModelingEditorUI/Private/SkinWeightDetailCustomization.cpp b/Engine/Plugins/Experimental/MeshModelingToolsetExp/Source/ModelingEditorUI/Private/SkinWeightDetailCustomization.cpp
--- a/Engine/Plugins/Experimental/MeshModelingToolsetExp/Source/ModelingEditorUI/Private/SkinWeightDetailCustomization.cpp
+++ b/Engine/Plugins/Experimental/MeshModelingToolsetExp/Source/ModelingEditorUI/Private/SkinWeightDetailCustomization.cpp
@@ -509,7 +509,6 @@ void FSkinWeightDetailCustomization::AddSelectionUI(IDetailLayoutBuilder& Detail
 				[
 					SNew(SButton)
 					.HAlign(HAlign_Center)
-					.ButtonStyle(&FAppStyle::Get().GetWidgetStyle<FButtonStyle>("VerticalToolBar.Button"))
 					.Text(LOCTEXT("AddWeightsButtonLabel", "Add"))
 					.ToolTipText(LOCTEXT("AddOpTooltip", "Add: applies the current weight plus the flood amount to the new weight."))
 					.OnClicked_Lambda([this]()
