Added component editor to directly edit skin weights on vertices.
Added ability to preview animations while editing skin weights.

diff --git a/Engine/Plugins/Experimental/MeshModelingToolsetExp/Source/MeshModelingToolsEditorOnlyExp/Private/SkeletalMesh/SkinWeightsPaintTool.cpp b/Engine/Plugins/Experimental/MeshModelingToolsetExp/Source/MeshModelingToolsEditorOnlyExp/Private/SkeletalMesh/SkinWeightsPaintTool.cpp
--- a/Engine/Plugins/Experimental/MeshModelingToolsetExp/Source/MeshModelingToolsEditorOnlyExp/Private/SkeletalMesh/SkinWeightsPaintTool.cpp
+++ b/Engine/Plugins/Experimental/MeshModelingToolsetExp/Source/MeshModelingToolsEditorOnlyExp/Private/SkeletalMesh/SkinWeightsPaintTool.cpp
@@ -174,6 +174,9 @@ void FSkinToolDeformer::Initialize(const USkeletalMeshComponent* SkeletalMeshCom
 
 	// set all vertices to be updated on first tick
 	SetAllVerticesToBeUpdated();
+
+	// record "prev" bone transforms to detect change in pose
+	PrevBoneTransforms = Component->GetComponentSpaceTransforms();
 }
 
 void FSkinToolDeformer::SetAllVerticesToBeUpdated()
@@ -189,6 +192,27 @@ void FSkinToolDeformer::UpdateVertexDeformation(USkinWeightsPaintTool* Tool)
 {
 	TRACE_CPUPROFILER_EVENT_SCOPE(SkinTool::UpdateDeformationTotal);
 
+	// if no weights have been modified, we must check for a modified pose which requires re-calculation of skinning
+	if (VerticesWithModifiedWeights.IsEmpty())
+	{
+		const TArray<FTransform>& CurrentBoneTransforms = Component->GetComponentSpaceTransforms();
+		for (int32 BoneIndex=0; BoneIndex<CurrentBoneTransforms.Num(); ++BoneIndex)
+		{
+			if (!Tool->Weights.IsBoneWeighted[BoneIndex])
+			{
+				continue;
+			}
+			
+			const FTransform& CurrentBoneTransform = CurrentBoneTransforms[BoneIndex];
+			const FTransform& PrevBoneTransform = PrevBoneTransforms[BoneIndex];
+			if (!CurrentBoneTransform.Equals(PrevBoneTransform))
+			{
+				SetAllVerticesToBeUpdated();
+				break;
+			}
+		}
+	}
+
 	if (VerticesWithModifiedWeights.IsEmpty())
 	{
 		return;
@@ -282,6 +306,9 @@ void FSkinToolDeformer::UpdateVertexDeformation(USkinWeightsPaintTool* Tool)
 
 	// empty queue of vertices to update
 	VerticesWithModifiedWeights.Reset();
+
+	// record the skeleton state we used to update the deformations
+	PrevBoneTransforms = Component->GetComponentSpaceTransforms();
 }
 
 void FSkinToolDeformer::SetVertexNeedsUpdated(int32 VertexIndex)
@@ -323,6 +350,19 @@ void FSkinToolWeights::InitializeSkinWeights(
 
 	// maintain relax-per stroke map
 	MaxFalloffPerVertexThisStroke.SetNumZeroed(NumVertices);
+
+	// maintain bool-per-bone if weighted or not
+	IsBoneWeighted.Init(false, Deformer.BoneNames.Num());
+	for (const VertexWeights& VertexData : CurrentWeights)
+	{
+		for (const FVertexBoneWeight& VertexBoneData : VertexData)
+		{
+			if (VertexBoneData.Weight > UE::AnimationCore::BoneWeightThreshold)
+			{
+				IsBoneWeighted[VertexBoneData.BoneIndex] = true;
+			}
+		}
+	}
 }
 
 void FSkinToolWeights::EditVertexWeightAndNormalize(
@@ -556,7 +596,7 @@ void FSkinToolWeights::SetWeightOfBoneOnVertex(
 	}
 }
 
-void FSkinToolWeights::ResetAfterChange()
+void FSkinToolWeights::SwapAfterChange()
 {
 	PreChangeWeights = CurrentWeights;
 
@@ -576,10 +616,9 @@ float FSkinToolWeights::SetCurrentFalloffAndGetMaxFalloffThisStroke(int32 Vertex
 	return MaxFalloffThisStroke;
 }
 
-void FSkinToolWeights::ApplyEditsToWeightMap(
-	const FMultiBoneWeightEdits& Edits, 
-	TArray<VertexWeights>& InOutWeights)
+void FSkinToolWeights::ApplyEditsToCurrentWeights(const FMultiBoneWeightEdits& Edits)
 {
+	// apply weight edits to the CurrentWeights data
 	for (const TTuple<BoneIndex, FSingleBoneWeightEdits>& BoneWeightEdits : Edits.PerBoneWeightEdits)
 	{
 		const FSingleBoneWeightEdits& WeightEdits = BoneWeightEdits.Value;
@@ -588,7 +627,34 @@ void FSkinToolWeights::ApplyEditsToWeightMap(
 		{
 			const int32 VertexID = NewWeight.Key;
 			const float Weight = NewWeight.Value;
-			SetWeightOfBoneOnVertex(BoneIndex, VertexID, Weight, InOutWeights);
+			SetWeightOfBoneOnVertex(BoneIndex, VertexID, Weight, CurrentWeights);
+		}
+	}
+
+	// weights on Bones were modified, so update IsBoneWeighted array
+	for (const TTuple<BoneIndex, FSingleBoneWeightEdits>& BoneWeightEdits : Edits.PerBoneWeightEdits)
+	{
+		const BoneIndex CurrentBoneIndex = BoneWeightEdits.Key;
+		UpdateIsBoneWeighted(CurrentBoneIndex);
+	}
+}
+
+void FSkinToolWeights::UpdateIsBoneWeighted(BoneIndex BoneToUpdate)
+{
+	IsBoneWeighted[BoneToUpdate] = false;
+	for (const VertexWeights& VertexData : CurrentWeights)
+	{
+		for (const FVertexBoneWeight& VertexBoneData : VertexData)
+		{
+			if (VertexBoneData.BoneIndex == BoneToUpdate && VertexBoneData.Weight > UE::AnimationCore::BoneWeightThreshold)
+			{
+				IsBoneWeighted[BoneToUpdate] = true;
+				break;
+			}
+		}
+		if (IsBoneWeighted[BoneToUpdate])
+		{
+			break;
 		}
 	}
 }
@@ -611,6 +677,9 @@ void FMeshSkinWeightsChange::Revert(UObject* Object)
 	{
 		Tool->ExternalUpdateWeights(Pair.Key, Pair.Value.OldWeights);
 	}
+
+	// notify dependent systems
+	Tool->OnWeightsChanged.Broadcast();
 }
 
 void FMeshSkinWeightsChange::AddBoneWeightEdit(const FSingleBoneWeightEdits& BoneWeightEdit)
@@ -706,7 +775,7 @@ void USkinWeightsPaintTool::Setup()
 	PolygonSelectionMechanic->bAddSelectionFilterPropertiesToParentTool = false;
 	PolygonSelectionMechanic->Setup(this);
 	PolygonSelectionMechanic->SetIsEnabled(false);
-	PolygonSelectionMechanic->OnSelectionChanged.AddUObject(this, &USkinWeightsPaintTool::OnSelectionModified);
+	PolygonSelectionMechanic->OnSelectionChanged.AddLambda([this](){OnSelectionChanged.Broadcast();} );
 	// only select vertices
 	PolygonSelectionMechanic->Properties->bSelectEdges = false;
 	PolygonSelectionMechanic->Properties->bSelectFaces = false;
@@ -949,11 +1018,9 @@ void USkinWeightsPaintTool::OnEndDrag(const FRay& Ray)
 	bSmoothStroke = false;
 	bStampPending = false;
 
-	// close change record
-	TUniquePtr<FMeshSkinWeightsChange> Change = EndChange();
-	GetToolManager()->BeginUndoTransaction(LOCTEXT("PaintWeightChange", "Paint skin weights."));
-	GetToolManager()->EmitObjectChange(this, MoveTemp(Change), LOCTEXT("PaintWeightChange", "Paint skin weights."));
-	GetToolManager()->EndUndoTransaction();
+	// close change, record transaction
+	const FText TransactionLabel = LOCTEXT("PaintWeightChange", "Paint skin weights.");
+	EndChange(TransactionLabel);
 }
 
 bool USkinWeightsPaintTool::OnUpdateHover(const FInputDeviceRay& DevicePos)
@@ -1179,8 +1246,10 @@ void USkinWeightsPaintTool::ApplyStamp(const FBrushStampData& Stamp)
 		{
 			// edit weight; either by "Add", "Remove", "Replace", "Multiply"
 			const float UseStrength = CalculateBrushStrengthToUse(WeightToolProperties->BrushMode);
-			EditWeightOnVertices(
+			const int32 CurrentBoneIndex = Weights.Deformer.BoneNameToIndexMap[CurrentBone];
+			EditWeightOfBoneOnVertices(
 				WeightToolProperties->BrushMode,
+				CurrentBoneIndex,
 				VerticesInStamp,
 				VertexFalloffs,
 				UseStrength,
@@ -1200,7 +1269,7 @@ void USkinWeightsPaintTool::ApplyStamp(const FBrushStampData& Stamp)
 	{
 		TRACE_CPUPROFILER_EVENT_SCOPE(SkinTool::ApplyWeightEditsToCurrentWeights);
 		// apply weights to current weights
-		Weights.ApplyEditsToWeightMap(WeightEditsFromStamp, Weights.CurrentWeights);
+		Weights.ApplyEditsToCurrentWeights(WeightEditsFromStamp);
 	}
 
 	{
@@ -1263,8 +1332,9 @@ float USkinWeightsPaintTool::CalculateBrushStrengthToUse(EWeightEditOperation Ed
 	return UseStrength;
 }
 
-void USkinWeightsPaintTool::EditWeightOnVertices(
+void USkinWeightsPaintTool::EditWeightOfBoneOnVertices(
 	EWeightEditOperation EditOperation,
+	const BoneIndex Bone,
 	const TArray<int32>& VerticesToEdit,
 	const TArray<float>& VertexFalloffs,
 	const float UseStrength,
@@ -1272,13 +1342,12 @@ void USkinWeightsPaintTool::EditWeightOnVertices(
 {
 	// spin through the vertices in the stamp and store new weight values in NewValuesFromStamp
 	// afterwards, these values are normalized while taking into consideration the user's desired changes
-	const int32 CurrentBoneIndex = Weights.Deformer.BoneNameToIndexMap[CurrentBone];
 	const int32 NumVerticesInStamp = VerticesToEdit.Num();
 	for (int32 Index = 0; Index < NumVerticesInStamp; ++Index)
 	{
 		const int32 VertexID = VerticesToEdit[Index];
 		const float UseFalloff = VertexFalloffs.IsValidIndex(Index) ? VertexFalloffs[Index] : 1.f;
-		const float ValueBeforeStroke = Weights.GetWeightOfBoneOnVertex(CurrentBoneIndex, VertexID, Weights.PreChangeWeights);
+		const float ValueBeforeStroke = Weights.GetWeightOfBoneOnVertex(Bone, VertexID, Weights.PreChangeWeights);
 
 		// calculate new weight value
 		float NewValueAfterStamp = ValueBeforeStroke;
@@ -1308,7 +1377,7 @@ void USkinWeightsPaintTool::EditWeightOnVertices(
 		// normalize the values across all bones affecting the vertices in the stamp, and record the bone edits
 		// normalization is done while holding all weights on the current bone constant so that user edits are not overwritten
 		Weights.EditVertexWeightAndNormalize(
-			CurrentBoneIndex,
+			Bone,
 			VertexID,
 			NewValueAfterStamp,
 			InOutWeightEdits);
@@ -1350,11 +1419,15 @@ void USkinWeightsPaintTool::RelaxWeightOnVertices(
 }
 
 void USkinWeightsPaintTool::ApplyWeightEditsToMesh(
+	const SkinPaintTool::FMultiBoneWeightEdits& WeightEdits,
 	const FText& TransactionLabel,
-	const SkinPaintTool::FMultiBoneWeightEdits& WeightEdits)
+	const bool bShouldTransact)
 {
-	// clear the active change
-	BeginChange();
+	if (bShouldTransact)
+	{
+		// clear the active change to start a new one
+		BeginChange();
+	}
 
 	// store weight edits in the active change
 	for (const TTuple<int32, FSingleBoneWeightEdits>& BoneWeightEdits : WeightEdits.PerBoneWeightEdits)
@@ -1367,11 +1440,11 @@ void USkinWeightsPaintTool::ApplyWeightEditsToMesh(
 	// in brush mode, this is normally called by the tool itself after a stroke)
 	ActiveChange->Apply(this);
 
-	// store active change in the transaction buffer
-	UInteractiveToolManager* ToolManager = GetToolManager();
-	ToolManager->BeginUndoTransaction(TransactionLabel);
-	ToolManager->EmitObjectChange(this, EndChange(), TransactionLabel);
-	ToolManager->EndUndoTransaction();
+	if (bShouldTransact)
+	{
+		// store active change in the transaction buffer
+		EndChange(TransactionLabel);
+	}
 }
 
 void USkinWeightsPaintTool::UpdateCurrentBone(const FName& BoneName)
@@ -1386,11 +1459,9 @@ void USkinWeightsPaintTool::GetVerticesToEdit(TArray<VertexIndex>& OutVertexIndi
 
 	//
 	// 1. Prioritize selected vertices
-	// TODO check selection manager for selected vertices here
-	const FGroupTopologySelection& Selection = PolygonSelectionMechanic->GetActiveSelection();
-	if (!Selection.SelectedCornerIDs.IsEmpty())
+	GetSelectedVertices(OutVertexIndices);
+	if (!OutVertexIndices.IsEmpty())
 	{
-		OutVertexIndices = Selection.SelectedCornerIDs.Array();
 		return;
 	}
 
@@ -1468,10 +1539,18 @@ void USkinWeightsPaintTool::BeginChange()
 	ActiveChange = MakeUnique<FMeshSkinWeightsChange>();
 }
 
-TUniquePtr<FMeshSkinWeightsChange> USkinWeightsPaintTool::EndChange()
+void USkinWeightsPaintTool::EndChange(const FText& TransactionLabel)
 {
-	Weights.ResetAfterChange();
-	return MoveTemp(ActiveChange);
+	// swap weight buffers
+	Weights.SwapAfterChange();
+	
+	// record transaction
+	GetToolManager()->BeginUndoTransaction(TransactionLabel);
+	GetToolManager()->EmitObjectChange(this, MoveTemp(ActiveChange), TransactionLabel);
+	GetToolManager()->EndUndoTransaction();
+
+	// notify dependent systems
+	OnWeightsChanged.Broadcast();
 }
 
 void USkinWeightsPaintTool::ExternalUpdateWeights(const int32 BoneIndex, const TMap<int32, float>& NewValues)
@@ -1484,11 +1563,13 @@ void USkinWeightsPaintTool::ExternalUpdateWeights(const int32 BoneIndex, const T
 		Weights.SetWeightOfBoneOnVertex(BoneIndex, VertexID, Weight, Weights.PreChangeWeights);
 	}
 
-	const FName BoneName = Weights.Deformer.BoneNames[BoneIndex];
+	const FName BoneName = GetBoneNameFromIndex(BoneIndex);
 	if (BoneName == CurrentBone)
 	{
 		UpdateCurrentBoneVertexColors();
 	}
+
+	Weights.UpdateIsBoneWeighted(BoneIndex);
 }
 
 void FSkinMirrorData::RegenerateMirrorData(
@@ -1684,7 +1765,8 @@ void USkinWeightsPaintTool::MirrorWeights(EAxis::Type Axis, EMirrorDirection Dir
 
 	// apply the changes
 	const FText TransactionLabel = LOCTEXT("MirrorWeightChange", "Mirror skin weights.");
-	ApplyWeightEditsToMesh(TransactionLabel, WeightEditsFromMirroring);
+	constexpr bool bShouldTransact = true;
+	ApplyWeightEditsToMesh(WeightEditsFromMirroring, TransactionLabel, bShouldTransact);
 
 	// warn if some vertices were not mirrored
 	if (!MirrorData.GetAllVerticesMirrored())
@@ -1701,6 +1783,7 @@ void USkinWeightsPaintTool::FloodWeights(const float Weight, const EWeightEditOp
 	}
 
 	// flood the weights on selected vertices
+	const int32 CurrentBoneIndex = Weights.Deformer.BoneNameToIndexMap[CurrentBone];
 	FMultiBoneWeightEdits WeightEditsFromFlood;
 	TArray<VertexIndex> VerticesToEdit;
 	GetVerticesToEdit(VerticesToEdit);
@@ -1711,12 +1794,35 @@ void USkinWeightsPaintTool::FloodWeights(const float Weight, const EWeightEditOp
 	}
 	else
 	{
-		EditWeightOnVertices(FloodMode, VerticesToEdit, VertexFalloffs, Weight, WeightEditsFromFlood);	
+		EditWeightOfBoneOnVertices(FloodMode, CurrentBoneIndex, VerticesToEdit, VertexFalloffs, Weight, WeightEditsFromFlood);	
 	}
 
 	// apply the changes
 	const FText TransactionLabel = LOCTEXT("FloodWeightChange", "Flood skin weights.");
-	ApplyWeightEditsToMesh(TransactionLabel, WeightEditsFromFlood);
+	constexpr bool bShouldTransact = true;
+	ApplyWeightEditsToMesh(WeightEditsFromFlood, TransactionLabel, bShouldTransact);
+}
+
+void USkinWeightsPaintTool::SetBoneWeightOnVertices(
+	BoneIndex Bone,
+	const float Weight,
+	const TArray<VertexIndex>& VerticesToEdit,
+	const bool bShouldTransact)
+{
+	// create weight edits from setting the weight directly
+	FMultiBoneWeightEdits DirectWeightEdits;
+	const TArray<float> VertexFalloffs = {}; // no falloff
+	EditWeightOfBoneOnVertices(
+				EWeightEditOperation::Replace,
+				Bone,
+				VerticesToEdit,
+				VertexFalloffs,
+				Weight,
+				DirectWeightEdits);
+	
+	// apply the changes
+	const FText TransactionLabel = LOCTEXT("SetWeightChange", "Set skin weights directly.");
+	ApplyWeightEditsToMesh(DirectWeightEdits, TransactionLabel, bShouldTransact);
 }
 
 void USkinWeightsPaintTool::PruneWeights(float Threshold)
@@ -1751,7 +1857,8 @@ void USkinWeightsPaintTool::PruneWeights(float Threshold)
 
 	// apply the changes
 	const FText TransactionLabel = LOCTEXT("PruneWeightValuesChange", "Prune skin weights.");
-	ApplyWeightEditsToMesh(TransactionLabel, WeightEditsFromPrune);
+	constexpr bool bShouldTransact = true;
+	ApplyWeightEditsToMesh(WeightEditsFromPrune, TransactionLabel, bShouldTransact);
 }
 
 void USkinWeightsPaintTool::AverageWeights()
@@ -1829,7 +1936,8 @@ void USkinWeightsPaintTool::AverageWeights()
 
 	// apply the changes
 	const FText TransactionLabel = LOCTEXT("AverageWeightValuesChange", "Average skin weights.");
-	ApplyWeightEditsToMesh(TransactionLabel, WeightEditsFromAveraging);
+	constexpr bool bShouldTransact = true;
+	ApplyWeightEditsToMesh(WeightEditsFromAveraging, TransactionLabel, bShouldTransact);
 }
 
 void USkinWeightsPaintTool::NormalizeWeights()
@@ -1850,7 +1958,8 @@ void USkinWeightsPaintTool::NormalizeWeights()
 
 	// apply the changes
 	const FText TransactionLabel = LOCTEXT("NormalizeWeightValuesChange", "Normalize skin weights.");
-	ApplyWeightEditsToMesh(TransactionLabel, WeightEditsFromNormalization);
+	constexpr bool bShouldTransact = true;
+	ApplyWeightEditsToMesh(WeightEditsFromNormalization, TransactionLabel, bShouldTransact);
 }
 
 void USkinWeightsPaintTool::HandleSkeletalMeshModified(const TArray<FName>& InBoneNames, const ESkeletalMeshNotifyType InNotifyType)
@@ -1890,10 +1999,6 @@ void USkinWeightsPaintTool::HandleSkeletalMeshModified(const TArray<FName>& InBo
 	}
 }
 
-void USkinWeightsPaintTool::OnSelectionModified()
-{
-}
-
 void USkinWeightsPaintTool::ToggleEditingMode()
 {
 	Weights.Deformer.SetAllVerticesToBeUpdated();
@@ -1901,6 +2006,65 @@ void USkinWeightsPaintTool::ToggleEditingMode()
 	PolygonSelectionMechanic->SetIsEnabled(WeightToolProperties->EditingMode == EWeightEditMode::Vertices);
 }
 
+void USkinWeightsPaintTool::GetSelectedVertices(TArray<int32>& OutVertexIndices) const
+{
+	const FGroupTopologySelection& Selection = PolygonSelectionMechanic->GetActiveSelection();
+	if (!Selection.SelectedCornerIDs.IsEmpty())
+	{
+		OutVertexIndices = Selection.SelectedCornerIDs.Array();
+	}
+}
+
+void USkinWeightsPaintTool::GetInfluences(const TArray<int32>& VertexIndices, TArray<BoneIndex>& OutBoneIndices)
+{
+	for (const int32 SelectedVertex : VertexIndices)
+	{
+		for (const FVertexBoneWeight& VertexBoneData : Weights.CurrentWeights[SelectedVertex])
+		{
+			if (VertexBoneData.Weight >= MinimumWeightThreshold)
+			{
+				OutBoneIndices.AddUnique(VertexBoneData.BoneIndex);
+			}
+		}
+	}
+	
+	// sort hierarchically (bone indices are sorted root to leaf)
+	OutBoneIndices.Sort([](BoneIndex A, BoneIndex B) {return A < B;});
+}
+
+float USkinWeightsPaintTool::GetAverageWeightOnBone(
+	const BoneIndex InBoneIndex,
+	const TArray<int32>& VertexIndices)
+{
+	float TotalWeight = 0.f;
+	float NumVerticesInfluencedByBone = 0.f;
+	
+	for (const int32 SelectedVertex : VertexIndices)
+	{
+		for (const FVertexBoneWeight& VertexBoneData : Weights.CurrentWeights[SelectedVertex])
+		{
+			if (VertexBoneData.BoneIndex == InBoneIndex)
+			{
+				++NumVerticesInfluencedByBone;
+				TotalWeight += VertexBoneData.Weight;
+			}
+		}
+	}
+
+	return NumVerticesInfluencedByBone > 0 ? TotalWeight / NumVerticesInfluencedByBone : TotalWeight;
+}
+
+FName USkinWeightsPaintTool::GetBoneNameFromIndex(BoneIndex InIndex) const
+{
+	const TArray<FName>& Names = Weights.Deformer.BoneNames;
+	if (Names.IsValidIndex(InIndex))
+	{
+		return Names[InIndex];
+	}
+
+	return NAME_None;
+}
+
 void USkinWeightsPaintTool::OnPropertyModified(UObject* ModifiedObject, FProperty* ModifiedProperty)
 {
 	Super::OnPropertyModified(ModifiedObject, ModifiedProperty);
diff --git a/Engine/Plugins/Experimental/MeshModelingToolsetExp/Source/MeshModelingToolsEditorOnlyExp/Public/SkeletalMesh/SkinWeightsPaintTool.h b/Engine/Plugins/Experimental/MeshModelingToolsetExp/Source/MeshModelingToolsEditorOnlyExp/Public/SkeletalMesh/SkinWeightsPaintTool.h
--- a/Engine/Plugins/Experimental/MeshModelingToolsetExp/Source/MeshModelingToolsEditorOnlyExp/Public/SkeletalMesh/SkinWeightsPaintTool.h
+++ b/Engine/Plugins/Experimental/MeshModelingToolsetExp/Source/MeshModelingToolsEditorOnlyExp/Public/SkeletalMesh/SkinWeightsPaintTool.h
@@ -106,6 +106,8 @@ namespace SkinPaintTool
 		TArray<FVector> RefPoseVertexPositions;
 		// inverted, component space ref pose transform of each bone
 		TArray<FTransform> InvCSRefPoseTransforms;
+		// bones transforms used in last deformation update
+		TArray<FTransform> PrevBoneTransforms;
 		// bone index to bone name
 		TArray<FName> BoneNames;
 		TMap<FName, BoneIndex> BoneNameToIndexMap;
@@ -182,11 +184,13 @@ namespace SkinPaintTool
 			const float Weight,
 			TArray<VertexWeights>& InOutVertexData);
 
-		void ResetAfterChange();
+		void SwapAfterChange();
 
 		float SetCurrentFalloffAndGetMaxFalloffThisStroke(int32 VertexID, float CurrentStrength);
 
-		void ApplyEditsToWeightMap(const FMultiBoneWeightEdits& Edits, TArray<VertexWeights>& InOutWeights);
+		void ApplyEditsToCurrentWeights(const FMultiBoneWeightEdits& Edits);
+
+		void UpdateIsBoneWeighted(BoneIndex BoneToUpdate);
 		
 		// double-buffer of the entire weight matrix (stored sparsely for fast deformation)
 		// "Pre" is state of weights at stroke start
@@ -200,6 +204,9 @@ namespace SkinPaintTool
 		// of modification (add/replace/relax etc) that was encountered for the duration of the stroke.
 		TArray<float> MaxFalloffPerVertexThisStroke;
 
+		// record which bones have any weight assigned to them
+		TArray<bool> IsBoneWeighted;
+
 		// update deformation when vertex weights are modified
 		FSkinToolDeformer Deformer;
 	};
@@ -341,6 +348,7 @@ public:
 	virtual bool SupportsWorldSpaceFocusBox() override { return true; }
 	virtual FBox GetWorldSpaceFocusBox() override;
 
+	// using when ToolChange is applied via Undo/Redo
 	void ExternalUpdateWeights(const int32 BoneIndex, const TMap<int32, float>& IndexValues);
 
 	// weight editing operations (selection based)
@@ -349,10 +357,56 @@ public:
 	void PruneWeights(const float Threshold);
 	void AverageWeights();
 	void NormalizeWeights();
+	
+	// method to set weights directly (numeric input, for example)
+	void SetBoneWeightOnVertices(
+		BoneIndex Bone,
+		const float Weight,
+		const TArray<VertexIndex>& VerticesToEdit,
+		const bool bShouldTransact);
 
 	// toggle brush / selection mode
 	void ToggleEditingMode();
 
+	// get a list of currently selected vertices
+	void GetSelectedVertices(TArray<int32>& OutVertexIndices) const;
+
+	// get the average weight value of each influence on the given vertices
+	void GetInfluences(const TArray<int32>& VertexIndices, TArray<BoneIndex>& OutBoneIndices);
+
+	// get the average weight value of a single bone on the given vertices
+	float GetAverageWeightOnBone(const BoneIndex InBoneIndex, const TArray<int32>& VertexIndices);
+
+	// convert an index to a name
+	FName GetBoneNameFromIndex(BoneIndex InIndex) const;
+
+	// HOW TO EDIT WEIGHTS WITH UNDO/REDO:
+	//
+	// Live Edits:
+	// For multiple weight editing operations that need to be grouped into a single transaction, like dragging a slider or
+	// dragging a brush, you must call BeginChange(), then ApplyWeightsEditsToMesh(bShouldTransact=false), then EndChange().
+	// All the edits are stored into the "ActiveChange" and applied as a single transaction. Deformations and vertex
+	// colors will be updated throughout the duration of the change.
+	//
+	// Single Edits:
+	// For all one-and-done edits, you can call ApplyWeightEditsToMesh(bShouldTransact=True).
+	// It will Begin/End the change and create a transaction for it.
+	//
+	void BeginChange();
+	void EndChange(const FText& TransactionLabel);
+	void ApplyWeightEditsToMesh(
+		const SkinPaintTool::FMultiBoneWeightEdits& WeightEdits,
+		const FText& TransactionLabel,
+		const bool bShouldTransact);
+
+	// called whenever the selection is modified
+	DECLARE_MULTICAST_DELEGATE(FOnSelectionChanged);
+	FOnSelectionChanged OnSelectionChanged;
+
+	// called whenever the weights are modified
+	DECLARE_MULTICAST_DELEGATE(FOnWeightsChanged);
+	FOnWeightsChanged OnWeightsChanged;
+
 protected:
 
 	virtual void ApplyStamp(const FBrushStampData& Stamp);
@@ -376,8 +430,9 @@ protected:
 
 	// modify vertex weights according to the specified operation,
 	// generating bone weight edits to be stored in a transaction
-	void EditWeightOnVertices(
+	void EditWeightOfBoneOnVertices(
 		EWeightEditOperation EditOperation,
+		const BoneIndex Bone,
 		const TArray<int32>& VerticesToEdit,
 		const TArray<float>& VertexFalloffs,
 		const float UseStrength,
@@ -409,6 +464,9 @@ protected:
 	// cached mirror data
 	SkinPaintTool::FSkinMirrorData MirrorData;
 
+	// storage for weight edits in the current transaction
+	TUniquePtr<SkinPaintTool::FMeshSkinWeightsChange> ActiveChange;
+
 	// Smooth weights data source and operator
 	TUniquePtr<UE::Geometry::TBoneWeightsDataSource<int32, float>> SmoothWeightsDataSource;
 	TUniquePtr<UE::Geometry::TSmoothBoneWeights<int32, float>> SmoothWeightsOp;
@@ -418,13 +476,6 @@ protected:
 	FVector4f WeightToColor(float Value) const;
 	bool bVisibleWeightsValid = false;
 
-	// weight editing transactions
-	void BeginChange();
-	TUniquePtr<SkinPaintTool::FMeshSkinWeightsChange> EndChange();
-	TUniquePtr<SkinPaintTool::FMeshSkinWeightsChange> ActiveChange;
-	// used to apply manual weight modifications (not from the brush)
-	void ApplyWeightEditsToMesh(const FText& TransactionLabel, const SkinPaintTool::FMultiBoneWeightEdits& WeightEdits);
-
 	// which bone are we currently painting?
 	void UpdateCurrentBone(const FName &BoneName);
 	FName CurrentBone = NAME_None;
@@ -445,7 +496,6 @@ protected:
 	TObjectPtr<UPolygonSelectionMechanic> PolygonSelectionMechanic;
 	TUniquePtr<UE::Geometry::FDynamicMeshAABBTree3> MeshSpatial = nullptr;
 	TUniquePtr<UE::Geometry::FTriangleGroupTopology> SelectionTopology = nullptr;
-	void OnSelectionModified();
 
 	UPROPERTY()
 	TWeakObjectPtr<USkeletalMeshEditorContextObjectBase> EditorContext = nullptr;
diff --git a/Engine/Plugins/Experimental/MeshModelingToolsetExp/Source/ModelingEditorUI/Private/SkinWeightDetailCustomization.cpp b/Engine/Plugins/Experimental/MeshModelingToolsetExp/Source/ModelingEditorUI/Private/SkinWeightDetailCustomization.cpp
--- a/Engine/Plugins/Experimental/MeshModelingToolsetExp/Source/ModelingEditorUI/Private/SkinWeightDetailCustomization.cpp
+++ b/Engine/Plugins/Experimental/MeshModelingToolsetExp/Source/ModelingEditorUI/Private/SkinWeightDetailCustomization.cpp
@@ -11,6 +11,8 @@
 #include "UObject/UnrealTypePrivate.h"
 #include "Widgets/Input/SSpinBox.h"
 #include "Widgets/Input/SButton.h"
+#include "Widgets/Input/SNumericEntryBox.h"
+#include "Widgets/Input/SSlider.h"
 
 #define LOCTEXT_NAMESPACE "SkinWeightToolSettingsEditor"
 
@@ -621,6 +623,127 @@ void FSkinWeightDetailCustomization::AddSelectionUI(IDetailLayoutBuilder& Detail
 			]
 		]
 	];
+
+	// VERTEX EDITOR category
+	EditWeightsCategory.AddCustomRow(LOCTEXT("VertexEditorRow", "Vertex Editor"), false)
+	.WholeRowContent()
+	[
+		SNew(SVertexWeightEditor, SkinToolSettings->WeightTool)
+	];
+}
+
+void SVertexWeightItem::Construct(const FArguments& InArgs, const TSharedRef<STableViewBase>& OwnerTableView)
+{
+	Element = InArgs._Element;
+	ParentTable = InArgs._ParentTable;
+	SMultiColumnTableRow<TSharedPtr<FWeightEditorElement>>::Construct(FSuperRowType::FArguments(), OwnerTableView);
+}
+
+TSharedRef<SWidget> SVertexWeightItem::GenerateWidgetForColumn(const FName& ColumnName)
+{
+	if (ColumnName == "Bone")
+	{
+		const FName BoneName = ParentTable->Tool->GetBoneNameFromIndex(Element->BoneIndex);
+		return SNew(STextBlock).Text(FText::FromName(BoneName));
+	}
+
+	if (ColumnName == "Weight")
+	{
+		// add a buffer to the slider to prevent ever fully getting a value to 1 or 0 using the slider alone
+		// doing so will remove other influences and cause the slider to no longer function as all other influences
+		// will be culled by normalization thus making the slider "stuck" at full value.
+		constexpr float SliderBuffer = 0.001f;
+		
+		return SNew(SNumericEntryBox<float>)
+		.AllowSpin(true)
+		.MinSliderValue(SliderBuffer)
+		.MinValue(0.f)
+		.MaxSliderValue(1.0f-SliderBuffer)
+		.MaxValue(1.f)
+		.Value_Lambda([this]()
+		{
+			return ParentTable->Tool->GetAverageWeightOnBone(Element->BoneIndex, ParentTable->SelectedVertices);
+		})
+		.OnValueChanged_Lambda([this](float NewValue)
+		{
+			TArray<VertexIndex> VerticesToEdit;
+			ParentTable->Tool->GetSelectedVertices(VerticesToEdit);
+			ParentTable->Tool->SetBoneWeightOnVertices(Element->BoneIndex, NewValue, VerticesToEdit, !bInTransaction);
+		})
+		.OnValueCommitted_Lambda([this](float NewValue, ETextCommit::Type CommitType)
+		{
+			bInTransaction = false;
+		})
+		.OnBeginSliderMovement_Lambda([this]()
+		{
+			ParentTable->Tool->BeginChange();
+			bInTransaction = true;
+		})
+		.OnEndSliderMovement_Lambda([this](float)
+		{
+			const FText TransactionLabel = LOCTEXT("DirectWeightChange", "Set weights on vertices.");
+			ParentTable->Tool->EndChange(TransactionLabel);
+			bInTransaction = false;
+		})
+		.ToolTipText(LOCTEXT("WeightSliderToolTip", "Set the weight on this bone for the selected vertices."));
+	}
+
+	checkNoEntry();
+	return SNullWidget::NullWidget;
+}
+
+void SVertexWeightEditor::Construct(const FArguments& InArgs, USkinWeightsPaintTool* InSkinTool)
+{
+	Tool = InSkinTool;
+	Tool->OnSelectionChanged.AddLambda([this](){ RefreshView(); });
+	Tool->OnWeightsChanged.AddLambda([this](){ RefreshView(); });
+	
+	ChildSlot
+	[
+		SNew(SBox)
+		[
+			SAssignNew( ListView, SWeightEditorListViewType )
+			.SelectionMode(ESelectionMode::Single)
+			.ListItemsSource( &ListViewItems )
+			.OnGenerateRow_Lambda([this](TSharedPtr<FWeightEditorElement> Element, const TSharedRef<STableViewBase>& OwnerTableView)
+			{
+				return SNew(SVertexWeightItem, OwnerTableView).Element(Element).ParentTable(SharedThis(this));
+			})
+			.HeaderRow
+			(
+				SNew(SHeaderRow)
+				+ SHeaderRow::Column("Bone").DefaultLabel(NSLOCTEXT("WeightEditorBoneColumn", "Bone", "Bone"))
+				+ SHeaderRow::Column("Weight").DefaultLabel(NSLOCTEXT("WeightEditorWeightColumn", "Weight (Average)", "Weight (Average)"))
+			)
+		]
+	];
+
+	RefreshView();
+}
+
+void SVertexWeightEditor::RefreshView()
+{
+	if (!Tool.IsValid())
+	{
+		return; 
+	}
+
+	// get list of selected vertex indices
+	SelectedVertices.Reset();
+	Tool->GetSelectedVertices(SelectedVertices);
+	
+	// get all bones affecting the selected vertices
+	TArray<int32> Influences;
+	Tool->GetInfluences(SelectedVertices, Influences);
+
+	// generate list view items
+	ListViewItems.Reset();
+	for (const int32 InfluenceIndex : Influences)
+	{
+		ListViewItems.Add(MakeShareable(new FWeightEditorElement(InfluenceIndex)));
+	}
+	
+	ListView->RequestListRefresh();
 }
 
 #undef LOCTEXT_NAMESPACE
diff --git a/Engine/Plugins/Experimental/MeshModelingToolsetExp/Source/ModelingEditorUI/Public/SkinWeightDetailCustomization.h b/Engine/Plugins/Experimental/MeshModelingToolsetExp/Source/ModelingEditorUI/Public/SkinWeightDetailCustomization.h
--- a/Engine/Plugins/Experimental/MeshModelingToolsetExp/Source/ModelingEditorUI/Public/SkinWeightDetailCustomization.h
+++ b/Engine/Plugins/Experimental/MeshModelingToolsetExp/Source/ModelingEditorUI/Public/SkinWeightDetailCustomization.h
@@ -5,7 +5,10 @@
 #include "CoreMinimal.h"
 #include "IDetailCustomization.h"
 #include "UObject/WeakObjectPtr.h"
+#include "Widgets/Views/STableRow.h"
+#include "Widgets/Views/SListView.h"
 
+class USkinWeightsPaintTool;
 class USkinWeightsPaintToolProperties;
 
 class FSkinWeightDetailCustomization : public IDetailCustomization
@@ -33,3 +36,55 @@ private:
 	static float WeightEditVerticalPadding;
 	static float WeightEditHorizontalPadding;
 };
+
+struct FWeightEditorElement
+{
+	int32 BoneIndex;
+
+	FWeightEditorElement(int32 InBoneIndex) : BoneIndex(InBoneIndex) {}
+};
+
+class SVertexWeightEditor;
+class SVertexWeightItem : public SMultiColumnTableRow<TSharedPtr<FWeightEditorElement>>
+{
+public:
+
+	SLATE_BEGIN_ARGS(SVertexWeightItem) {}
+	SLATE_ARGUMENT(TSharedPtr<FWeightEditorElement>, Element)
+	SLATE_ARGUMENT(TSharedPtr<SVertexWeightEditor>, ParentTable)
+	SLATE_END_ARGS()
+
+	virtual ~SVertexWeightItem() override {};
+
+	void Construct(const FArguments& InArgs, const TSharedRef<STableViewBase>& OwnerTableView);
+
+	virtual TSharedRef<SWidget> GenerateWidgetForColumn(const FName& ColumnName) override;
+
+private:
+	TSharedPtr<FWeightEditorElement> Element;
+	TSharedPtr<SVertexWeightEditor> ParentTable;
+	bool bInTransaction = false;
+};
+
+typedef SListView< TSharedPtr<FWeightEditorElement> > SWeightEditorListViewType;
+
+class SVertexWeightEditor : public SCompoundWidget
+{
+public:
+	SLATE_BEGIN_ARGS(SVertexWeightEditor) {}
+	SLATE_END_ARGS()
+
+	void Construct(const FArguments& InArgs, USkinWeightsPaintTool* InSkinTool);
+	
+	void RefreshView();
+
+private:
+
+	// the vertex weight list view
+	TSharedPtr<SWeightEditorListViewType> ListView;
+	TArray< TSharedPtr<FWeightEditorElement> > ListViewItems;
+	TWeakObjectPtr<USkinWeightsPaintTool> Tool;
+	TArray<int32> SelectedVertices;
+
+	friend SVertexWeightItem;
+};
\ No newline at end of file
