Skeleton Editing Tool: fixed preview mesh display

#rnx

diff --git a/Engine/Plugins/Experimental/MeshModelingToolsetExp/Source/MeshModelingToolsEditorOnlyExp/Private/SkeletalMesh/SkeletonEditingTool.cpp b/Engine/Plugins/Experimental/MeshModelingToolsetExp/Source/MeshModelingToolsEditorOnlyExp/Private/SkeletalMesh/SkeletonEditingTool.cpp
--- a/Engine/Plugins/Experimental/MeshModelingToolsetExp/Source/MeshModelingToolsEditorOnlyExp/Private/SkeletalMesh/SkeletonEditingTool.cpp
+++ b/Engine/Plugins/Experimental/MeshModelingToolsetExp/Source/MeshModelingToolsEditorOnlyExp/Private/SkeletalMesh/SkeletonEditingTool.cpp
@@ -13,6 +13,7 @@
 #include "Components/SkeletalMeshComponent.h"
 #include "UnrealClient.h"
 #include "HitProxies.h"
+#include "ToolSetupUtil.h"
 #include "BaseBehaviors/ClickDragBehavior.h"
 #include "BaseGizmos/GizmoViewContext.h"
 
@@ -131,6 +132,8 @@ void USkeletonEditingTool::Setup()
 
 		PreviewMesh->SetTransform(UE::ToolTarget::GetLocalToWorldTransform(Target));
 
+		ToolSetupUtil::ApplyRenderingConfigurationToPreview(PreviewMesh, Target);
+		PreviewMesh->SetTangentsMode(EDynamicMeshComponentTangentsMode::AutoCalculated);
 		PreviewMesh->ReplaceMesh(UE::ToolTarget::GetDynamicMeshCopy(Target));
 
 		const FComponentMaterialSet MaterialSet = UE::ToolTarget::GetMaterialSet(Target);
