Fix missing transform on initial preview mesh for self union tool, which was causing the initial preview (before first union was computed) to appear at the world origin

diff --git a/Engine/Plugins/Experimental/MeshModelingToolsetExp/Source/MeshModelingToolsExp/Private/SelfUnionMeshesTool.cpp b/Engine/Plugins/Experimental/MeshModelingToolsetExp/Source/MeshModelingToolsExp/Private/SelfUnionMeshesTool.cpp
--- a/Engine/Plugins/Experimental/MeshModelingToolsetExp/Source/MeshModelingToolsExp/Private/SelfUnionMeshesTool.cpp
+++ b/Engine/Plugins/Experimental/MeshModelingToolsetExp/Source/MeshModelingToolsExp/Private/SelfUnionMeshesTool.cpp
@@ -149,6 +149,7 @@ void USelfUnionMeshesTool::ConvertInputsAndSetPreviewMaterials(bool bSetPreviewM
 	if (bSetPreviewMesh)
 	{
 		Preview->PreviewMesh->UpdatePreview(CombinedSourceMeshes.Get());
+		Preview->PreviewMesh->SetTransform(FTransform(CombinedCenter));
 	}
 }
 
