Skipping re-triangulation of source LODs that should have their UVs preserved

diff --git a/Engine/Plugins/Experimental/MeshModelingToolsetExp/Source/GeometryProcessingAdapters/Private/CombineMeshInstancesImpl.cpp b/Engine/Plugins/Experimental/MeshModelingToolsetExp/Source/GeometryProcessingAdapters/Private/CombineMeshInstancesImpl.cpp
--- a/Engine/Plugins/Experimental/MeshModelingToolsetExp/Source/GeometryProcessingAdapters/Private/CombineMeshInstancesImpl.cpp
+++ b/Engine/Plugins/Experimental/MeshModelingToolsetExp/Source/GeometryProcessingAdapters/Private/CombineMeshInstancesImpl.cpp
@@ -1425,11 +1425,15 @@ void ComputeMeshApproximations(
 		// try remeshing the last Source LOD to reduce it's triangle count, by removing spurious geometry
 		if (CombineOptions.bRetriangulateSourceLODs)
 		{
-			for (int32 SourceLODIndex = CombineOptions.bRetriangulateSourceLODs; SourceLODIndex < NumSourceLODs; ++SourceLODIndex)
+			for (int32 SourceLODIndex = CombineOptions.StartRetriangulateSourceLOD; SourceLODIndex < NumSourceLODs; ++SourceLODIndex)
 			{
-				PlanarRetriangulatePartMesh(
-					SourceGeo.SourceMeshLODs[SourceLODIndex],
-					CombineOptions.SimplifyBaseTolerance, AngleThresholdDeg);
+				if (!(Part->bPreserveUVs || CombineOptions.bSimplifyPreserveUVs || (SourceLODIndex <= CombineOptions.PreserveUVLODLevel)))
+				{
+					PlanarRetriangulatePartMesh(
+						SourceGeo.SourceMeshLODs[SourceLODIndex],
+						CombineOptions.SimplifyBaseTolerance,
+						AngleThresholdDeg);
+				}
 			}
 		}
 
