Support Geometry Selection in the Tangents tool and fix some bugs

diff --git a/Engine/Plugins/Experimental/MeshModelingToolsetExp/Source/MeshModelingToolsEditorOnlyExp/Private/MeshTangentsTool.cpp b/Engine/Plugins/Experimental/MeshModelingToolsetExp/Source/MeshModelingToolsEditorOnlyExp/Private/MeshTangentsTool.cpp
--- a/Engine/Plugins/Experimental/MeshModelingToolsetExp/Source/MeshModelingToolsEditorOnlyExp/Private/MeshTangentsTool.cpp
+++ b/Engine/Plugins/Experimental/MeshModelingToolsetExp/Source/MeshModelingToolsEditorOnlyExp/Private/MeshTangentsTool.cpp
@@ -22,6 +22,13 @@
 #include "ModelingToolTargetUtil.h"
 #include "ToolTargetManager.h"
 
+#include "PropertySets/GeometrySelectionVisualizationProperties.h"
+#include "GroupTopology.h"
+#include "DynamicMeshEditor.h"
+#include "Selection/GeometrySelectionVisualization.h"
+#include "Selection/StoredMeshSelectionUtil.h"
+#include "Selections/GeometrySelectionUtil.h"
+
 #include UE_INLINE_GENERATED_CPP_BY_NAME(MeshTangentsTool)
 
 using namespace UE::Geometry;
@@ -48,6 +55,23 @@ USingleSelectionMeshEditingTool* UMeshTangentsToolBuilder::CreateNewTool(const F
 	return NewObject<UMeshTangentsTool>(SceneState.ToolManager);
 }
 
+void UMeshTangentsToolBuilder::InitializeNewTool(USingleSelectionMeshEditingTool* NewTool, const FToolBuilderState& SceneState) const
+{
+	UToolTarget* Target = SceneState.TargetManager->BuildFirstSelectedTargetable(SceneState, GetTargetRequirements());
+	check(Target);
+	NewTool->SetTarget(Target);
+	NewTool->SetWorld(SceneState.World);
+
+	if (UMeshTangentsTool* TangentsTool = Cast<UMeshTangentsTool>(NewTool))
+	{
+		UE::Geometry::FGeometrySelection Selection;
+		if (UE::Geometry::GetCurrentGeometrySelectionForTarget(SceneState, Target, Selection))
+		{
+			TangentsTool->SetGeometrySelection(MoveTemp(Selection));
+		}
+	}
+}
+
 bool UMeshTangentsToolBuilder::CanBuildTool(const FToolBuilderState& SceneState) const
 {
 	return USingleSelectionMeshEditingToolBuilder::CanBuildTool(SceneState) &&
@@ -103,7 +127,59 @@ void UMeshTangentsTool::Setup()
 	Settings->WatchProperty(Settings->LineThickness, [this](float) { bThicknessDirty = true; });
 	Settings->WatchProperty(Settings->bShowTangents, [this](bool) { bVisibilityChanged = true; });
 	Settings->WatchProperty(Settings->bShowNormals, [this](bool) { bVisibilityChanged = true; });
-	Settings->WatchProperty(Settings->bCompareWithMikkt, [this](bool) { ComputeMikkTDeviations(nullptr); });
+	Settings->WatchProperty(Settings->bCompareWithMikkt, [this](bool) { ComputeMikkTDeviations(ComputeDegenerateTris()); });
+
+	if (InputGeometrySelection.IsEmpty() == false)
+	{
+		GeometrySelectionVizProperties = NewObject<UGeometrySelectionVisualizationProperties>(this);
+		GeometrySelectionVizProperties->RestoreProperties(this);
+		AddToolPropertySource(GeometrySelectionVizProperties);
+		GeometrySelectionVizProperties->Initialize(this);
+		GeometrySelectionVizProperties->SelectionElementType = static_cast<EGeometrySelectionElementType>(InputGeometrySelection.ElementType);
+		GeometrySelectionVizProperties->SelectionTopologyType = static_cast<EGeometrySelectionTopologyType>(InputGeometrySelection.TopologyType);
+		GeometrySelectionVizProperties->bEnableShowEdgeSelectionVertices = true;
+		// TODO Enable this but note we need to compute a ROI which only includes triangles incident to the
+		//      polygroup feature eg do not include all triangles in the groups incident to a polygroup edge
+		//GeometrySelectionVizProperties->bEnableShowTriangleROIBorder = true;
+
+		// Compute group topology if the selection has Polygroup topology, and do nothing otherwise
+		// Currently it is only possible to make a polygroup geometry selection using polygroup set stored directly in the mesh
+		FGroupTopology GroupTopology(InputMesh.Get(), InputGeometrySelection.TopologyType == EGeometryTopologyType::Polygroup);
+
+		// Compute the overlay selection and a proxy triangle vertex selection used to make edge selections behave like
+		// vertex selections. See :EdgeSelectionsBehaveLikeVertexSelections
+		if (InputGeometrySelection.TopologyType == EGeometryTopologyType::Polygroup)
+		{
+			ConvertPolygroupSelectionToIncidentOverlaySelection(
+				*InputMesh,
+				GroupTopology,
+				InputGeometrySelection,
+				EditTriangles,
+				EditVertices,
+				&TriangleVertexGeometrySelection);
+		}
+		else
+		{
+			ConvertTriangleSelectionToOverlaySelection(
+				*InputMesh,
+				InputGeometrySelection,
+				EditTriangles,
+				EditVertices,
+				&TriangleVertexGeometrySelection);
+		}
+
+		// Setup input geometry selection visualization
+		FTransform ApplyTransform = UE::ToolTarget::GetLocalToWorldTransform(Target);
+		GeometrySelectionViz = NewObject<UPreviewGeometry>(this);
+		GeometrySelectionViz->CreateInWorld(GetTargetWorld(), ApplyTransform);
+		InitializeGeometrySelectionVisualization(
+			GeometrySelectionViz,
+			GeometrySelectionVizProperties,
+			*InputMesh,
+			InputGeometrySelection,
+			&GroupTopology,
+			!TriangleVertexGeometrySelection.IsEmpty() ? &TriangleVertexGeometrySelection : nullptr);
+	}
 
 	PreviewGeometry = NewObject<UPreviewGeometry>(this);
 	PreviewGeometry->CreateInWorld(TargetActor->GetWorld(), PreviewMesh->GetTransform());
@@ -136,6 +212,11 @@ void UMeshTangentsTool::Setup()
 		EToolMessageLevel::UserNotification);
 }
 
+void UMeshTangentsTool::SetGeometrySelection(UE::Geometry::FGeometrySelection&& SelectionIn)
+{
+	InputGeometrySelection = MoveTemp(SelectionIn);
+}
+
 
 
 void UMeshTangentsTool::OnShutdown(EToolShutdownType ShutdownType)
@@ -145,6 +226,16 @@ void UMeshTangentsTool::OnShutdown(EToolShutdownType ShutdownType)
 
 	Settings->SaveProperties(this);
 
+	if (GeometrySelectionViz)
+	{
+		GeometrySelectionViz->Disconnect();
+	}
+
+	if (GeometrySelectionVizProperties)
+	{
+		GeometrySelectionVizProperties->SaveProperties(this);
+	}
+
 	// Restore (unhide) the source meshes
 	UE::ToolTarget::ShowSourceObject(Target);
 
@@ -169,7 +260,7 @@ void UMeshTangentsTool::OnShutdown(EToolShutdownType ShutdownType)
 			}
 		}
 
-		Tangents->CopyToOverlays(*InputMesh);
+		CopyToOverlays(*Tangents, *InputMesh);
 
 		FConversionToMeshDescriptionOptions Options;
 		Options.bUpdatePositions = false;
@@ -190,6 +281,11 @@ void UMeshTangentsTool::OnTick(float DeltaTime)
 		UpdateVisualization(bThicknessDirty, bLengthDirty);
 		bThicknessDirty = bLengthDirty = bVisibilityChanged = false;
 	}
+
+	if (GeometrySelectionViz)
+	{
+		UpdateGeometrySelectionVisualization(GeometrySelectionViz, GeometrySelectionVizProperties);
+	}
 }
 
 
@@ -274,18 +370,24 @@ void UMeshTangentsTool::OnTangentsUpdated(const TUniquePtr<FMeshTangentsd>& NewR
 	PreviewGeometry->CreateOrUpdateLineSet(TEXT("Tangents"), InputMesh->MaxTriangleID(),
 		[&](int32 Index, TArray<FRenderableLine>& Lines) 
 	{
-
-		if (InputMesh->IsTriangle(Index) && DegenerateTris.Contains(Index) == false)
+		bool bValid = InputMesh->IsTriangle(Index) && !DegenerateTris.Contains(Index);
+		bool bIncludedTriangle = EditTriangles.IsEmpty() || EditTriangles.Contains(Index);
+		if (bValid && bIncludedTriangle)
 		{
-			FVector3d Verts[3];
-			InputMesh->GetTriVertices(Index, Verts[0], Verts[1], Verts[2]);
+			FIndex3i Vids = InputMesh->GetTriangle(Index);
 			for (int j = 0; j < 3; ++j)
 			{
-				FVector3d Tangent, Bitangent;
-				NewResult->GetPerTriangleTangent(Index, j, Tangent, Bitangent);
+				bool bIncludedVertex = EditVertices.IsEmpty() || EditVertices.Contains(Vids[j]);
+				if (bIncludedVertex)
+				{
+					FVector3d Origin = InputMesh->GetVertex(Vids[j]);
+
+					FVector3d Tangent, Bitangent;
+					NewResult->GetPerTriangleTangent(Index, j, Tangent, Bitangent);
 
-				Lines.Add(FRenderableLine((FVector)Verts[j], (FVector)Verts[j] + LineLength * (FVector)Tangent, FColor(240,15,15), Thickness));
-				Lines.Add(FRenderableLine((FVector)Verts[j], (FVector)Verts[j] + LineLength * (FVector)Bitangent, FColor(15,240,15), Thickness));
+					Lines.Add(FRenderableLine((FVector)Origin, (FVector)Origin + LineLength * (FVector)Tangent, FColor(240,15,15), Thickness));
+					Lines.Add(FRenderableLine((FVector)Origin, (FVector)Origin + LineLength * (FVector)Bitangent, FColor(15,240,15), Thickness));
+				}
 			}
 		}
 	}, 6);
@@ -298,18 +400,22 @@ void UMeshTangentsTool::OnTangentsUpdated(const TUniquePtr<FMeshTangentsd>& NewR
 	{
 		if (NormalOverlay->IsElement(Index))
 		{
-			FVector3f Normal = NormalOverlay->GetElement(Index);
 			int32 ParentVtx = NormalOverlay->GetParentVertex(Index);
-			FVector3f Position = (FVector3f)InputMesh->GetVertex(ParentVtx);
-			Lines.Add(FRenderableLine((FVector)Position, (FVector)Position + LineLength * (FVector)Normal, FColor(15,15,240), Thickness));
+			bool bIncluded = InputGeometrySelection.IsEmpty() || EditVertices.Contains(ParentVtx);
+			if (bIncluded)
+			{
+				FVector3f Normal = NormalOverlay->GetElement(Index);
+				FVector3f Origin = (FVector3f)InputMesh->GetVertex(ParentVtx);
+				Lines.Add(FRenderableLine((FVector)Origin, (FVector)Origin + LineLength * (FVector)Normal, FColor(15,15,240), Thickness));
+			}
 		}
 	}, 1);
 
-	ComputeMikkTDeviations(&DegenerateTris);
+	ComputeMikkTDeviations(DegenerateTris);
 
 	PreviewMesh->DeferredEditMesh([&](FDynamicMesh3& EditMesh)
 	{
-		NewResult.Get()->CopyToOverlays(EditMesh);
+		CopyToOverlays(*NewResult, EditMesh);
 	}, false);
 	PreviewMesh->NotifyDeferredEditCompleted(UPreviewMesh::ERenderUpdateMode::FastUpdate, EMeshRenderAttributeFlags::VertexNormals, false);
 
@@ -317,6 +423,48 @@ void UMeshTangentsTool::OnTangentsUpdated(const TUniquePtr<FMeshTangentsd>& NewR
 }
 
 
+void UMeshTangentsTool::CopyToOverlays(const FMeshTangentsd& Tangents, FDynamicMesh3& Mesh)
+{
+	if (ensure(Mesh.HasAttributes()) == false || ensure(Mesh.Attributes()->NumNormalLayers() == 3) == false)
+	{
+		return;
+	}
+	
+	if (InputGeometrySelection.IsEmpty())
+	{
+		Tangents.CopyToOverlays(Mesh);
+	}
+	else
+	{
+		FDynamicMesh3 MeshCopy;
+		{
+			bool bCopyNormals = false;
+			bool bCopyColors = false;
+			bool bCopyUVs = false;
+			bool bCopyAttributes = true;
+			MeshCopy.Copy(Mesh, bCopyNormals, bCopyColors, bCopyUVs, bCopyAttributes);
+		}
+
+		// Copy all tangents to the MeshCopy overlays
+		Tangents.CopyToOverlays(MeshCopy);
+
+		// Copy the subset of tangents corresponding to the geometry selection to the Mesh overlays
+		FDynamicMeshEditor Editor(&Mesh);
+		Editor.AppendElementSubset(
+			&MeshCopy,
+			EditTriangles,
+			EditVertices,
+			MeshCopy.Attributes()->PrimaryTangents(),
+			Mesh.Attributes()->PrimaryTangents());
+		Editor.AppendElementSubset(
+			&MeshCopy,
+			EditTriangles,
+			EditVertices,
+			MeshCopy.Attributes()->PrimaryBiTangents(),
+			Mesh.Attributes()->PrimaryBiTangents());
+	}
+}
+
 TSet<int32> UMeshTangentsTool::ComputeDegenerateTris() const
 {
 	TSet<int32> DegenerateTris;
@@ -330,19 +478,12 @@ TSet<int32> UMeshTangentsTool::ComputeDegenerateTris() const
 }
 
 
-void UMeshTangentsTool::ComputeMikkTDeviations(const TSet<int32>* DegenerateTris)
+void UMeshTangentsTool::ComputeMikkTDeviations(const TSet<int32>& DegenerateTris)
 {
 	// calculate deviation between what we have and MikkT, if necessary
 	Deviations.Reset();
 	if (Settings->bCompareWithMikkt && Settings->CalculationMethod == EMeshTangentsType::FastMikkTSpace)
 	{
-		TSet<int32> TempDegenerateTris;
-		if (DegenerateTris == nullptr)
-		{
-			TempDegenerateTris = ComputeDegenerateTris();
-			DegenerateTris = &TempDegenerateTris;
-		}
-
 		FProgressCancel TmpCancel;
 		FCalculateTangentsOp MikktOp;
 		MikktOp.SourceMesh = InputMesh;
@@ -358,27 +499,49 @@ void UMeshTangentsTool::ComputeMikkTDeviations(const TSet<int32>* DegenerateTris
 
 		for (int32 Index : InputMesh->TriangleIndicesItr())
 		{
-			if (DegenerateTris->Contains(Index) == false)
+			bool bValid = !DegenerateTris.Contains(Index);
+			bool bIncludedTriangle = EditTriangles.IsEmpty() || EditTriangles.Contains(Index);
+			if (bValid && bIncludedTriangle)
 			{
-				FVector3d Verts[3];
-				InputMesh->GetTriVertices(Index, Verts[0], Verts[1], Verts[2]);
+				FIndex3i Vids = InputMesh->GetTriangle(Index);
 				for (int j = 0; j < 3; ++j)
 				{
-					FVector3f TangentMikkt, BitangentMikkt;
-					MikktTangents->GetPerTriangleTangent<FVector3f, float>(Index, j, TangentMikkt, BitangentMikkt);
-					UE::Geometry::Normalize(TangentMikkt);
-					UE::Geometry::Normalize(BitangentMikkt);
-					FVector3f TangentNew, BitangentNew;
-					NewTangents->GetPerTriangleTangent<FVector3f, float>(Index, j, TangentNew, BitangentNew);
-					UE::Geometry::Normalize(TangentNew); 
-					UE::Geometry::Normalize(BitangentNew);
-					ensure(UE::Geometry::IsNormalized(TangentMikkt) && UE::Geometry::IsNormalized(BitangentMikkt));
-					ensure(UE::Geometry::IsNormalized(TangentNew) && UE::Geometry::IsNormalized(BitangentNew));
-					float MaxAngleDeg = FMathf::Max(UE::Geometry::AngleD(TangentMikkt, TangentNew), UE::Geometry::AngleD(BitangentMikkt, BitangentNew));
-					if (MaxAngleDeg > 0.5f)
+					bool bIncludedVertex = EditVertices.IsEmpty() || EditVertices.Contains(Vids[j]);
+					if (bIncludedVertex)
 					{
-						FMikktDeviation Deviation{ MaxAngleDeg, Index, j, (FVector3f)Verts[j], TangentMikkt, BitangentMikkt, TangentNew, BitangentNew };
-						Deviations.Add(Deviation);
+						FVector3f TangentMikkt, BitangentMikkt;
+						MikktTangents->GetPerTriangleTangent<FVector3f, float>(Index, j, TangentMikkt, BitangentMikkt);
+						UE::Geometry::Normalize(TangentMikkt);
+						UE::Geometry::Normalize(BitangentMikkt);
+						ensure(UE::Geometry::IsNormalized(TangentMikkt));
+						ensure(UE::Geometry::IsNormalized(BitangentMikkt));
+
+						FVector3f TangentNew, BitangentNew;
+						NewTangents->GetPerTriangleTangent<FVector3f, float>(Index, j, TangentNew, BitangentNew);
+						UE::Geometry::Normalize(TangentNew); 
+						UE::Geometry::Normalize(BitangentNew);
+						ensure(UE::Geometry::IsNormalized(TangentNew));
+						ensure(UE::Geometry::IsNormalized(BitangentNew));
+
+						float TangentAngleDeg = UE::Geometry::AngleD(TangentMikkt, TangentNew);
+						float BiTangentAngleDeg = UE::Geometry::AngleD(BitangentMikkt, BitangentNew);
+						float MaxAngleDeg = FMathf::Max(TangentAngleDeg, BiTangentAngleDeg);
+
+						if (MaxAngleDeg > 0.5f)
+						{
+							FMikktDeviation Deviation;
+
+							Deviation.MaxAngleDeg    = MaxAngleDeg;
+							Deviation.TriangleID     = Index;
+							Deviation.TriVertIndex   = j;
+							Deviation.VertexPos      = (FVector3f)InputMesh->GetVertex(Vids[j]);
+							Deviation.MikktTangent   = TangentMikkt;
+							Deviation.MikktBitangent = BitangentMikkt;
+							Deviation.OtherTangent   = TangentNew;
+							Deviation.OtherBitangent = BitangentNew;
+							
+							Deviations.Add(Deviation);
+						}
 					}
 				}
 			}
diff --git a/Engine/Plugins/Experimental/MeshModelingToolsetExp/Source/MeshModelingToolsEditorOnlyExp/Public/MeshTangentsTool.h b/Engine/Plugins/Experimental/MeshModelingToolsetExp/Source/MeshModelingToolsEditorOnlyExp/Public/MeshTangentsTool.h
--- a/Engine/Plugins/Experimental/MeshModelingToolsetExp/Source/MeshModelingToolsEditorOnlyExp/Public/MeshTangentsTool.h
+++ b/Engine/Plugins/Experimental/MeshModelingToolsetExp/Source/MeshModelingToolsEditorOnlyExp/Public/MeshTangentsTool.h
@@ -10,6 +10,7 @@
 #include "DynamicMesh/MeshTangents.h"
 #include "ParameterizationOps/CalculateTangentsOp.h"
 #include "BaseTools/SingleSelectionMeshEditingTool.h"
+#include "Selections/GeometrySelection.h"
 #include "MeshTangentsTool.generated.h"
 
 
@@ -18,6 +19,8 @@ struct FMeshDescription;
 class UDynamicMeshComponent;
 class UMaterialInterface;
 class UMaterialInstanceDynamic;
+class UPreviewGeometry;
+class UGeometrySelectionVisualizationProperties;
 
 
 /**
@@ -30,6 +33,7 @@ class MESHMODELINGTOOLSEDITORONLYEXP_API UMeshTangentsToolBuilder : public USing
 
 public:
 	virtual USingleSelectionMeshEditingTool* CreateNewTool(const FToolBuilderState& SceneState) const override;
+	virtual void InitializeNewTool(USingleSelectionMeshEditingTool* Tool, const FToolBuilderState& SceneState) const;
 	virtual bool CanBuildTool(const FToolBuilderState& SceneState) const override;
 
 protected:
@@ -106,6 +110,9 @@ public:
 	virtual bool HasAccept() const override { return true; }
 	virtual bool CanAccept() const override;
 
+	// input selection support
+	void SetGeometrySelection(UE::Geometry::FGeometrySelection&& SelectionIn);
+
 	// IGenericDataOperatorFactory API
 	virtual TUniquePtr<UE::Geometry::TGenericDataOperator<UE::Geometry::FMeshTangentsd>> MakeNewOperator() override;
 
@@ -133,9 +140,10 @@ protected:
 	bool bVisibilityChanged = false;
 
 	void OnTangentsUpdated(const TUniquePtr<UE::Geometry::FMeshTangentsd>& NewResult);
+	void CopyToOverlays(const UE::Geometry::FMeshTangentsd& Tangents, FDynamicMesh3& Mesh);
 	void UpdateVisualization(bool bThicknessChanged, bool bLengthChanged);
 	TSet<int32> ComputeDegenerateTris() const;
-	void ComputeMikkTDeviations(const TSet<int32>* DegenerateTris);
+	void ComputeMikkTDeviations(const TSet<int32>& DegenerateTris);
 
 	struct FMikktDeviation
 	{
@@ -148,6 +156,29 @@ protected:
 	};
 	TArray<FMikktDeviation> Deviations;
 
+	//
+	// Selection
+	//
+
+	UPROPERTY()
+	TObjectPtr<UGeometrySelectionVisualizationProperties> GeometrySelectionVizProperties = nullptr;
+
+	UPROPERTY()
+	TObjectPtr<UPreviewGeometry> GeometrySelectionViz = nullptr;
+
+	// The geometry selection that the user started the tool with.
+	// If the selection is empty we operate on the whole mesh, otherwise only edit the overlay elements implied by the selection.
+	UE::Geometry::FGeometrySelection InputGeometrySelection;
+
+	// If the user starts the tool with an edge selection we convert it to a vertex selection with triangle topology
+	// and store it here, we do this since we expect users to want vertex and edge selections to behave similarly.
+	UE::Geometry::FGeometrySelection TriangleVertexGeometrySelection;
+
+	// These are indices into InputMesh.
+	// If both are non-empty we edit the corresponding elements in the overlay, otherwise operate on the whole overlay
+	TSet<int> EditTriangles;
+	TSet<int> EditVertices;
+
 private:
 	bool bHasDisplayedNoAttributeError = false;
 };
diff --git a/Engine/Plugins/Experimental/MeshModelingToolsetExp/Source/MeshModelingToolsExp/Private/EditNormalsTool.cpp b/Engine/Plugins/Experimental/MeshModelingToolsetExp/Source/MeshModelingToolsExp/Private/EditNormalsTool.cpp
--- a/Engine/Plugins/Experimental/MeshModelingToolsetExp/Source/MeshModelingToolsExp/Private/EditNormalsTool.cpp
+++ b/Engine/Plugins/Experimental/MeshModelingToolsetExp/Source/MeshModelingToolsExp/Private/EditNormalsTool.cpp
@@ -142,31 +142,31 @@ void UEditNormalsTool::Setup()
 			FTransform ApplyTransform = UE::ToolTarget::GetLocalToWorldTransform(Targets[0]);
 
 			// Compute group topology if the selection has Polygroup topology, and do nothing otherwise
-			FGroupTopology GroupTopology(OriginalDynamicMeshes[0].Get(),
-				ActiveGroupSet.IsValid() ? ActiveGroupSet->PolygroupAttrib : nullptr,
-				InputGeometrySelection.TopologyType == EGeometryTopologyType::Polygroup);
-			
-			// Special case handling for edge and vertex element selections
-			bool bComputeTriangleVertexGeometrySelection = 
-				InputGeometrySelection.ElementType == EGeometryElementType::Vertex ||
-				InputGeometrySelection.ElementType == EGeometryElementType::Edge;
-			if (bComputeTriangleVertexGeometrySelection)
+			// Currently it is only possible to make a polygroup geometry selection using polygroup set stored directly in the mesh
+			FGroupTopology GroupTopology(OriginalDynamicMeshes[0].Get(), InputGeometrySelection.TopologyType == EGeometryTopologyType::Polygroup);
+
+			// Compute the overlay selection and a proxy triangle vertex selection used to make edge selections behave like
+			// vertex selections. TODO if we added Triangle/Vertex ROI visualization (or overlay element visualization) it
+			// would be clearer to users how those selection types affect different overlay elements and maybe we can remove
+			// this conversion, see also :EdgeSelectionsBehaveLikeVertexSelections
+			if (InputGeometrySelection.TopologyType == EGeometryTopologyType::Polygroup)
 			{
-				// Convert to a Triangle+Vertex selection and operate on that. If the input selection has Edge elements then
-				// we do this because users will expect this to behave similarly to a vertex selection containing the
-				// touched vertices. If the input selection has Vertex elements then we do this to simplify the
-				// implementation (the reason being that we'd want to use EnumeratePolygroupSelectionVertices but this
-				// visits all vertices in any group touching the corner, which isn't what we want for building the overlay
-				// selection from a Polygroup Topology)
-
-				TriangleVertexGeometrySelection.InitializeTypes(EGeometryElementType::Vertex, EGeometryTopologyType::Triangle);
-				bool bSuccess = ConvertSelection(
+				ConvertPolygroupSelectionToIncidentOverlaySelection(
 					*OriginalDynamicMeshes[0],
-					&GroupTopology, // This argument is ignored if InputGeometrySelection has Triangle topology
+					GroupTopology,
 					InputGeometrySelection,
-					TriangleVertexGeometrySelection);
-				ensure(bSuccess == true);
-				ensure(!TriangleVertexGeometrySelection.IsEmpty());
+					EditTriangles,
+					EditVertices,
+					&TriangleVertexGeometrySelection);
+			}
+			else
+			{
+				ConvertTriangleSelectionToOverlaySelection(
+					*OriginalDynamicMeshes[0],
+					InputGeometrySelection,
+					EditTriangles,
+					EditVertices,
+					&TriangleVertexGeometrySelection);
 			}
 
 			GeometrySelectionViz = NewObject<UPreviewGeometry>(this);
@@ -177,8 +177,7 @@ void UEditNormalsTool::Setup()
 				*OriginalDynamicMeshes[0],
 				InputGeometrySelection,
 				&GroupTopology,
-				bComputeTriangleVertexGeometrySelection ? &TriangleVertexGeometrySelection : nullptr);
-
+				!TriangleVertexGeometrySelection.IsEmpty() ? &TriangleVertexGeometrySelection : nullptr);
 		}
 	}
 
@@ -287,6 +286,8 @@ TUniquePtr<FDynamicMeshOperator> UEditNormalsOperatorFactory::MakeNewOperator()
 	NormalsOp->bAllowSharpVertices = Tool->BasicProperties->bAllowSharpVertices;
 	NormalsOp->NormalCalculationMethod = Tool->BasicProperties->NormalCalculationMethod;
 	NormalsOp->NormalSplitThreshold = Tool->BasicProperties->SharpEdgeAngleThreshold;
+	NormalsOp->EditTriangles = Tool->EditTriangles;
+	NormalsOp->EditVertices = Tool->EditVertices;
 
 	const FTransform LocalToWorld = (FTransform) UE::ToolTarget::GetLocalToWorldTransform(Tool->Targets[ComponentIndex]);
 	NormalsOp->OriginalMesh = Tool->OriginalDynamicMeshes[ComponentIndex];
@@ -297,47 +298,6 @@ TUniquePtr<FDynamicMeshOperator> UEditNormalsOperatorFactory::MakeNewOperator()
 		NormalsOp->MeshPolygroups = Tool->ActiveGroupSet;
 	}
 
-	if (!Tool->InputGeometrySelection.IsEmpty())
-	{
-		ensure(ComponentIndex == 0 && Tool->Targets.Num() == 1); // Sanity checks
-
-		if (Tool->InputGeometrySelection.ElementType == EGeometryElementType::Face)
-		{
-			if (Tool->InputGeometrySelection.TopologyType == EGeometryTopologyType::Polygroup)
-			{
-				const FPolygroupSet GroupSet = NormalsOp->MeshPolygroups.IsValid()
-						? *NormalsOp->MeshPolygroups
-						: FPolygroupSet(NormalsOp->OriginalMesh.Get());
-
-				EnumeratePolygroupSelectionTriangles(Tool->InputGeometrySelection, *NormalsOp->OriginalMesh, GroupSet,
-					[&NormalsOp](int32 ValidTid)
-				{
-					NormalsOp->EditTriangles.Add(ValidTid);
-					const FIndex3i Verts = NormalsOp->OriginalMesh->GetTriangle(ValidTid);
-					NormalsOp->EditVertices.Add(Verts.A);
-					NormalsOp->EditVertices.Add(Verts.B);
-					NormalsOp->EditVertices.Add(Verts.C);
-				});
-			}
-			else
-			{
-				UE::Geometry::ConvertTriangleSelectionToOverlaySelection(
-					*NormalsOp->OriginalMesh,
-					Tool->InputGeometrySelection,
-					NormalsOp->EditTriangles,
-					NormalsOp->EditVertices);
-			}
-		}
-		else
-		{
-			UE::Geometry::ConvertTriangleSelectionToOverlaySelection(
-				*NormalsOp->OriginalMesh,
-				Tool->TriangleVertexGeometrySelection,
-				NormalsOp->EditTriangles,
-				NormalsOp->EditVertices);
-		}
-	}
-
 	NormalsOp->SetTransform(LocalToWorld);
 
 	return NormalsOp;
diff --git a/Engine/Plugins/Experimental/MeshModelingToolsetExp/Source/MeshModelingToolsExp/Public/EditNormalsTool.h b/Engine/Plugins/Experimental/MeshModelingToolsetExp/Source/MeshModelingToolsExp/Public/EditNormalsTool.h
--- a/Engine/Plugins/Experimental/MeshModelingToolsetExp/Source/MeshModelingToolsExp/Public/EditNormalsTool.h
+++ b/Engine/Plugins/Experimental/MeshModelingToolsetExp/Source/MeshModelingToolsExp/Public/EditNormalsTool.h
@@ -10,7 +10,6 @@
 #include "DynamicMesh/DynamicMesh3.h"
 #include "PropertySets/PolygroupLayersProperties.h"
 #include "Polygroups/PolygroupSet.h"
-#include "Selection/ToolSelectionUtil.h"
 #include "Selections/GeometrySelection.h"
 #include "EditNormalsTool.generated.h"
 
@@ -178,7 +177,7 @@ protected:
 	void UpdateActiveGroupLayer();
 
 	//
-	// Selection
+	// Selection. Only used when the tool is run with one target
 	//
 
 	UPROPERTY()
@@ -195,6 +194,11 @@ protected:
 	// and store it here, we do this since we expect users to want vertex and edge selections to behave similarly.
 	UE::Geometry::FGeometrySelection TriangleVertexGeometrySelection;
 
+	// These are indices into the tool target mesh.
+	// If both are non-empty we edit the corresponding elements in the overlay, otherwise operate on the whole overlay
+	TSet<int> EditTriangles;
+	TSet<int> EditVertices;
+
 	// Cache the input polygroup set which was used to start the tool. We do this because users can change the
 	// polygroup referenced by the operator while using the tool.
 	TSharedPtr<UE::Geometry::FPolygroupSet, ESPMode::ThreadSafe> InputGeometrySelectionPolygroupSet;
diff --git a/Engine/Plugins/Runtime/GeometryProcessing/Source/DynamicMesh/Private/Selections/GeometrySelectionUtil.cpp b/Engine/Plugins/Runtime/GeometryProcessing/Source/DynamicMesh/Private/Selections/GeometrySelectionUtil.cpp
--- a/Engine/Plugins/Runtime/GeometryProcessing/Source/DynamicMesh/Private/Selections/GeometrySelectionUtil.cpp
+++ b/Engine/Plugins/Runtime/GeometryProcessing/Source/DynamicMesh/Private/Selections/GeometrySelectionUtil.cpp
@@ -1267,10 +1267,13 @@ bool UE::Geometry::ConvertSelection(
 		checkSlow(ToSelectionOut.TopologyType == EGeometryTopologyType::Triangle);
 		checkSlow(ToSelectionOut.ElementType == EGeometryElementType::Vertex);
 
-		for (FGroupTopology::FCorner Corner : GroupTopology->Corners)
-		{
-			ToSelectionOut.Selection.Add( FGeoSelectionID::MeshVertex((int32)Corner.VertexID).Encoded() );
-		}
+		// TODO Add a function which gets the Vids only, to remove the matrix-vector multiplication we just ignore
+		const FTransform Transform = FTransform::Identity;
+		EnumeratePolygroupSelectionVertices(FromSelectionIn, Mesh, GroupTopology, Transform,
+			[&ToSelectionOut](uint64 Vid, const FVector3d& Unused)
+			{
+				ToSelectionOut.Selection.Add( FGeoSelectionID::MeshVertex((int32)Vid).Encoded() );
+			});
 
 		return true;
 	};
@@ -1299,10 +1302,11 @@ bool UE::Geometry::ConvertSelection(
 }
 
 bool UE::Geometry::ConvertTriangleSelectionToOverlaySelection(
-	const UE::Geometry::FDynamicMesh3& Mesh,
+	const FDynamicMesh3& Mesh,
 	const FGeometrySelection& MeshSelection,
 	TSet<int>& TrianglesOut,
-	TSet<int>& VerticesOut)
+	TSet<int>& VerticesOut,
+	FGeometrySelection* IncidentSelection)
 {
 	if (!ensure(MeshSelection.TopologyType == EGeometryTopologyType::Triangle))
 	{
@@ -1330,6 +1334,25 @@ bool UE::Geometry::ConvertTriangleSelectionToOverlaySelection(
 			VerticesOut.Add(Verts.C);
 		});
 	}
+	else if (MeshSelection.ElementType == EGeometryElementType::Edge && IncidentSelection)
+	{
+		IncidentSelection->InitializeTypes(EGeometryElementType::Vertex, EGeometryTopologyType::Triangle);
+
+		EnumerateTriangleSelectionTriangles(MeshSelection, Mesh,
+			[&TrianglesOut](int32 ValidTid)
+		{
+			TrianglesOut.Add(ValidTid);
+		});
+
+		// TODO Add a function which gets the Vids only, to remove the matrix-vector multiplication we just ignore
+		const FTransform Transform = FTransform::Identity;
+		EnumerateTriangleSelectionVertices(MeshSelection, Mesh, Transform,
+			[&VerticesOut, IncidentSelection](uint64 Vid, const FVector3d& Unused)
+		{
+			IncidentSelection->Selection.Add(FGeoSelectionID::MeshVertex((int)Vid).Encoded() );
+			VerticesOut.Add((int)Vid);
+		});
+	}
 	else
 	{
 		EnumerateTriangleSelectionTriangles(MeshSelection, Mesh,
@@ -1350,6 +1373,79 @@ bool UE::Geometry::ConvertTriangleSelectionToOverlaySelection(
 	return true;
 }
 
+bool UE::Geometry::ConvertPolygroupSelectionToOverlaySelection(
+	const FDynamicMesh3& Mesh,
+	const FPolygroupSet& GroupSet,
+	const FGeometrySelection& MeshSelection,
+	TSet<int>& TrianglesOut,
+	TSet<int>& VerticesOut)
+{
+	return EnumeratePolygroupSelectionTriangles(MeshSelection, Mesh, GroupSet,
+		[&TrianglesOut, &VerticesOut, &Mesh](int32 ValidTid)
+	{
+		TrianglesOut.Add(ValidTid);
+		const FIndex3i Verts = Mesh.GetTriangle(ValidTid);
+		VerticesOut.Add(Verts.A);
+		VerticesOut.Add(Verts.B);
+		VerticesOut.Add(Verts.C);
+	});
+}
+
+bool UE::Geometry::ConvertPolygroupSelectionToIncidentOverlaySelection(
+	const FDynamicMesh3& Mesh,
+	const FGroupTopology& GroupTopology,
+	const FGeometrySelection& MeshSelection,
+	TSet<int>& TrianglesOut,
+	TSet<int>& VerticesOut,
+	FGeometrySelection* IncidentSelection)
+{
+	if (!ensure(MeshSelection.TopologyType == EGeometryTopologyType::Polygroup))
+	{
+		return false;
+	}
+
+	if (MeshSelection.ElementType == EGeometryElementType::Face)
+	{
+		if (MeshSelection.TopologyType == EGeometryTopologyType::Polygroup)
+		{
+			// TODO This uses the polygroup set stored directly in the mesh, this is a source of potential inconsistency
+			// with the given GroupTopology
+			const FPolygroupSet GroupSet = FPolygroupSet(&Mesh);
+
+			return EnumeratePolygroupSelectionTriangles(MeshSelection, Mesh, GroupSet,
+				[&TrianglesOut, &VerticesOut, &Mesh](int32 ValidTid)
+			{
+				TrianglesOut.Add(ValidTid);
+				const FIndex3i Verts = Mesh.GetTriangle(ValidTid);
+				VerticesOut.Add(Verts.A);
+				VerticesOut.Add(Verts.B);
+				VerticesOut.Add(Verts.C);
+			});
+		}
+		else
+		{
+			return ConvertTriangleSelectionToOverlaySelection(Mesh, MeshSelection, TrianglesOut, VerticesOut);
+		}
+	}
+	else
+	{
+		FGeometrySelection TempIncidentSelection;
+		if (IncidentSelection == nullptr)
+		{
+			IncidentSelection = &TempIncidentSelection;
+		}
+
+		IncidentSelection->InitializeTypes(EGeometryElementType::Vertex, EGeometryTopologyType::Triangle);
+
+		// GroupTopology argument is ignored if MeshSelection has Triangle topology
+		bool bSuccess = ConvertSelection(Mesh, &GroupTopology, MeshSelection, *IncidentSelection);
+		ensure(bSuccess == true);
+		ensure(!IncidentSelection->IsEmpty());
+
+		return ConvertTriangleSelectionToOverlaySelection(Mesh, *IncidentSelection, TrianglesOut, VerticesOut);
+	}
+}
+
 bool UE::Geometry::MakeSelectAllSelection(
 	const UE::Geometry::FDynamicMesh3& Mesh,
 	const FGroupTopology* GroupTopology,
diff --git a/Engine/Plugins/Runtime/GeometryProcessing/Source/DynamicMesh/Public/Selections/GeometrySelectionUtil.h b/Engine/Plugins/Runtime/GeometryProcessing/Source/DynamicMesh/Public/Selections/GeometrySelectionUtil.h
--- a/Engine/Plugins/Runtime/GeometryProcessing/Source/DynamicMesh/Public/Selections/GeometrySelectionUtil.h
+++ b/Engine/Plugins/Runtime/GeometryProcessing/Source/DynamicMesh/Public/Selections/GeometrySelectionUtil.h
@@ -244,6 +244,10 @@ DYNAMICMESH_API bool ConvertSelection(
  * since any FGeometrySelection can be represented as an overlay element selection because
  * any overlay element can be represented as a (Triangle,Vertex) pair.
  *
+ * If TriangleVertexSelectionIncidentToEdgeSelection is not null and MeshSelection is an Edge Selection it will
+ * be set to a Vertex selection with Triangle Topology corresponding to the vertices touched by the edge selection. This
+ * is useful when users expect an edge selection to behave similarly to the incident vertex selection.
+ *
  * @note it is not necessarily the case that all vertices of triangles in TrianglesOut will be in VerticesOut.
  * @return false if the MeshSelection topology type is not Triangle and true otherwise
  */
@@ -251,8 +255,49 @@ DYNAMICMESH_API bool ConvertTriangleSelectionToOverlaySelection(
 	const UE::Geometry::FDynamicMesh3& Mesh,
 	const FGeometrySelection& MeshSelection,
 	TSet<int>& TrianglesOut,
+	TSet<int>& VerticesOut,
+	FGeometrySelection* TriangleVertexSelectionIncidentToEdgeSelection = nullptr);
+
+/**
+ * Convert the given MeshSelection to a list of Triangles and Vertices into the Mesh,
+ * which can be used to represent a selection of overlay elements. This is always possible
+ * since any FGeometrySelection can be represented as an overlay element selection because
+ * any overlay element can be represented as a (Triangle,Vertex) pair.
+ * 
+ * For Polygroup Faces, all triangles in the face are included.
+ * For Polygroup Edges, all triangles in any group adjacent to the edge are included.
+ * For Polygroup Corners, all triangles in any group touching the corner are included.
+ * See ConvertPolygroupSelectionToIncidentOverlaySelection for a similar function which only includes triangles
+ * immediately incident to the polygroup element
+ * 
+ * @return false if the MeshSelection topology type is not Polygroup and true otherwise
+ */
+DYNAMICMESH_API bool ConvertPolygroupSelectionToOverlaySelection(
+	const UE::Geometry::FDynamicMesh3& Mesh,
+	const FPolygroupSet& GroupSet,
+	const FGeometrySelection& MeshSelection,
+	TSet<int>& TrianglesOut,
 	TSet<int>& VerticesOut);
 
+/**
+ * Like ConvertPolygroupSelectionToOverlaySelection but only includes overlay elements that are immediately incident to
+ * Polygroup Vertices/Edges.
+ *
+ * If TriangleVertexSelectionIncidentToEdgeOrVertexSelection is not null and MeshSelection is an Edge or Vertex
+ * selection it will be set to a Vertex selection with Triangle Topology corresponding to the vertices touched by the
+ * edge selection. This is useful when users expect an edge selection to behave similarly to the incident vertex
+ * selection.
+ *
+ * @return false if the MeshSelection topology type is not Polygroup and true otherwise
+ */
+DYNAMICMESH_API bool ConvertPolygroupSelectionToIncidentOverlaySelection(
+	const UE::Geometry::FDynamicMesh3& Mesh,
+	const FGroupTopology& GroupTopology,
+	const FGeometrySelection& MeshSelection,
+	TSet<int>& TrianglesOut,
+	TSet<int>& VerticesOut,
+	FGeometrySelection* TriangleVertexSelectionIncidentToEdgeOrVertexSelection = nullptr);
+
 
 /**
  * Select all elements of the provided Mesh and GroupTopology that pass the provided SelectionIDPredicate, 
diff --git a/Engine/Plugins/Runtime/MeshModelingToolset/Source/ModelingComponents/Private/PropertySets/GeometrySelectionVisualizationProperties.cpp b/Engine/Plugins/Runtime/MeshModelingToolset/Source/ModelingComponents/Private/PropertySets/GeometrySelectionVisualizationProperties.cpp
--- a/Engine/Plugins/Runtime/MeshModelingToolset/Source/ModelingComponents/Private/PropertySets/GeometrySelectionVisualizationProperties.cpp
+++ b/Engine/Plugins/Runtime/MeshModelingToolset/Source/ModelingComponents/Private/PropertySets/GeometrySelectionVisualizationProperties.cpp
@@ -117,7 +117,7 @@ bool UGeometrySelectionVisualizationProperties::ShowFaceSelectionTriangleSet() c
 
 bool UGeometrySelectionVisualizationProperties::ShowEdgeSelectionVerticesPointSet() const
 {
-	if (ShowEdgeSelectionLineSet())
+	if (bShowSelection == false)
 	{
 		return false;
 	}
diff --git a/Engine/Plugins/Runtime/MeshModelingToolset/Source/ModelingOperators/Private/CleaningOps/EditNormalsOp.cpp b/Engine/Plugins/Runtime/MeshModelingToolset/Source/ModelingOperators/Private/CleaningOps/EditNormalsOp.cpp
--- a/Engine/Plugins/Runtime/MeshModelingToolset/Source/ModelingOperators/Private/CleaningOps/EditNormalsOp.cpp
+++ b/Engine/Plugins/Runtime/MeshModelingToolset/Source/ModelingOperators/Private/CleaningOps/EditNormalsOp.cpp
@@ -2,6 +2,7 @@
 
 #include "CleaningOps/EditNormalsOp.h"
 
+#include "DynamicMeshEditor.h"
 #include "MeshDescriptionToDynamicMesh.h"
 #include "ToolContextInterfaces.h"
 
@@ -335,33 +336,12 @@ void FEditNormalsOp::CalculateResultSelection(FProgressCancel* Progress)
 		);
 	}
 
-	// Copy the changed elements from Workspace to ResultNormals
-	TMap<int32,int32> WorkspaceElementIDToNewResultElementID;
-	for (int32 Tid : EditTriangles) {
-
-		const FIndex3i WorkspaceElementIDs = Workspace.GetTriangle(Tid);
-		const FIndex3i TriVids = ResultMesh->GetTriangle(Tid);
-		FIndex3i TriElementsToSet = ResultNormals.GetTriangle(Tid);
-
-		for (int SubIdx = 0; SubIdx < 3; ++SubIdx){
-
-			const int WorkspaceElementID = WorkspaceElementIDs[SubIdx];
-
-			if (EditVertices.Contains(TriVids[SubIdx])){
-
-				int* FoundElementID = WorkspaceElementIDToNewResultElementID.Find(WorkspaceElementID);
-				if (!FoundElementID)
-				{
-					const FVector3f WorkspaceNormal = Workspace.GetElement(WorkspaceElementID);
-					const int NewResultElementID = ResultNormals.AppendElement(WorkspaceNormal);
-
-					FoundElementID = &WorkspaceElementIDToNewResultElementID.Add(WorkspaceElementID, NewResultElementID);
-				}
-
-				TriElementsToSet[SubIdx] = *FoundElementID;
-			}
-		}
-  
-		ResultNormals.SetTriangle(Tid, TriElementsToSet);
-	}
+	// Copy the subset of normals corresponding to the geometry selection to the Mesh overlays
+	FDynamicMeshEditor Editor(ResultMesh.Get());
+	Editor.AppendElementSubset(
+		ResultMesh.Get(),
+		EditTriangles,
+		EditVertices,
+		&Workspace,
+		&ResultNormals);
 }
\ No newline at end of file
diff --git a/Engine/Source/Runtime/GeometryCore/Private/DynamicMeshEditor.cpp b/Engine/Source/Runtime/GeometryCore/Private/DynamicMeshEditor.cpp
--- a/Engine/Source/Runtime/GeometryCore/Private/DynamicMeshEditor.cpp
+++ b/Engine/Source/Runtime/GeometryCore/Private/DynamicMeshEditor.cpp
@@ -2324,3 +2324,56 @@ bool FDynamicMeshEditor::SplitMesh(const FDynamicMesh3* SourceMesh, TArray<FDyna
 	
 	return true;
 }
+
+
+template <typename RealType, int ElementSize>
+void FDynamicMeshEditor::AppendElementSubset(
+	const FDynamicMesh3* FromMesh,
+	const TSet<int>& TriangleROI,
+	const TSet<int>& VertexROI,
+	const TDynamicMeshOverlay<RealType, ElementSize>* FromOverlay,
+	TDynamicMeshOverlay<RealType, ElementSize>* ToOverlay)
+{
+	TMap<int32, int32> FromElementIDToElementID;
+	for (int32 Tid : TriangleROI) {
+
+		const FIndex3i TriVids = Mesh->GetTriangle(Tid);
+		ensureMsgf(TriVids == FromMesh->GetTriangle(Tid),
+			TEXT("Expected FromOverlay and ToOverlay to have a matching parent meshes"));
+
+		const FIndex3i FromElementIDs = FromOverlay->GetTriangle(Tid);
+		FIndex3i ToElementIDs = ToOverlay->GetTriangle(Tid);
+
+		for (int SubIdx = 0; SubIdx < 3; ++SubIdx){
+
+			const int FromElementID = FromElementIDs[SubIdx];
+
+			if (VertexROI.Contains(TriVids[SubIdx])){
+
+				int* ToElementID = FromElementIDToElementID.Find(FromElementID);
+				if (!ToElementID)
+				{
+					for (int Index = 0; Index < ElementSize; Index++)
+					{
+						RealType FromData[ElementSize];
+						FromOverlay->GetElement(FromElementID, FromData);
+						const int ElementID = ToOverlay->AppendElement(FromData);
+
+						ToElementID = &FromElementIDToElementID.Add(FromElementID, ElementID);
+					}
+				}
+
+				ToElementIDs[SubIdx] = *ToElementID;
+			}
+		}
+
+		ToOverlay->SetTriangle(Tid, ToElementIDs);
+	}
+}
+
+template GEOMETRYCORE_API void FDynamicMeshEditor::AppendElementSubset(
+	const FDynamicMesh3*,
+	const TSet<int>&,
+	const TSet<int>&,
+	const TDynamicMeshOverlay<float, 3>*,
+	TDynamicMeshOverlay<float, 3>*);
diff --git a/Engine/Source/Runtime/GeometryCore/Public/DynamicMeshEditor.h b/Engine/Source/Runtime/GeometryCore/Public/DynamicMeshEditor.h
--- a/Engine/Source/Runtime/GeometryCore/Public/DynamicMeshEditor.h
+++ b/Engine/Source/Runtime/GeometryCore/Public/DynamicMeshEditor.h
@@ -565,6 +565,28 @@ public:
 		FIndexMapi& ColorMapOut);
 
 
+
+	/**
+	 * Append overlay elements in both ROIs from one overlay to another
+	 *
+	 * @param FromMesh     Mesh that owns FromOverlay attribute overlay
+	 * @param TriangleROI  Triangles containing the copied overlay elements
+	 * @param VertexROI    Parent vertices of the copied overlay elements
+	 * @param FromOverlay  Overlay we want to append from (owned by FromMesh)
+	 * @param ToOverlay    Overlay we want to append into (owned by .Mesh)
+	 *
+	 * @note Required overloads should be declared in the .cpp file
+	 */
+	template <typename RealType, int ElementSize>
+	GEOMETRYCORE_API void AppendElementSubset(
+		const FDynamicMesh3* FromMesh,
+		const TSet<int>& TriangleROI,
+		const TSet<int>& VertexROI,
+		const TDynamicMeshOverlay<RealType, ElementSize>* FromOverlay,
+		TDynamicMeshOverlay<RealType, ElementSize>* ToOverlay);
+
+
+
 	/**
 	 * Append triangles of an existing mesh. This duplicates the current groups and also any attributes existing on the triangles.
 	 * @param SourceMesh the mesh to copy from
