Approximate Actors - Allow approximate actors to use an array of components as input, rather than only actors

diff --git a/Engine/Plugins/Editor/WorldPartitionHLODUtilities/Source/Private/WorldPartition/HLOD/Builders/HLODBuilderMeshApproximate.cpp b/Engine/Plugins/Editor/WorldPartitionHLODUtilities/Source/Private/WorldPartition/HLOD/Builders/HLODBuilderMeshApproximate.cpp
--- a/Engine/Plugins/Editor/WorldPartitionHLODUtilities/Source/Private/WorldPartition/HLOD/Builders/HLODBuilderMeshApproximate.cpp
+++ b/Engine/Plugins/Editor/WorldPartitionHLODUtilities/Source/Private/WorldPartition/HLOD/Builders/HLODBuilderMeshApproximate.cpp
@@ -73,11 +73,9 @@ TArray<UActorComponent*> UHLODBuilderMeshApproximate::Build(const FHLODBuildCont
 {
 	TRACE_CPUPROFILER_EVENT_SCOPE(UHLODBuilderMeshApproximate::Build);
 
-	TArray<UPrimitiveComponent*> PrimitiveComponents = FilterComponents<UPrimitiveComponent>(InSourceComponents);
+	IGeometryProcessing_ApproximateActors::FInput Input;
+	Input.Components = InSourceComponents;
 
-	TSet<AActor*> Actors;
-	Algo::Transform(PrimitiveComponents, Actors, [](UPrimitiveComponent* PrimitiveComponent) { return PrimitiveComponent->GetOwner(); });
-	
 	IGeometryProcessingInterfacesModule& GeomProcInterfaces = FModuleManager::Get().LoadModuleChecked<IGeometryProcessingInterfacesModule>("GeometryProcessingInterfaces");
 	IGeometryProcessing_ApproximateActors* ApproxActorsAPI = GeomProcInterfaces.GetApproximateActorsImplementation();
 
@@ -112,6 +110,7 @@ TArray<UActorComponent*> UHLODBuilderMeshApproximate::Build(const FHLODBuildCont
 	Options.PackedMRSTexParamName = FName("PackedTexture");
 
 	// Compute texel density if needed, depending on the TextureSizingType setting
+	TArray<UPrimitiveComponent*> PrimitiveComponents = FilterComponents<UPrimitiveComponent>(InSourceComponents);
 	if (UseSettings.MaterialSettings.ResolveTexelDensity(PrimitiveComponents, Options.MeshTexelDensity))
 	{ 
 		Options.TextureSizePolicy = IGeometryProcessing_ApproximateActors::ETextureSizePolicy::TexelDensity;
@@ -119,7 +118,7 @@ TArray<UActorComponent*> UHLODBuilderMeshApproximate::Build(const FHLODBuildCont
 
 	// run actor approximation computation
 	IGeometryProcessing_ApproximateActors::FResults Results;
-	ApproxActorsAPI->ApproximateActors(Actors.Array(), Options, Results);
+	ApproxActorsAPI->ApproximateActors(Input, Options, Results);
 
 	TArray<UActorComponent*> Components;
 	if (Results.ResultCode == IGeometryProcessing_ApproximateActors::EResultCode::Success)
diff --git a/Engine/Plugins/Experimental/MeshModelingToolsetExp/Source/GeometryProcessingAdapters/Private/ApproximateActorsImpl.cpp b/Engine/Plugins/Experimental/MeshModelingToolsetExp/Source/GeometryProcessingAdapters/Private/ApproximateActorsImpl.cpp
--- a/Engine/Plugins/Experimental/MeshModelingToolsetExp/Source/GeometryProcessingAdapters/Private/ApproximateActorsImpl.cpp
+++ b/Engine/Plugins/Experimental/MeshModelingToolsetExp/Source/GeometryProcessingAdapters/Private/ApproximateActorsImpl.cpp
@@ -75,7 +75,7 @@ struct FGeneratedResultTextures
 
 
 static TUniquePtr<FSceneCapturePhotoSet> CapturePhotoSet(
-	const TArray<AActor*>& Actors,
+	const IGeometryProcessing_ApproximateActors::FInput& Input,
 	const IGeometryProcessing_ApproximateActors::FOptions& Options
 )
 {
@@ -111,13 +111,14 @@ static TUniquePtr<FSceneCapturePhotoSet> CapturePhotoSet(
 		SceneCapture->SetCaptureTypeEnabled(ERenderCaptureType::Specular, bSpecular);
 	}
 
-	UWorld* World = Actors.IsEmpty() ? nullptr : Actors[0]->GetWorld();
+	UWorld* World = Input.Actors.IsEmpty() ? (Input.Components.IsEmpty() ? nullptr : Input.Components[0]->GetWorld()) : Input.Actors[0]->GetWorld();
 
-	SceneCapture->SetCaptureSceneActors(World, Actors);
+	SceneCapture->SetCaptureSceneActorsAndComponents(World, Input.Actors, Input.Components);
 
 	const TArray<FSpatialPhotoParams> SpatialParams = ComputeStandardExteriorSpatialPhotoParameters(
 		World,
-		Actors,
+		Input.Actors,
+		Input.Components,
 		CaptureDimensions,
 		FieldOfView,
 		NearPlaneDist,
@@ -992,19 +993,19 @@ IGeometryProcessing_ApproximateActors::FOptions FApproximateActorsImpl::Construc
 }
 
 
-void FApproximateActorsImpl::ApproximateActors(const TArray<AActor*>& Actors, const FOptions& Options, FResults& ResultsOut)
+void FApproximateActorsImpl::ApproximateActors(const FInput& Input, const FOptions& Options, FResults& ResultsOut)
 {
 	int32 ActorClusters = 1;
 	FScopedSlowTask Progress(1.f, LOCTEXT("ApproximatingActors", "Generating Actor Approximation..."));
 	Progress.MakeDialog(true);
 	Progress.EnterProgressFrame(1.f);
-	GenerateApproximationForActorSet(Actors, Options, ResultsOut);
+	GenerateApproximationForActorSet(Input, Options, ResultsOut);
 }
 
 
 
 
-void FApproximateActorsImpl::GenerateApproximationForActorSet(const TArray<AActor*>& Actors, const FOptions& Options, FResults& ResultsOut)
+void FApproximateActorsImpl::GenerateApproximationForActorSet(const FInput& Input, const FOptions& Options, FResults& ResultsOut)
 {
 	TRACE_CPUPROFILER_EVENT_SCOPE(ApproximateActorsImpl_Generate);
 
@@ -1047,7 +1048,9 @@ void FApproximateActorsImpl::GenerateApproximationForActorSet(const TArray<AActo
 	{
 		TRACE_CPUPROFILER_EVENT_SCOPE(ApproximateActorsImpl_Generate_BuildScene);
 		TRACE_BOOKMARK(TEXT("ApproximateActors-Adding Actors"));
-		Scene->AddActors(Actors);
+		Scene->AddActors(Input.Actors);
+		TRACE_BOOKMARK(TEXT("ApproximateActors-Adding Components"));
+		Scene->AddComponents(Input.Components);
 		TRACE_BOOKMARK(TEXT("ApproximateActors-Building Scene"));
 		Scene->Build(SceneBuildOptions);
 	}
@@ -1121,7 +1124,7 @@ void FApproximateActorsImpl::GenerateApproximationForActorSet(const TArray<AActo
 		WaitForMeshAvailable();
 		if (ResultsOut.ResultCode == EResultCode::Success)
 		{
-			EmitGeneratedMeshAsset(Actors, Options, ResultsOut, &FinalMesh, nullptr, WriteDebugMesh);
+			EmitGeneratedMeshAsset(Options, ResultsOut, &FinalMesh, nullptr, WriteDebugMesh);
 		}
 		return;
 	}
@@ -1142,7 +1145,7 @@ void FApproximateActorsImpl::GenerateApproximationForActorSet(const TArray<AActo
 	Progress.EnterProgressFrame(1.f, LOCTEXT("CapturingScene", "Capturing Scene..."));
 	TRACE_BOOKMARK(TEXT("ApproximateActors-Capture Photos"));
 
-	TUniquePtr<FSceneCapturePhotoSet> SceneCapture = CapturePhotoSet(Actors, Options);
+	TUniquePtr<FSceneCapturePhotoSet> SceneCapture = CapturePhotoSet(Input, Options);
 
 	// if parallel capture was allowed, need to force the mesh compute to finish now to be able to proceed
 	if (Options.bMaximizeBakeParallelism == true)
@@ -1274,13 +1277,12 @@ void FApproximateActorsImpl::GenerateApproximationForActorSet(const TArray<AActo
 	// (does this do that? Let calling code do it?)
 	NewMaterial->PostEditChange();
 
-	EmitGeneratedMeshAsset(Actors, Options, ResultsOut, &FinalMesh, NewMaterial, WriteDebugMesh);
+	EmitGeneratedMeshAsset(Options, ResultsOut, &FinalMesh, NewMaterial, WriteDebugMesh);
 	ResultsOut.ResultCode = EResultCode::Success;
 }
 
 
 UStaticMesh* FApproximateActorsImpl::EmitGeneratedMeshAsset(
-	const TArray<AActor*>& Actors, 
 	const FOptions& Options, 
 	FResults& ResultsOut,
 	FDynamicMesh3* FinalMesh,
diff --git a/Engine/Plugins/Experimental/MeshModelingToolsetExp/Source/GeometryProcessingAdapters/Public/GeometryProcessing/ApproximateActorsImpl.h b/Engine/Plugins/Experimental/MeshModelingToolsetExp/Source/GeometryProcessingAdapters/Public/GeometryProcessing/ApproximateActorsImpl.h
--- a/Engine/Plugins/Experimental/MeshModelingToolsetExp/Source/GeometryProcessingAdapters/Public/GeometryProcessing/ApproximateActorsImpl.h
+++ b/Engine/Plugins/Experimental/MeshModelingToolsetExp/Source/GeometryProcessingAdapters/Public/GeometryProcessing/ApproximateActorsImpl.h
@@ -22,16 +22,15 @@ class GEOMETRYPROCESSINGADAPTERS_API FApproximateActorsImpl : public IGeometryPr
 public:
 	virtual FOptions ConstructOptions(const FMeshApproximationSettings& MeshApproximationSettings) override;
 
-	virtual void ApproximateActors(const TArray<AActor*>& Actors, const FOptions& Options, FResults& ResultsOut) override;
+	virtual void ApproximateActors(const FInput& Input, const FOptions& Options, FResults& ResultsOut) override;
 	
 
 
 protected:
 
-	virtual void GenerateApproximationForActorSet(const TArray<AActor*>& Actors, const FOptions& Options, FResults& ResultsOut);
+	virtual void GenerateApproximationForActorSet(const FInput& Input, const FOptions& Options, FResults& ResultsOut);
 
 	virtual UStaticMesh* EmitGeneratedMeshAsset(
-		const TArray<AActor*>& Actors,
 		const FOptions& Options,
 		FResults& ResultsOut,
 		FDynamicMesh3* FinalMesh,
diff --git a/Engine/Plugins/Runtime/MeshModelingToolset/Source/ModelingComponents/Private/Baking/RenderCaptureFunctions.cpp b/Engine/Plugins/Runtime/MeshModelingToolset/Source/ModelingComponents/Private/Baking/RenderCaptureFunctions.cpp
--- a/Engine/Plugins/Runtime/MeshModelingToolset/Source/ModelingComponents/Private/Baking/RenderCaptureFunctions.cpp
+++ b/Engine/Plugins/Runtime/MeshModelingToolset/Source/ModelingComponents/Private/Baking/RenderCaptureFunctions.cpp
@@ -253,6 +253,7 @@ void UE::Geometry::ConfigureSceneCapture(
 	const TArray<FSpatialPhotoParams> SpatialParams = ComputeStandardExteriorSpatialPhotoParameters(
 		World,
 		Actors,
+		TArray<UActorComponent*>(),
 		FImageDimensions(Options.RenderCaptureImageSize, Options.RenderCaptureImageSize),
 		Options.FieldOfViewDegrees,
 		Options.NearPlaneDist,
diff --git a/Engine/Plugins/Runtime/MeshModelingToolset/Source/ModelingComponents/Private/Scene/SceneCapturePhotoSet.cpp b/Engine/Plugins/Runtime/MeshModelingToolset/Source/ModelingComponents/Private/Scene/SceneCapturePhotoSet.cpp
--- a/Engine/Plugins/Runtime/MeshModelingToolset/Source/ModelingComponents/Private/Scene/SceneCapturePhotoSet.cpp
+++ b/Engine/Plugins/Runtime/MeshModelingToolset/Source/ModelingComponents/Private/Scene/SceneCapturePhotoSet.cpp
@@ -26,7 +26,17 @@ FSceneCapturePhotoSet::FSceneCapturePhotoSet()
 
 void FSceneCapturePhotoSet::SetCaptureSceneActors(UWorld* World, const TArray<AActor*>& Actors)
 {
-	if (this->TargetWorld != World || this->VisibleActors != Actors)
+	SetCaptureSceneActorsAndComponents(World, Actors, {});
+}
+
+void FSceneCapturePhotoSet::SetCaptureSceneComponents(UWorld* World, const TArray<UActorComponent*>& Components)
+{
+	SetCaptureSceneActorsAndComponents(World, {}, Components);
+}
+
+void FSceneCapturePhotoSet::SetCaptureSceneActorsAndComponents(UWorld* World, const TArray<AActor*>& Actors, const TArray<UActorComponent*>& Components)
+{
+	if (this->TargetWorld != World || this->VisibleActors != Actors || this->VisibleComponents != Components)
 	{
 		ForEachCaptureType([this](ERenderCaptureType CaptureType)
 		{
@@ -44,11 +54,7 @@ void FSceneCapturePhotoSet::SetCaptureSceneActors(UWorld* World, const TArray<AA
 	}
 	this->TargetWorld = World;
 	this->VisibleActors = Actors;
-}
-
-TArray<AActor*> FSceneCapturePhotoSet::GetCaptureSceneActors()
-{
-	return VisibleActors;
+	this->VisibleComponents = Components;
 }
 
 UWorld* FSceneCapturePhotoSet::GetCaptureTargetWorld()
@@ -182,7 +188,7 @@ void FSceneCapturePhotoSet::Compute()
 
 	FWorldRenderCapture RenderCapture;
 	RenderCapture.SetWorld(TargetWorld);
-	RenderCapture.SetVisibleActors(VisibleActors);
+	RenderCapture.SetVisibleActorsAndComponents(VisibleActors, VisibleComponents);
 	if (bWriteDebugImages)
 	{
 		RenderCapture.SetEnableWriteDebugImage(true, 0, DebugImagesFolderName);
@@ -363,7 +369,7 @@ void FSceneCapturePhotoSet::AddExteriorCaptures(
 
 	FWorldRenderCapture RenderCapture;
 	RenderCapture.SetWorld(TargetWorld);
-	RenderCapture.SetVisibleActors(VisibleActors);
+	RenderCapture.SetVisibleActorsAndComponents(VisibleActors, VisibleComponents);
 	RenderCapture.SetDimensions(PhotoDimensions);
 	if (bWriteDebugImages)
 	{
@@ -869,6 +875,7 @@ void FSceneCapturePhotoSet::EmptyPhotoSet(ERenderCaptureType CaptureType)
 TArray<FSpatialPhotoParams> UE::Geometry::ComputeStandardExteriorSpatialPhotoParameters(
 	UWorld* World,
 	const TArray<AActor*>& Actors,
+	const TArray<UActorComponent*>& Components,
 	FImageDimensions PhotoDimensions,
 	double HorizontalFOVDegrees,
 	double NearPlaneDist,
@@ -878,7 +885,7 @@ TArray<FSpatialPhotoParams> UE::Geometry::ComputeStandardExteriorSpatialPhotoPar
 	bool bUpperEdges,
 	bool bSideEdges)
 {
-	if (!World || Actors.IsEmpty())
+	if (!World || (Actors.IsEmpty() && Components.IsEmpty()))
 	{
 		return {};
 	}
@@ -928,7 +935,7 @@ TArray<FSpatialPhotoParams> UE::Geometry::ComputeStandardExteriorSpatialPhotoPar
 	FSphere RenderSphere;
 	{
 		FWorldRenderCapture RenderCapture;
-		RenderCapture.SetVisibleActors(Actors);
+		RenderCapture.SetVisibleActorsAndComponents(Actors, Components);
 		RenderSphere = RenderCapture.ComputeContainingRenderSphere(HorizontalFOVDegrees);
 	}
 
diff --git a/Engine/Plugins/Runtime/MeshModelingToolset/Source/ModelingComponents/Private/Scene/WorldRenderCapture.cpp b/Engine/Plugins/Runtime/MeshModelingToolset/Source/ModelingComponents/Private/Scene/WorldRenderCapture.cpp
--- a/Engine/Plugins/Runtime/MeshModelingToolset/Source/ModelingComponents/Private/Scene/WorldRenderCapture.cpp
+++ b/Engine/Plugins/Runtime/MeshModelingToolset/Source/ModelingComponents/Private/Scene/WorldRenderCapture.cpp
@@ -204,51 +204,66 @@ void FWorldRenderCapture::SetWorld(UWorld* WorldIn)
 	this->World = WorldIn;
 }
 
-
 void FWorldRenderCapture::SetVisibleActors(const TArray<AActor*>& Actors)
 {
-	CaptureActors = Actors;
+	SetVisibleActorsAndComponents(Actors, {});
+}
+
+void FWorldRenderCapture::SetVisibleComponents(const TArray<UActorComponent*>& Components)
+{
+	SetVisibleActorsAndComponents({}, Components);
+}
 
+void FWorldRenderCapture::SetVisibleActorsAndComponents(const TArray<AActor*>& Actors, const TArray<UActorComponent*>& Components)
+{
 	VisiblePrimitives.Reset();
 
 	FBoxSphereBounds::Builder BoundsBuilder;
 
 	// Find all components that need to be included in rendering.
 	// This also descends into any ChildActorComponents
-	// TODO Consider using Actor->GetActorBounds() or ForEachComponent here
-	TArray<UActorComponent*> ComponentQueue;
-	for (AActor* Actor : Actors)
+		
+	auto AddComponent = [&](UActorComponent* Component)
 	{
-		ComponentQueue.Reset();
-		for (UActorComponent* Component : Actor->GetComponents())
-		{
-			ComponentQueue.Add(Component);
-		}
-		while (ComponentQueue.Num() > 0)
+		auto AddComponentImpl = [&](UActorComponent* Component, auto& AddComponentRef) -> void
 		{
-			UActorComponent* Component = ComponentQueue.Pop(false);
-			if (Cast<UPrimitiveComponent>(Component) != nullptr)
-			{
-				UPrimitiveComponent* PrimitiveComponent = Cast<UPrimitiveComponent>(Component);
+			if (UPrimitiveComponent* PrimitiveComponent = Cast<UPrimitiveComponent>(Component))
+			{				
 				VisiblePrimitives.Add(PrimitiveComponent->GetPrimitiveSceneId());
 
 				// Append bounds of visible components only
 				BoundsBuilder += PrimitiveComponent->Bounds;
 			}
-			else if (Cast<UChildActorComponent>(Component) != nullptr)
+			else if (UChildActorComponent* ChildActorComponent = Cast<UChildActorComponent>(Component))
 			{
-				AActor* ChildActor = Cast<UChildActorComponent>(Component)->GetChildActor();
-				if (ChildActor != nullptr)
+				if (AActor* ChildActor = ChildActorComponent->GetChildActor())
 				{
 					for (UActorComponent* SubComponent : ChildActor->GetComponents())
 					{
-						ComponentQueue.Add(SubComponent);
+						AddComponentRef(SubComponent, AddComponentRef);
 					}
 				}
-			}
+			}			
+		};
+
+		AddComponentImpl(Component, AddComponentImpl);
+	};
+
+	// Add actors
+	for (AActor* Actor : Actors)
+	{
+		for (UActorComponent* Component : Actor->GetComponents())
+		{
+			AddComponent(Component);
 		}
 	}
 
+	// Add components
+	for (UActorComponent* Component : Components)
+	{
+		AddComponent(Component);
+	}
+
 	VisibleBounds = BoundsBuilder;
 }
 
diff --git a/Engine/Plugins/Runtime/MeshModelingToolset/Source/ModelingComponents/Public/Scene/SceneCapturePhotoSet.h b/Engine/Plugins/Runtime/MeshModelingToolset/Source/ModelingComponents/Public/Scene/SceneCapturePhotoSet.h
--- a/Engine/Plugins/Runtime/MeshModelingToolset/Source/ModelingComponents/Public/Scene/SceneCapturePhotoSet.h
+++ b/Engine/Plugins/Runtime/MeshModelingToolset/Source/ModelingComponents/Public/Scene/SceneCapturePhotoSet.h
@@ -8,6 +8,7 @@
 
 class UWorld;
 class AActor;
+class UActorComponent;
 struct FScopedSlowTask;
 
 namespace UE
@@ -48,11 +49,13 @@ public:
 	FSceneCapturePhotoSet();
 
 	/**
-	 * Set the target World and set of Actors
+	 * Set the target World and set of Actors/Components
 	 * If World != this->World or Actors != this->Actors all existing photo sets are cleared
 	 */
 	void SetCaptureSceneActors(UWorld* World, const TArray<AActor*>& Actors);
-	TArray<AActor*> GetCaptureSceneActors();
+	void SetCaptureSceneComponents(UWorld* World, const TArray<UActorComponent*>& Components);
+	void SetCaptureSceneActorsAndComponents(UWorld* World, const TArray<AActor*>& Actors, const TArray<UActorComponent*>& Components);
+
 	UWorld* GetCaptureTargetWorld();
 
 	/**
@@ -248,6 +251,7 @@ public:
 protected:
 	UWorld* TargetWorld = nullptr;
 	TArray<AActor*> VisibleActors;
+	TArray<UActorComponent*> VisibleComponents;
 
 	bool bEnforceVisibilityViaUnregister = false;
 
@@ -364,6 +368,7 @@ MODELINGCOMPONENTS_API
 TArray<FSpatialPhotoParams> ComputeStandardExteriorSpatialPhotoParameters(
 	UWorld* World,
 	const TArray<AActor*>& Actors,
+	const TArray<UActorComponent*>& Components,
 	FImageDimensions PhotoDimensions,
 	double HorizontalFOVDegrees,
 	double NearPlaneDist,
diff --git a/Engine/Plugins/Runtime/MeshModelingToolset/Source/ModelingComponents/Public/Scene/WorldRenderCapture.h b/Engine/Plugins/Runtime/MeshModelingToolset/Source/ModelingComponents/Public/Scene/WorldRenderCapture.h
--- a/Engine/Plugins/Runtime/MeshModelingToolset/Source/ModelingComponents/Public/Scene/WorldRenderCapture.h
+++ b/Engine/Plugins/Runtime/MeshModelingToolset/Source/ModelingComponents/Public/Scene/WorldRenderCapture.h
@@ -199,6 +199,7 @@ struct FRenderCaptureCoordinateConverter2D
 /**
  * FWorldRenderCapture captures a rendering of a set of Actors in a World from a
  * specific viewpoint. Various types of rendering are supported, as defined by ERenderCaptureType.
+ * Currently rendering an entire World, ie without an explicit list of Actors or Components, is not supported.
  */
 class MODELINGCOMPONENTS_API FWorldRenderCapture
 {
@@ -214,10 +215,19 @@ public:
 
 	/** 
 	 *  Set the set of Actors in the target World that should be included in the Rendering.
-	 *  Currently rendering an entire World, ie without an explicit list of Actors, is not supported.
 	 */
 	void SetVisibleActors(const TArray<AActor*>& Actors);
 
+	/**
+	 *  Set the set of Components in the target World that should be included in the Rendering.
+	 */
+	void SetVisibleComponents(const TArray<UActorComponent*>& Components);
+
+	/**
+	 *  Set the set of Actors and Components in the target World that should be included in the Rendering.
+	 */
+	void SetVisibleActorsAndComponents(const TArray<AActor*>& VisibleActors, const TArray<UActorComponent*>& VisibleComponents);
+
 	/** Get bounding-box of the Visible actors */
 	FBoxSphereBounds GetVisibleActorBounds() const { return VisibleBounds; }
 
@@ -264,7 +274,6 @@ public:
 protected:
 	UWorld* World = nullptr;
 
-	TArray<AActor*> CaptureActors;
 	TSet<FPrimitiveComponentId> VisiblePrimitives;
 	FBoxSphereBounds VisibleBounds;
 
diff --git a/Engine/Source/Developer/GeometryProcessingInterfaces/Public/GeometryProcessingInterfaces/ApproximateActors.h b/Engine/Source/Developer/GeometryProcessingInterfaces/Public/GeometryProcessingInterfaces/ApproximateActors.h
--- a/Engine/Source/Developer/GeometryProcessingInterfaces/Public/GeometryProcessingInterfaces/ApproximateActors.h
+++ b/Engine/Source/Developer/GeometryProcessingInterfaces/Public/GeometryProcessingInterfaces/ApproximateActors.h
@@ -99,6 +99,15 @@ public:
 		RelativeError
 	};
 
+	/**
+	 * Input to approximate, as actors and/or components
+	 */
+	struct FInput
+	{
+		TArray<AActor*>				Actors;
+		TArray<UActorComponent*>	Components;
+	};
+
 
 	/**
 	 * Input options to Actor Approximation process
@@ -289,9 +298,9 @@ public:
 	}
 
 	/**
-	 * Top-level driver function that clients call to generate the approximation for a set of input Actors.
+	 * Top-level driver function that clients call to generate the approximation for a set of input Actors or Components.
 	 */
-	virtual void ApproximateActors(const TArray<AActor*>& Actors, const FOptions& Options, FResults& ResultsOut) 
+	virtual void ApproximateActors(const FInput& Input, const FOptions& Options, FResults& ResultsOut) 
 	{
 		check(false);		// not implemented in base class
 	}
diff --git a/Engine/Source/Developer/HierarchicalLODUtilities/Private/HierarchicalLODUtilities.cpp b/Engine/Source/Developer/HierarchicalLODUtilities/Private/HierarchicalLODUtilities.cpp
--- a/Engine/Source/Developer/HierarchicalLODUtilities/Private/HierarchicalLODUtilities.cpp
+++ b/Engine/Source/Developer/HierarchicalLODUtilities/Private/HierarchicalLODUtilities.cpp
@@ -582,11 +582,11 @@ FHLODBuildResults GenerateHLODMesh_Approximate(const FHLODBuildParams& InBuildPa
 	Options.BasePackagePath = PackagePath / NewAssetNamePrefix + AssetName;
 
 	// run actor approximation computation
-	TSet<AActor*> Actors;
-	Algo::Transform(InBuildParams.Components, Actors, [](UPrimitiveComponent* PrimitiveComponent) { return PrimitiveComponent->GetOwner(); });
+	IGeometryProcessing_ApproximateActors::FInput Input;
+	Algo::Transform(InBuildParams.Components, Input.Components, [](UPrimitiveComponent* PrimitiveComponent) { return PrimitiveComponent; });
 
 	IGeometryProcessing_ApproximateActors::FResults Results;
-	ApproxActorsAPI->ApproximateActors(Actors.Array(), Options, Results);
+	ApproxActorsAPI->ApproximateActors(Input, Options, Results);
 
 	auto RenameNewAsset = [&PackagePath, &NewAssetNamePrefix, &InBuildParams](UObject* NewAsset)
 	{
diff --git a/Engine/Source/Editor/MergeActors/Private/MeshApproximationTool/MeshApproximationTool.cpp b/Engine/Source/Editor/MergeActors/Private/MeshApproximationTool/MeshApproximationTool.cpp
--- a/Engine/Source/Editor/MergeActors/Private/MeshApproximationTool/MeshApproximationTool.cpp
+++ b/Engine/Source/Editor/MergeActors/Private/MeshApproximationTool/MeshApproximationTool.cpp
@@ -197,10 +197,28 @@ bool FMeshApproximationTool::RunMerge(const FString& PackageName, const TArray<T
 	FString PackagePath = FPackageName::GetLongPackagePath(PackageName);
 	FString AssetName = FPackageName::GetLongPackageAssetName(PackageName);
 	Options.BasePackagePath = PackagePath / NewAssetNamePrefix + AssetName;
+
+	// Extracting static mesh components from the selected mesh components in the dialog
+	IGeometryProcessing_ApproximateActors::FInput Input;
+
+	for (const TSharedPtr<FMergeComponentData>& SelectedComponent : SelectedComponents)
+	{
+		// Determine whether or not this component should be incorporated according the user settings
+		if (SelectedComponent->bShouldIncorporate && SelectedComponent->PrimComponent.IsValid())
+		{
+			if (UStaticMeshComponent* StaticMeshComponent = Cast<UStaticMeshComponent>(SelectedComponent->PrimComponent.Get()))
+			{
+				if (StaticMeshComponent->GetStaticMesh() != nullptr)
+				{
+					Input.Components.Add(StaticMeshComponent);
+				}
+			}
+		}
+	}
 	
 	// Run actor approximation computation
 	IGeometryProcessing_ApproximateActors::FResults Results;
-	ApproxActorsAPI->ApproximateActors(Actors, Options, Results);
+	ApproxActorsAPI->ApproximateActors(Input, Options, Results);
 
 	auto ProcessNewAsset = [&PackagePath, &NewAssetNamePrefix](UObject* NewAsset)
 	{
