Fix the 'mesh to collision' tool not properly updating dynamic mesh component collision settings

diff --git a/Engine/Plugins/Experimental/MeshModelingToolsetExp/Source/MeshModelingToolsExp/Private/Physics/SetCollisionGeometryTool.cpp b/Engine/Plugins/Experimental/MeshModelingToolsetExp/Source/MeshModelingToolsExp/Private/Physics/SetCollisionGeometryTool.cpp
--- a/Engine/Plugins/Experimental/MeshModelingToolsetExp/Source/MeshModelingToolsExp/Private/Physics/SetCollisionGeometryTool.cpp
+++ b/Engine/Plugins/Experimental/MeshModelingToolsetExp/Source/MeshModelingToolsExp/Private/Physics/SetCollisionGeometryTool.cpp
@@ -485,10 +485,9 @@ void USetCollisionGeometryTool::OnShutdown(EToolShutdownType ShutdownType)
 		}
 		else if (UDynamicMeshComponent* DynamicMeshComponent = Cast<UDynamicMeshComponent>(Component))
 		{
-			if (UBodySetup* BodySetup = DynamicMeshComponent->GetBodySetup())
-			{
-				UpdateBodySetup(BodySetup);
-			}
+			DynamicMeshComponent->CollisionType = (ECollisionTraceFlag)(int32)Settings->SetCollisionType;
+			DynamicMeshComponent->SetSimpleCollisionShapes(GeneratedCollision->AggGeom, true);		
+			DynamicMeshComponent->MarkRenderStateDirty();
 		}
 
 		// post the undo transaction
