add offset parameter to spline tool, to support drawing splines that are not directly on the clicked surfaces

diff --git a/Engine/Plugins/Experimental/MeshModelingToolsetExp/Source/MeshModelingToolsEditorOnlyExp/Private/DrawSplineTool.cpp b/Engine/Plugins/Experimental/MeshModelingToolsetExp/Source/MeshModelingToolsEditorOnlyExp/Private/DrawSplineTool.cpp
--- a/Engine/Plugins/Experimental/MeshModelingToolsetExp/Source/MeshModelingToolsEditorOnlyExp/Private/DrawSplineTool.cpp
+++ b/Engine/Plugins/Experimental/MeshModelingToolsetExp/Source/MeshModelingToolsEditorOnlyExp/Private/DrawSplineTool.cpp
@@ -739,6 +739,17 @@ bool UDrawSplineTool::Raycast(const FRay& WorldRay, FVector3d& HitLocationOut, F
 			BestHitT = HitTOut;
 		}
 	}
+
+	if (Settings->ClickOffset != 0.0)
+	{
+		FVector3d OffsetDirection = HitNormalOut;
+		if (Settings->OffsetMethod == ESplineOffsetMethod::Custom)
+		{
+			OffsetDirection = Settings->OffsetDirection.GetSafeNormal(UE_SMALL_NUMBER, FVector3d::UnitZ());
+		}
+
+		HitLocationOut += OffsetDirection * Settings->ClickOffset;
+	}
 	
 	return BestHitT < TNumericLimits<double>::Max();
 }
diff --git a/Engine/Plugins/Experimental/MeshModelingToolsetExp/Source/MeshModelingToolsEditorOnlyExp/Public/DrawSplineTool.h b/Engine/Plugins/Experimental/MeshModelingToolsetExp/Source/MeshModelingToolsEditorOnlyExp/Public/DrawSplineTool.h
--- a/Engine/Plugins/Experimental/MeshModelingToolsetExp/Source/MeshModelingToolsEditorOnlyExp/Public/DrawSplineTool.h
+++ b/Engine/Plugins/Experimental/MeshModelingToolsetExp/Source/MeshModelingToolsEditorOnlyExp/Public/DrawSplineTool.h
@@ -31,6 +31,13 @@ enum class EDrawSplineDrawMode : uint8
 	FreeDraw,
 };
 
+UENUM()
+enum class ESplineOffsetMethod : uint8
+{
+	HitNormal,
+	Custom
+};
+
 UENUM()
 enum class EDrawSplineOutputMode : uint8
 {
@@ -114,6 +121,16 @@ public:
 		EditCondition = "DrawMode == EDrawSplineDrawMode::FreeDraw", EditConditionHides))
 	double MinPointSpacing = 200;
 
+	/** How far to offset spline points from the clicked surface, along the surface normal */
+	UPROPERTY(EditAnywhere, Category = Spline, meta = (UIMin = 0, UIMax = 100))
+	double ClickOffset = 0;
+
+	UPROPERTY(EditAnywhere, Category = Spline, meta = (EditCondition = "ClickOffset > 0", EditConditionHides))
+	ESplineOffsetMethod OffsetMethod = ESplineOffsetMethod::HitNormal;
+
+	UPROPERTY(EditAnywhere, Category = Spline, meta = (EditCondition = "ClickOffset > 0 && OffsetMethod == ESplineOffsetMethod::Custom", EditConditionHides))
+	FVector OffsetDirection = FVector::UpVector;
+
 	/**
 	 * When nonzero, allows a visualization of the rotation of the spline. Can be controlled
 	 * in the detail panel after creation via "Scale Visualization Width" property.
