Use deprecated typedef to fix FSelectedOjectsChangeList typo in ITF's ToolContextInterfaces

diff --git a/Engine/Plugins/Experimental/MeshModelingToolsetExp/Source/MeshModelingToolsEditorOnlyExp/Private/BspConversionTool.cpp b/Engine/Plugins/Experimental/MeshModelingToolsetExp/Source/MeshModelingToolsEditorOnlyExp/Private/BspConversionTool.cpp
--- a/Engine/Plugins/Experimental/MeshModelingToolsetExp/Source/MeshModelingToolsEditorOnlyExp/Private/BspConversionTool.cpp
+++ b/Engine/Plugins/Experimental/MeshModelingToolsetExp/Source/MeshModelingToolsEditorOnlyExp/Private/BspConversionTool.cpp
@@ -672,7 +672,7 @@ void UBspConversionTool::ApplyAction(EBspConversionToolAction ActionType)
 	{
 		FScopedTransaction Transaction(LOCTEXT("SelectAllValidBrushes", "Select All Valid Brushes"));
 
-		FSelectedOjectsChangeList NewSelection;
+		FSelectedObjectsChangeList NewSelection;
 		NewSelection.ModificationType = ESelectedObjectsModificationType::Replace;
 
 		// Accumulate all actors in the current level that are valid brushes
@@ -692,7 +692,7 @@ void UBspConversionTool::ApplyAction(EBspConversionToolAction ActionType)
 	{
 		FScopedTransaction Transaction(LOCTEXT("DeselectVolumes", "Deselect Volumes"));
 
-		FSelectedOjectsChangeList NewSelection;
+		FSelectedObjectsChangeList NewSelection;
 		NewSelection.ModificationType = ESelectedObjectsModificationType::Replace;
 
 		// Out of the current selection, keep anything except volume brushes
@@ -715,7 +715,7 @@ void UBspConversionTool::ApplyAction(EBspConversionToolAction ActionType)
 		// Normally, "GEditor->SelectNone(true, false, false)" would deselect all brushes, but
 		// it does not deselect volumes, which we want to do depending on the settings.
 
-		FSelectedOjectsChangeList NewSelection;
+		FSelectedObjectsChangeList NewSelection;
 		NewSelection.ModificationType = ESelectedObjectsModificationType::Replace;
 
 		// From current selection, select everything that is valid given current settings
diff --git a/Engine/Plugins/Experimental/MeshModelingToolsetExp/Source/MeshModelingToolsExp/Private/BakeTransformTool.cpp b/Engine/Plugins/Experimental/MeshModelingToolsetExp/Source/MeshModelingToolsExp/Private/BakeTransformTool.cpp
--- a/Engine/Plugins/Experimental/MeshModelingToolsetExp/Source/MeshModelingToolsExp/Private/BakeTransformTool.cpp
+++ b/Engine/Plugins/Experimental/MeshModelingToolsetExp/Source/MeshModelingToolsExp/Private/BakeTransformTool.cpp
@@ -344,7 +344,7 @@ void UBakeTransformTool::UpdateAssets()
 	if (BasicProperties->bRecenterPivot)
 	{
 		// hack to ensure user sees the updated pivot immediately: request re-select of the original selection
-		FSelectedOjectsChangeList NewSelection;
+		FSelectedObjectsChangeList NewSelection;
 		NewSelection.ModificationType = ESelectedObjectsModificationType::Replace;
 		for (int OrigMeshIdx = 0; OrigMeshIdx < Targets.Num(); OrigMeshIdx++)
 		{
diff --git a/Engine/Plugins/Experimental/MeshModelingToolsetExp/Source/MeshModelingToolsExp/Private/EditPivotTool.cpp b/Engine/Plugins/Experimental/MeshModelingToolsetExp/Source/MeshModelingToolsExp/Private/EditPivotTool.cpp
--- a/Engine/Plugins/Experimental/MeshModelingToolsetExp/Source/MeshModelingToolsExp/Private/EditPivotTool.cpp
+++ b/Engine/Plugins/Experimental/MeshModelingToolsetExp/Source/MeshModelingToolsExp/Private/EditPivotTool.cpp
@@ -638,7 +638,7 @@ void UEditPivotTool::UpdateAssets(const FFrame3d& NewPivotWorldFrame)
 	}
 
 	// hack to ensure user sees the updated pivot immediately: request re-select of the original selection
-	FSelectedOjectsChangeList NewSelection;
+	FSelectedObjectsChangeList NewSelection;
 	NewSelection.ModificationType = ESelectedObjectsModificationType::Replace;
 	for (int OrigMeshIdx = 0; OrigMeshIdx < Targets.Num(); OrigMeshIdx++)
 	{
diff --git a/Engine/Plugins/Experimental/MeshModelingToolsetExp/Source/MeshModelingToolsExp/Private/MirrorTool.cpp b/Engine/Plugins/Experimental/MeshModelingToolsetExp/Source/MeshModelingToolsExp/Private/MirrorTool.cpp
--- a/Engine/Plugins/Experimental/MeshModelingToolsetExp/Source/MeshModelingToolsExp/Private/MirrorTool.cpp
+++ b/Engine/Plugins/Experimental/MeshModelingToolsetExp/Source/MeshModelingToolsExp/Private/MirrorTool.cpp
@@ -452,7 +452,7 @@ void UMirrorTool::GenerateAsset(const TArray<FDynamicMeshOpResult>& Results)
 	}
 
 	// Properly deal with each result, setting up the selection at the same time.
-	FSelectedOjectsChangeList NewSelection;
+	FSelectedObjectsChangeList NewSelection;
 	NewSelection.ModificationType = ESelectedObjectsModificationType::Replace;
 	for (int OrigMeshIdx = 0; OrigMeshIdx < NumSourceMeshes; OrigMeshIdx++)
 	{
diff --git a/Engine/Plugins/Experimental/MeshModelingToolsetExp/Source/MeshModelingToolsExp/Private/PlaneCutTool.cpp b/Engine/Plugins/Experimental/MeshModelingToolsetExp/Source/MeshModelingToolsExp/Private/PlaneCutTool.cpp
--- a/Engine/Plugins/Experimental/MeshModelingToolsetExp/Source/MeshModelingToolsExp/Private/PlaneCutTool.cpp
+++ b/Engine/Plugins/Experimental/MeshModelingToolsetExp/Source/MeshModelingToolsExp/Private/PlaneCutTool.cpp
@@ -408,7 +408,7 @@ void UPlaneCutTool::GenerateAsset(const TArray<FDynamicMeshOpResult>& Results)
 	TArray<TArray<FDynamicMesh3>> AllSplitMeshes; AllSplitMeshes.SetNum(NumSourceMeshes);
 
 	// build a selection change starting w/ the original selection (used if objects are added below)
-	FSelectedOjectsChangeList NewSelection;
+	FSelectedObjectsChangeList NewSelection;
 	NewSelection.ModificationType = ESelectedObjectsModificationType::Replace;
 	for (int OrigMeshIdx = 0; OrigMeshIdx < NumSourceMeshes; OrigMeshIdx++)
 	{
diff --git a/Engine/Plugins/Experimental/MeshModelingToolsetExp/Source/MeshModelingToolsExp/Private/RemoveOccludedTrianglesTool.cpp b/Engine/Plugins/Experimental/MeshModelingToolsetExp/Source/MeshModelingToolsExp/Private/RemoveOccludedTrianglesTool.cpp
--- a/Engine/Plugins/Experimental/MeshModelingToolsetExp/Source/MeshModelingToolsExp/Private/RemoveOccludedTrianglesTool.cpp
+++ b/Engine/Plugins/Experimental/MeshModelingToolsetExp/Source/MeshModelingToolsExp/Private/RemoveOccludedTrianglesTool.cpp
@@ -487,7 +487,7 @@ void URemoveOccludedTrianglesTool::GenerateAsset(const TArray<FDynamicMeshOpResu
 		}
 	}
 
-	FSelectedOjectsChangeList NewSelection;
+	FSelectedObjectsChangeList NewSelection;
 	for (int32 PreviewIdx = 0; PreviewIdx < Previews.Num(); PreviewIdx++)
 	{
 		check(Results[PreviewIdx].Mesh.Get() != nullptr);
diff --git a/Engine/Plugins/Experimental/XRCreativeFramework/Source/XRCreative/Private/ITF/XRCreativeITFComponent.cpp b/Engine/Plugins/Experimental/XRCreativeFramework/Source/XRCreative/Private/ITF/XRCreativeITFComponent.cpp
--- a/Engine/Plugins/Experimental/XRCreativeFramework/Source/XRCreative/Private/ITF/XRCreativeITFComponent.cpp
+++ b/Engine/Plugins/Experimental/XRCreativeFramework/Source/XRCreative/Private/ITF/XRCreativeITFComponent.cpp
@@ -284,7 +284,7 @@ public:
 		}
 	}
 
-	virtual bool RequestSelectionChange(const FSelectedOjectsChangeList& SelectionChange) override
+	virtual bool RequestSelectionChange(const FSelectedObjectsChangeList& SelectionChange) override
 	{
 		if (UTypedElementSelectionSet* SelectionSet = WeakSelectionSet.Get())
 		{
diff --git a/Engine/Plugins/MeshPainting/Source/MeshPaintingToolset/Private/MeshPaintInteractions.cpp b/Engine/Plugins/MeshPainting/Source/MeshPaintingToolset/Private/MeshPaintInteractions.cpp
--- a/Engine/Plugins/MeshPainting/Source/MeshPaintingToolset/Private/MeshPaintInteractions.cpp
+++ b/Engine/Plugins/MeshPainting/Source/MeshPaintingToolset/Private/MeshPaintInteractions.cpp
@@ -65,7 +65,7 @@ void UMeshPaintSelectionMechanic::OnClicked(const FInputDeviceRay& ClickPos)
 			GetParentTool()->GetToolManager()->BeginUndoTransaction(LOCTEXT("MeshSelection", "Select Mesh"));
 
 
-			FSelectedOjectsChangeList NewSelection;
+			FSelectedObjectsChangeList NewSelection;
 			// TODO add CTRL handling
 			const bool bShouldAddToSelection = bAddToSelectionSet && Interface->AllowsMultiselect();
 			NewSelection.ModificationType = bShouldAddToSelection ? ESelectedObjectsModificationType::Add : ESelectedObjectsModificationType::Replace;
diff --git a/Engine/Plugins/Runtime/MeshModelingToolset/Source/ModelingComponents/Private/Selection/ToolSelectionUtil.cpp b/Engine/Plugins/Runtime/MeshModelingToolset/Source/ModelingComponents/Private/Selection/ToolSelectionUtil.cpp
--- a/Engine/Plugins/Runtime/MeshModelingToolset/Source/ModelingComponents/Private/Selection/ToolSelectionUtil.cpp
+++ b/Engine/Plugins/Runtime/MeshModelingToolset/Source/ModelingComponents/Private/Selection/ToolSelectionUtil.cpp
@@ -24,7 +24,7 @@ using namespace UE::Geometry;
 
 void ToolSelectionUtil::SetNewActorSelection(UInteractiveToolManager* ToolManager, AActor* Actor)
 {
-	FSelectedOjectsChangeList NewSelection;
+	FSelectedObjectsChangeList NewSelection;
 	NewSelection.ModificationType = ESelectedObjectsModificationType::Replace;
 	NewSelection.Actors.Add(Actor);
 	ToolManager->RequestSelectionChange(NewSelection);
@@ -33,7 +33,7 @@ void ToolSelectionUtil::SetNewActorSelection(UInteractiveToolManager* ToolManage
 
 void ToolSelectionUtil::SetNewActorSelection(UInteractiveToolManager* ToolManager, const TArray<AActor*>& Actors)
 {
-	FSelectedOjectsChangeList NewSelection;
+	FSelectedObjectsChangeList NewSelection;
 	NewSelection.ModificationType = ESelectedObjectsModificationType::Replace;
 	for (AActor* Actor : Actors)
 	{
diff --git a/Engine/Source/Editor/Experimental/EditorInteractiveToolsFramework/Private/EdModeInteractiveToolsContext.cpp b/Engine/Source/Editor/Experimental/EditorInteractiveToolsFramework/Private/EdModeInteractiveToolsContext.cpp
--- a/Engine/Source/Editor/Experimental/EditorInteractiveToolsFramework/Private/EdModeInteractiveToolsContext.cpp
+++ b/Engine/Source/Editor/Experimental/EditorInteractiveToolsFramework/Private/EdModeInteractiveToolsContext.cpp
@@ -259,7 +259,7 @@ public:
 	}
 
 
-	virtual bool RequestSelectionChange(const FSelectedOjectsChangeList& SelectionChange) override
+	virtual bool RequestSelectionChange(const FSelectedObjectsChangeList& SelectionChange) override
 	{
 		checkf(SelectionChange.Components.Num() == 0, TEXT("FEdModeToolsContextTransactionImpl::RequestSelectionChange - Component selection not supported yet"));
 
diff --git a/Engine/Source/Runtime/InteractiveToolsFramework/Private/InteractiveToolManager.cpp b/Engine/Source/Runtime/InteractiveToolsFramework/Private/InteractiveToolManager.cpp
--- a/Engine/Source/Runtime/InteractiveToolsFramework/Private/InteractiveToolManager.cpp
+++ b/Engine/Source/Runtime/InteractiveToolsFramework/Private/InteractiveToolManager.cpp
@@ -449,7 +449,7 @@ void UInteractiveToolManager::EmitObjectChange(UObject* TargetObject, TUniquePtr
 	}
 }
 
-bool UInteractiveToolManager::RequestSelectionChange(const FSelectedOjectsChangeList& SelectionChange)
+bool UInteractiveToolManager::RequestSelectionChange(const FSelectedObjectsChangeList& SelectionChange)
 {
 	return TransactionsAPI->RequestSelectionChange(SelectionChange);
 }
diff --git a/Engine/Source/Runtime/InteractiveToolsFramework/Public/InteractiveToolManager.h b/Engine/Source/Runtime/InteractiveToolsFramework/Public/InteractiveToolManager.h
--- a/Engine/Source/Runtime/InteractiveToolsFramework/Public/InteractiveToolManager.h
+++ b/Engine/Source/Runtime/InteractiveToolsFramework/Public/InteractiveToolManager.h
@@ -223,7 +223,7 @@ public:
 	/**
 	 * Forward an FChange object to the Context
 	 */
-	INTERACTIVETOOLSFRAMEWORK_API virtual bool RequestSelectionChange(const FSelectedOjectsChangeList& SelectionChange);
+	INTERACTIVETOOLSFRAMEWORK_API virtual bool RequestSelectionChange(const FSelectedObjectsChangeList& SelectionChange);
 
 	//
 	// State control  (@todo: have the Context call these? not safe for anyone to call)
diff --git a/Engine/Source/Runtime/InteractiveToolsFramework/Public/ToolContextInterfaces.h b/Engine/Source/Runtime/InteractiveToolsFramework/Public/ToolContextInterfaces.h
--- a/Engine/Source/Runtime/InteractiveToolsFramework/Public/ToolContextInterfaces.h
+++ b/Engine/Source/Runtime/InteractiveToolsFramework/Public/ToolContextInterfaces.h
@@ -264,7 +264,7 @@ enum class ESelectedObjectsModificationType
 
 
 /** Represents a change to a set of selected Actors and Components */
-struct FSelectedOjectsChangeList
+struct FSelectedObjectsChangeList
 {
 	/** How should this list be interpreted in the context of a larger selection set */
 	ESelectedObjectsModificationType ModificationType;
@@ -274,6 +274,9 @@ struct FSelectedOjectsChangeList
 	TArray<UActorComponent*> Components;
 };
 
+UE_DEPRECATED(5.4, "Use FSelectedObjectsChangeList instead")
+typedef FSelectedObjectsChangeList FSelectedOjectsChangeList;
+
 
 /**
  * Users of the Tools Framework need to implement IToolsContextTransactionsAPI so that
@@ -329,7 +332,7 @@ public:
 	 * @param SelectionChange desired modification to current selection
 	 * @return true if the selection change could be applied
 	 */
-	virtual bool RequestSelectionChange(const FSelectedOjectsChangeList& SelectionChange) = 0;
+	virtual bool RequestSelectionChange(const FSelectedObjectsChangeList& SelectionChange) = 0;
 
 };
 
