Approximate Actors - Fix ensure() caused by empty scene representation
* Return an error if no valid input was provided

diff --git a/Engine/Plugins/Experimental/MeshModelingToolsetExp/Source/GeometryProcessingAdapters/Private/ApproximateActorsImpl.cpp b/Engine/Plugins/Experimental/MeshModelingToolsetExp/Source/GeometryProcessingAdapters/Private/ApproximateActorsImpl.cpp
--- a/Engine/Plugins/Experimental/MeshModelingToolsetExp/Source/GeometryProcessingAdapters/Private/ApproximateActorsImpl.cpp
+++ b/Engine/Plugins/Experimental/MeshModelingToolsetExp/Source/GeometryProcessingAdapters/Private/ApproximateActorsImpl.cpp
@@ -1049,12 +1049,23 @@ void FApproximateActorsImpl::GenerateApproximationForActorSet(const FInput& Inpu
 	SceneBuildOptions.bOnlySurfaceMaterials = true;				// don't include decal geometry in 3D mesh scene (will be included in renderings)
 	SceneBuildOptions.bEnableUVQueries = SceneBuildOptions.bEnableNormalsQueries = false;		// not required in this context, will reduce memory usage
 	SceneBuildOptions.bPrintDebugMessages = Options.bVerbose;
+
 	{
 		TRACE_CPUPROFILER_EVENT_SCOPE(ApproximateActorsImpl_Generate_BuildScene);
 		TRACE_BOOKMARK(TEXT("ApproximateActors-Adding Actors"));
 		Scene->AddActors(Input.Actors);
 		TRACE_BOOKMARK(TEXT("ApproximateActors-Adding Components"));
 		Scene->AddComponents(Input.Components);
+	}
+
+	if (!Scene->IsValid())
+	{
+		UE_LOG(LogApproximateActors, Error, TEXT("No valid input actors/components - unable to generate mesh"));
+		ResultsOut.ResultCode = EResultCode::MeshGenerationFailed;
+		return;
+	}
+
+	{
 		TRACE_BOOKMARK(TEXT("ApproximateActors-Building Scene"));
 		Scene->Build(SceneBuildOptions);
 	}
diff --git a/Engine/Plugins/Runtime/MeshModelingToolset/Source/ModelingComponents/Private/Scene/MeshSceneAdapter.cpp b/Engine/Plugins/Runtime/MeshModelingToolset/Source/ModelingComponents/Private/Scene/MeshSceneAdapter.cpp
--- a/Engine/Plugins/Runtime/MeshModelingToolset/Source/ModelingComponents/Private/Scene/MeshSceneAdapter.cpp
+++ b/Engine/Plugins/Runtime/MeshModelingToolset/Source/ModelingComponents/Private/Scene/MeshSceneAdapter.cpp
@@ -1505,6 +1505,10 @@ void FMeshSceneAdapter::InitializeSpatialWrappers(const TArray<FActorAdapter*>&
 	}
 }
 
+bool FMeshSceneAdapter::IsValid() const
+{
+	return !SpatialAdapters.IsEmpty();
+}
 
 void FMeshSceneAdapter::Build(const FMeshSceneAdapterBuildOptions& BuildOptions)
 {
diff --git a/Engine/Plugins/Runtime/MeshModelingToolset/Source/ModelingComponents/Public/Scene/MeshSceneAdapter.h b/Engine/Plugins/Runtime/MeshModelingToolset/Source/ModelingComponents/Public/Scene/MeshSceneAdapter.h
--- a/Engine/Plugins/Runtime/MeshModelingToolset/Source/ModelingComponents/Public/Scene/MeshSceneAdapter.h
+++ b/Engine/Plugins/Runtime/MeshModelingToolset/Source/ModelingComponents/Public/Scene/MeshSceneAdapter.h
@@ -388,6 +388,9 @@ public:
 	 */
 	virtual void ProcessActorChildMeshes(TFunctionRef<void(const FActorAdapter* ActorAdapter, const FActorChildMesh* ChildMesh)> ProcessFunc);
 
+	/** Whether valid actors/components were added to the scene representation */
+	bool IsValid() const;
+
 protected:
 	// top-level list of ActorAdapters, which represent each Actor and set of Components
 	TArray<TUniquePtr<FActorAdapter>> SceneActors;
