Approximate Actors: Changed error to warning

diff --git a/Engine/Plugins/Experimental/MeshModelingToolsetExp/Source/GeometryProcessingAdapters/Private/ApproximateActorsImpl.cpp b/Engine/Plugins/Experimental/MeshModelingToolsetExp/Source/GeometryProcessingAdapters/Private/ApproximateActorsImpl.cpp
--- a/Engine/Plugins/Experimental/MeshModelingToolsetExp/Source/GeometryProcessingAdapters/Private/ApproximateActorsImpl.cpp
+++ b/Engine/Plugins/Experimental/MeshModelingToolsetExp/Source/GeometryProcessingAdapters/Private/ApproximateActorsImpl.cpp
@@ -1060,7 +1060,7 @@ void FApproximateActorsImpl::GenerateApproximationForActorSet(const FInput& Inpu
 
 	if (!Scene->IsValid())
 	{
-		UE_LOG(LogApproximateActors, Error, TEXT("No valid input actors/components - unable to generate mesh"));
+		UE_LOG(LogApproximateActors, Warning, TEXT("No valid input actors/components - unable to generate mesh"));
 		ResultsOut.ResultCode = EResultCode::MeshGenerationFailed;
 		return;
 	}
