Add option for mesh inspector to draw edges w/out depth testing, so it is easier to see all the holes / potential issues in a mesh

diff --git a/Engine/Plugins/Experimental/MeshModelingToolsetExp/Source/MeshModelingToolsExp/Private/MeshInspectorTool.cpp b/Engine/Plugins/Experimental/MeshModelingToolsetExp/Source/MeshModelingToolsExp/Private/MeshInspectorTool.cpp
--- a/Engine/Plugins/Experimental/MeshModelingToolsetExp/Source/MeshModelingToolsExp/Private/MeshInspectorTool.cpp
+++ b/Engine/Plugins/Experimental/MeshModelingToolsetExp/Source/MeshModelingToolsExp/Private/MeshInspectorTool.cpp
@@ -199,7 +199,6 @@ void UMeshInspectorTool::Setup()
 
 	DrawnLineSet->SetupAttachment(PreviewMesh->GetRootComponent());
 
-	DrawnLineSet->SetLineMaterial(ToolSetupUtil::GetDefaultLineComponentMaterial(GetToolManager()));
 	DrawnLineSet->RegisterComponent();
 
 	Precompute();
@@ -594,6 +593,7 @@ void UMeshInspectorTool::UpdateVisualization()
 	const FDynamicMesh3* TargetMesh = PreviewMesh->GetPreviewDynamicMesh();
 	FVector3d A, B;
 
+	DrawnLineSet->SetLineMaterial(ToolSetupUtil::GetDefaultLineComponentMaterial(GetToolManager(), !Settings->bDrawHiddenEdgesAndSeams));
 	DrawnLineSet->Clear();
 	if (Settings->bBoundaryEdges)
 	{
diff --git a/Engine/Plugins/Experimental/MeshModelingToolsetExp/Source/MeshModelingToolsExp/Public/MeshInspectorTool.h b/Engine/Plugins/Experimental/MeshModelingToolsetExp/Source/MeshModelingToolsExp/Public/MeshInspectorTool.h
--- a/Engine/Plugins/Experimental/MeshModelingToolsetExp/Source/MeshModelingToolsExp/Public/MeshInspectorTool.h
+++ b/Engine/Plugins/Experimental/MeshModelingToolsetExp/Source/MeshModelingToolsExp/Public/MeshInspectorTool.h
@@ -96,6 +96,10 @@ public:
 	UPROPERTY(EditAnywhere, Category = Options)
 	bool bTangentVectors = false;
 
+	/** Toggle visibility of occluded boundary edges and seams */
+	UPROPERTY(EditAnywhere, Category = Options)
+	bool bDrawHiddenEdgesAndSeams = true;
+
 	/** Length of line segments representing normal vectors */
 	UPROPERTY(EditAnywhere, Category = Options, AdvancedDisplay, meta = (EditCondition = "bNormalVectors",
 		UIMin="0", UIMax="400", ClampMin = "0", ClampMax = "1000000000.0"))
