fix crash when inspect tool properties are modified after tool shutdown

(this happens when the editor closes a color picker window as part of cleaning up the tool properties panel)

diff --git a/Engine/Plugins/Experimental/MeshModelingToolsetExp/Source/MeshModelingToolsExp/Private/MeshInspectorTool.cpp b/Engine/Plugins/Experimental/MeshModelingToolsetExp/Source/MeshModelingToolsExp/Private/MeshInspectorTool.cpp
--- a/Engine/Plugins/Experimental/MeshModelingToolsetExp/Source/MeshModelingToolsExp/Private/MeshInspectorTool.cpp
+++ b/Engine/Plugins/Experimental/MeshModelingToolsetExp/Source/MeshModelingToolsExp/Private/MeshInspectorTool.cpp
@@ -529,7 +529,11 @@ void UMeshInspectorTool::OnPropertyModified(UObject* PropertySet, FProperty* Pro
 
 void UMeshInspectorTool::UpdateVisualization()
 {
-	if (!MaterialSettings) return;
+	// return if tool is not in a valid state (e.g., has already shut down)
+	if (!MaterialSettings || !PreviewMesh)
+	{
+		return;
+	}
 
 	PreviewMesh->EnableWireframe(Settings->bWireframe);
 
