Skeleton Editing Tool: Copy, Paste, Duplicate bones

diff --git a/Engine/Plugins/Experimental/Animation/SkeletalMeshModelingTools/Source/SkeletalMeshModelingTools/Private/SReferenceSkeletonTree.cpp b/Engine/Plugins/Experimental/Animation/SkeletalMeshModelingTools/Source/SkeletalMeshModelingTools/Private/SReferenceSkeletonTree.cpp
--- a/Engine/Plugins/Experimental/Animation/SkeletalMeshModelingTools/Source/SkeletalMeshModelingTools/Private/SReferenceSkeletonTree.cpp
+++ b/Engine/Plugins/Experimental/Animation/SkeletalMeshModelingTools/Source/SkeletalMeshModelingTools/Private/SReferenceSkeletonTree.cpp
@@ -4,6 +4,7 @@
 
 #include "SkeletonModifier.h"
 #include "SkeletalMeshModelingToolsCommands.h"
+#include "SkeletonClipboard.h"
 
 #include "SPositiveActionButton.h"
 
@@ -13,6 +14,7 @@
 
 #include "Framework/MultiBox/MultiBoxBuilder.h"
 #include "Editor.h"
+
 #include "Misc/ITransaction.h"
 
 #define LOCTEXT_NAMESPACE "SReferenceSkeletonTree"
@@ -359,6 +361,18 @@ void SReferenceSkeletonTree::BindCommands()
 	CommandList->MapAction(Commands.UnParentBone,
 		FExecuteAction::CreateSP(this, &SReferenceSkeletonTree::HandleUnParentBone),
 		FCanExecuteAction::CreateSP(this, &SReferenceSkeletonTree::CanUnParentBone));
+
+	CommandList->MapAction(Commands.CopyBones,
+		FExecuteAction::CreateSP(this, &SReferenceSkeletonTree::HandleCopyBones),
+		FCanExecuteAction::CreateSP(this, &SReferenceSkeletonTree::CanCopyBones));
+	
+	CommandList->MapAction(Commands.PasteBones,
+		FExecuteAction::CreateSP(this, &SReferenceSkeletonTree::HandlePasteBones),
+		FCanExecuteAction::CreateSP(this, &SReferenceSkeletonTree::CanPasteBones));
+	
+	CommandList->MapAction(Commands.UnParentBone,
+		FExecuteAction::CreateSP(this, &SReferenceSkeletonTree::HandleDuplicateBones),
+		FCanExecuteAction::CreateSP(this, &SReferenceSkeletonTree::CanDuplicateBones));
 }
 
 void SReferenceSkeletonTree::HandleNewBone()
@@ -466,6 +480,99 @@ bool SReferenceSkeletonTree::CanUnParentBone() const
 	return HasSelectedItems();
 }
 
+void SReferenceSkeletonTree::HandleCopyBones() const
+{
+	if (!ensure(Modifier.IsValid()))
+	{
+		return;
+	}
+
+	TArray<FName> BoneNames;
+	GetSelectedBoneNames(BoneNames);
+
+	if (BoneNames.IsEmpty())
+	{
+		return;
+	}
+
+	SkeletonClipboard::CopyToClipboard(*Modifier.Get(), BoneNames);
+}
+
+bool SReferenceSkeletonTree::CanCopyBones() const
+{
+	return Modifier.IsValid() && HasSelectedItems();
+}
+
+void SReferenceSkeletonTree::HandlePasteBones()
+{
+	TArray<FName> BoneNames;
+	GetSelectedBoneNames(BoneNames);
+
+	const FName DefaultParent = BoneNames.IsEmpty() ? NAME_None : BoneNames[0];  
+
+	BeginChange();
+	
+	const TArray<FName> NewBones = SkeletonClipboard::PasteFromClipboard(*Modifier.Get(), DefaultParent);
+	if (NewBones.IsEmpty())
+	{
+		CancelChange();
+		return;
+	}
+
+	static constexpr bool bRebuildAll = true;
+	RefreshTreeView(bRebuildAll);
+	SelectItemFromNames(NewBones);
+
+	if (Notifier.IsValid())
+	{
+		Notifier->Notify(NewBones, ESkeletalMeshNotifyType::HierarchyChanged);
+		Notifier->Notify(NewBones, ESkeletalMeshNotifyType::BonesSelected);
+	}
+
+	EndChange();
+}
+
+bool SReferenceSkeletonTree::CanPasteBones() const
+{
+	return Modifier.IsValid() && SkeletonClipboard::IsClipboardValid();
+}
+
+void SReferenceSkeletonTree::HandleDuplicateBones()
+{
+	HandleCopyBones();
+
+	if (!SkeletonClipboard::IsClipboardValid())
+	{
+		return;
+	}  
+
+	BeginChange();
+	
+	const TArray<FName> NewBones = SkeletonClipboard::PasteFromClipboard(*Modifier.Get(), NAME_None);
+	if (NewBones.IsEmpty())
+	{
+		CancelChange();
+		return;
+	}
+
+	static constexpr bool bRebuildAll = true;
+	RefreshTreeView(bRebuildAll);
+	SelectItemFromNames(NewBones);
+
+	if (Notifier.IsValid())
+	{
+		Notifier->Notify(NewBones, ESkeletalMeshNotifyType::HierarchyChanged);
+		Notifier->Notify(NewBones, ESkeletalMeshNotifyType::BonesSelected);
+	}
+
+	EndChange();
+}
+
+bool SReferenceSkeletonTree::CanDuplicateBones() const
+{
+	return CanCopyBones();
+}
+
 void SReferenceSkeletonTree::GetSelectedBoneNames(TArray<FName>& OutSelectedBoneNames) const
 {
 	OutSelectedBoneNames.Reset();
@@ -768,6 +875,12 @@ TSharedPtr<SWidget> SReferenceSkeletonTree::CreateContextMenu()
 		MenuBuilder.AddMenuEntry(Commands.NewBone);
 
 		MenuBuilder.EndSection();
+
+		MenuBuilder.BeginSection("CopyPasteBones", LOCTEXT("CopyPasteBonesOperations", "Copy & Paste"));
+
+		MenuBuilder.AddMenuEntry(Commands.PasteBones);
+
+		MenuBuilder.EndSection();
 	}
 	else
 	{
@@ -779,6 +892,14 @@ TSharedPtr<SWidget> SReferenceSkeletonTree::CreateContextMenu()
 		MenuBuilder.AddMenuEntry(Commands.RenameBone);
 
 		MenuBuilder.EndSection();
+
+		MenuBuilder.BeginSection("CopyPasteBones", LOCTEXT("CopyPasteBonesOperations", "Copy & Paste"));
+
+		MenuBuilder.AddMenuEntry(Commands.CopyBones);
+		MenuBuilder.AddMenuEntry(Commands.PasteBones);
+		MenuBuilder.AddMenuEntry(Commands.DuplicateBones);
+
+		MenuBuilder.EndSection();
 	}
 	
 	return MenuBuilder.MakeWidget();
diff --git a/Engine/Plugins/Experimental/Animation/SkeletalMeshModelingTools/Source/SkeletalMeshModelingTools/Private/SReferenceSkeletonTree.h b/Engine/Plugins/Experimental/Animation/SkeletalMeshModelingTools/Source/SkeletalMeshModelingTools/Private/SReferenceSkeletonTree.h
--- a/Engine/Plugins/Experimental/Animation/SkeletalMeshModelingTools/Source/SkeletalMeshModelingTools/Private/SReferenceSkeletonTree.h
+++ b/Engine/Plugins/Experimental/Animation/SkeletalMeshModelingTools/Source/SkeletalMeshModelingTools/Private/SReferenceSkeletonTree.h
@@ -201,6 +201,14 @@ private:
 	void HandleUnParentBone();
 	bool CanUnParentBone() const;
 
+	// Copy & Paste
+	void HandleCopyBones() const;
+	bool CanCopyBones() const;
+	void HandlePasteBones();
+	bool CanPasteBones() const;
+	void HandleDuplicateBones();
+	bool CanDuplicateBones() const;
+
 	// Filter
 	void OnFilterTextChanged(const FText& SearchText);
 	
diff --git a/Engine/Plugins/Experimental/Animation/SkeletalMeshModelingTools/Source/SkeletalMeshModelingTools/Private/SkeletalMeshModelingToolsCommands.cpp b/Engine/Plugins/Experimental/Animation/SkeletalMeshModelingTools/Source/SkeletalMeshModelingTools/Private/SkeletalMeshModelingToolsCommands.cpp
--- a/Engine/Plugins/Experimental/Animation/SkeletalMeshModelingTools/Source/SkeletalMeshModelingTools/Private/SkeletalMeshModelingToolsCommands.cpp
+++ b/Engine/Plugins/Experimental/Animation/SkeletalMeshModelingTools/Source/SkeletalMeshModelingTools/Private/SkeletalMeshModelingToolsCommands.cpp
@@ -15,6 +15,10 @@ void FSkeletalMeshModelingToolsCommands::RegisterCommands()
 	UI_COMMAND(RemoveBone, "Remove Bone", "Remove selected bone(s).", EUserInterfaceActionType::Button, FInputChord(EKeys::Delete), FInputChord(EKeys::BackSpace));
 	UI_COMMAND(UnParentBone, "Unparent Bone", "Unparent selected bone(s).", EUserInterfaceActionType::Button, FInputChord(EKeys::P, EModifierKey::Shift));
 	UI_COMMAND(RenameBone, "Rename Bone", "Rename the selected bone.", EUserInterfaceActionType::Button, FInputChord(EKeys::F2));
+
+	UI_COMMAND(CopyBones, "Copy Bone(s)", "Copy selected bone(s).", EUserInterfaceActionType::Button, FInputChord(EKeys::C, EModifierKey::Control));
+	UI_COMMAND(PasteBones, "Paste Bone(s)", "Paste selected bone(s).", EUserInterfaceActionType::Button, FInputChord(EKeys::V, EModifierKey::Control));
+	UI_COMMAND(DuplicateBones, "Duplicate Bone(s)", "Duplicate selected bone(s).", EUserInterfaceActionType::Button, FInputChord(EKeys::D, EModifierKey::Control));
 }
 
 const FSkeletalMeshModelingToolsCommands& FSkeletalMeshModelingToolsCommands::Get()
diff --git a/Engine/Plugins/Experimental/Animation/SkeletalMeshModelingTools/Source/SkeletalMeshModelingTools/Private/SkeletalMeshModelingToolsCommands.h b/Engine/Plugins/Experimental/Animation/SkeletalMeshModelingTools/Source/SkeletalMeshModelingTools/Private/SkeletalMeshModelingToolsCommands.h
--- a/Engine/Plugins/Experimental/Animation/SkeletalMeshModelingTools/Source/SkeletalMeshModelingTools/Private/SkeletalMeshModelingToolsCommands.h
+++ b/Engine/Plugins/Experimental/Animation/SkeletalMeshModelingTools/Source/SkeletalMeshModelingTools/Private/SkeletalMeshModelingToolsCommands.h
@@ -31,6 +31,11 @@ public:
 	TSharedPtr<FUICommandInfo> RemoveBone;
 	TSharedPtr<FUICommandInfo> UnParentBone;
 	TSharedPtr<FUICommandInfo> RenameBone;
+
+	// Copy & Paste commands
+	TSharedPtr<FUICommandInfo> CopyBones;
+	TSharedPtr<FUICommandInfo> PasteBones;
+	TSharedPtr<FUICommandInfo> DuplicateBones;
 };
 
 class FSkeletalMeshModelingToolsActionCommands : public TInteractiveToolCommands<FSkeletalMeshModelingToolsActionCommands>
diff --git a/Engine/Plugins/Experimental/MeshModelingToolsetExp/Source/MeshModelingToolsEditorOnlyExp/Private/SkeletalMesh/SkeletonEditingTool.cpp b/Engine/Plugins/Experimental/MeshModelingToolsetExp/Source/MeshModelingToolsEditorOnlyExp/Private/SkeletalMesh/SkeletonEditingTool.cpp
--- a/Engine/Plugins/Experimental/MeshModelingToolsetExp/Source/MeshModelingToolsEditorOnlyExp/Private/SkeletalMesh/SkeletonEditingTool.cpp
+++ b/Engine/Plugins/Experimental/MeshModelingToolsetExp/Source/MeshModelingToolsEditorOnlyExp/Private/SkeletalMesh/SkeletonEditingTool.cpp
@@ -13,6 +13,7 @@
 #include "Components/SkeletalMeshComponent.h"
 #include "UnrealClient.h"
 #include "HitProxies.h"
+#include "SkeletonClipboard.h"
 #include "ToolSetupUtil.h"
 #include "BaseBehaviors/ClickDragBehavior.h"
 #include "BaseGizmos/GizmoViewContext.h"
@@ -351,18 +352,29 @@ void USkeletonEditingTool::Shutdown(EToolShutdownType ShutdownType)
 	}
 }
 
-void USkeletonEditingTool::RegisterActions(FInteractiveToolActionSet& ActionSet)
+void USkeletonEditingTool::RegisterActions(FInteractiveToolActionSet& InOutActionSet)
 {
-	Super::RegisterActions(ActionSet);
+	Super::RegisterActions(InOutActionSet);
 
 	int32 ActionId = static_cast<int32>(EStandardToolActions::BaseClientDefinedActionID) + 400;
 	auto GetActionId = [&ActionId]
 	{
 		return ActionId++;
 	};
-	
-	// register New key
-	ActionSet.RegisterAction(this, GetActionId(), TEXT("CreateNewBone"),
+
+	RegisterCreateAction(InOutActionSet, GetActionId());
+	RegisterDeleteAction(InOutActionSet, GetActionId());
+	RegisterSelectAction(InOutActionSet, GetActionId());
+	RegisterParentAction(InOutActionSet, GetActionId());
+	RegisterUnParentAction(InOutActionSet, GetActionId());
+	RegisterCopyAction(InOutActionSet, GetActionId());
+	RegisterPasteAction(InOutActionSet, GetActionId());
+	RegisterDuplicateAction(InOutActionSet, GetActionId());
+}
+
+void USkeletonEditingTool::RegisterCreateAction(FInteractiveToolActionSet& InOutActionSet, const int32 InActionId)
+{
+	InOutActionSet.RegisterAction(this, InActionId, TEXT("CreateNewBone"),
 		LOCTEXT("CreateNewBone", "Create New Bone"),
 		LOCTEXT("CreateNewBoneDesc", "Create New Bone"),
 		EModifierKey::None, EKeys::N,
@@ -372,9 +384,11 @@ void USkeletonEditingTool::RegisterActions(FInteractiveToolActionSet& ActionSet)
 			UpdateGizmo();
 			GetToolManager()->DisplayMessage(LOCTEXT("Create", "Click & Drag to place a new bone."), EToolMessageLevel::UserNotification);
 		});
-	
-	// register Delete key
-	ActionSet.RegisterAction(this, GetActionId(), TEXT("DeleteSelectedBones"),
+}
+
+void USkeletonEditingTool::RegisterDeleteAction(FInteractiveToolActionSet& InOutActionSet, const int32 InActionId)
+{
+	InOutActionSet.RegisterAction(this, InActionId, TEXT("DeleteSelectedBones"),
 		LOCTEXT("DeleteSelectedBones", "Delete Selected Bone(s)"),
 		LOCTEXT("DeleteSelectedBonesDesc", "Delete Selected Bone(s)"),
 		EModifierKey::None, EKeys::Delete,
@@ -382,9 +396,11 @@ void USkeletonEditingTool::RegisterActions(FInteractiveToolActionSet& ActionSet)
 		{
 			RemoveBones();
 		});
+}
 
-	// register Select key
-	ActionSet.RegisterAction(this, GetActionId(), TEXT("SelectBones"),
+void USkeletonEditingTool::RegisterSelectAction(FInteractiveToolActionSet& InOutActionSet, const int32 InActionId)
+{
+	InOutActionSet.RegisterAction(this, InActionId, TEXT("SelectBones"),
 		LOCTEXT("SelectBone", "Select Bone"),
 		LOCTEXT("SelectDesc", "Select Bone"),
 		EModifierKey::None, EKeys::Escape,
@@ -397,9 +413,25 @@ void USkeletonEditingTool::RegisterActions(FInteractiveToolActionSet& ActionSet)
 				GetToolManager()->DisplayMessage(LOCTEXT("Select", "Click on a bone to select it."), EToolMessageLevel::UserNotification);
 			}
 		});
+}
 
-	// register UnParent key
-	ActionSet.RegisterAction(this, GetActionId(), TEXT("UnparentBones"),
+void USkeletonEditingTool::RegisterParentAction(FInteractiveToolActionSet& InOutActionSet, const int32 InActionId)
+{
+	InOutActionSet.RegisterAction(this, InActionId, TEXT("ParentBones"),
+		LOCTEXT("ParentBones", "Parent Bones"),
+		LOCTEXT("ParentBonesDesc", "Parent Bones"),
+		EModifierKey::None, EKeys::B, // FIXME find another shortcut
+		[this]()
+		{
+			Operation = EEditingOperation::Parent;
+			UpdateGizmo();
+			GetToolManager()->DisplayMessage(LOCTEXT("Parent", "Click on a bone to be set as the new parent."), EToolMessageLevel::UserNotification);
+		});
+}
+
+void USkeletonEditingTool::RegisterUnParentAction(FInteractiveToolActionSet& InOutActionSet, const int32 InActionId)
+{
+	InOutActionSet.RegisterAction(this, InActionId, TEXT("UnparentBones"),
 		LOCTEXT("UnparentBones", "Unparent Bones"),
 		LOCTEXT("UnparentBonesDesc", "Unparent Bones"),
 		EModifierKey::Shift, EKeys::P,
@@ -407,17 +439,99 @@ void USkeletonEditingTool::RegisterActions(FInteractiveToolActionSet& ActionSet)
 		{
 			UnParentBones();
 		});
-		
-	// register Parent key
-	ActionSet.RegisterAction(this, GetActionId(), TEXT("ParentBones"),
-		LOCTEXT("ParentBones", "Parent Bones"),
-		LOCTEXT("ParentBonesDesc", "Parent Bones"),
-		EModifierKey::None, EKeys::B, // FIXME find another shortcut
+}
+
+void USkeletonEditingTool::RegisterCopyAction(FInteractiveToolActionSet& InOutActionSet, const int32 InActionId)
+{
+	InOutActionSet.RegisterAction(this, InActionId, TEXT("CopyBones"),
+		LOCTEXT("CopyBones", "Copy Bone(s)"),
+		LOCTEXT("CopyBonesDesc", "Copy Bone(s)"),
+		EModifierKey::Control, EKeys::C,
 		[this]()
 		{
-			Operation = EEditingOperation::Parent;
-			UpdateGizmo();
-			GetToolManager()->DisplayMessage(LOCTEXT("Parent", "Click on a bone to be set as the new parent."), EToolMessageLevel::UserNotification);
+			if (Selection.IsEmpty())
+			{
+				return;
+			}
+			SkeletonClipboard::CopyToClipboard(*Modifier.Get(), Selection);
+		});
+}
+
+void USkeletonEditingTool::RegisterPasteAction(FInteractiveToolActionSet& InOutActionSet, const int32 InActionId)
+{
+	InOutActionSet.RegisterAction(this, InActionId, TEXT("PasteBones"),
+		LOCTEXT("PasteBones", "Paste Bone(s)"),
+		LOCTEXT("PasteBonesDesc", "Paste Bone(s)"),
+		EModifierKey::Control, EKeys::V,
+		[this]()
+		{
+			if (!SkeletonClipboard::IsClipboardValid())
+			{
+				return;
+			}
+			
+			const FName DefaultParent = Selection.IsEmpty() ? NAME_None : Selection[0];
+
+			TGuardValue OperationGuard(Operation, EEditingOperation::Create);
+			BeginChange();
+				
+			const TArray<FName> NewBones = SkeletonClipboard::PasteFromClipboard(*Modifier.Get(), DefaultParent);
+			if (!NewBones.IsEmpty())
+			{
+				Selection = NewBones;
+
+				if (NeedsNotification())
+				{
+					GetNotifier().Notify(NewBones, ESkeletalMeshNotifyType::HierarchyChanged);
+					GetNotifier().Notify(NewBones, ESkeletalMeshNotifyType::BonesSelected);
+				}
+
+				EndChange();
+			}
+			
+			CancelChange();
+		});
+}
+
+void USkeletonEditingTool::RegisterDuplicateAction(FInteractiveToolActionSet& InOutActionSet, const int32 InActionId)
+{
+	InOutActionSet.RegisterAction(this, InActionId, TEXT("DuplicateBones"),
+		LOCTEXT("DuplicateBones", "Duplicate Bone(s)"),
+		LOCTEXT("DuplicateBonesDesc", "Duplicate Bone(s)"),
+		EModifierKey::Control, EKeys::D,
+		[this]()
+		{
+			if (Selection.IsEmpty())
+			{
+				return;
+			}
+			
+			SkeletonClipboard::CopyToClipboard(*Modifier.Get(), Selection);
+
+			if (!SkeletonClipboard::IsClipboardValid())
+			{
+				return;
+			} 
+
+			TGuardValue OperationGuard(Operation, EEditingOperation::Create);
+			BeginChange();
+				
+			const TArray<FName> NewBones = SkeletonClipboard::PasteFromClipboard(*Modifier.Get(), NAME_None);
+			if (!NewBones.IsEmpty())
+			{
+				Selection = NewBones;
+
+				if (NeedsNotification())
+				{
+					GetNotifier().Notify(NewBones, ESkeletalMeshNotifyType::HierarchyChanged);
+					GetNotifier().Notify(NewBones, ESkeletalMeshNotifyType::BonesSelected);
+				}
+
+				EndChange();
+				return;
+			}
+			
+			CancelChange();
 		});
 }
 
diff --git a/Engine/Plugins/Experimental/MeshModelingToolsetExp/Source/MeshModelingToolsEditorOnlyExp/Public/SkeletalMesh/SkeletonEditingTool.h b/Engine/Plugins/Experimental/MeshModelingToolsetExp/Source/MeshModelingToolsEditorOnlyExp/Public/SkeletalMesh/SkeletonEditingTool.h
--- a/Engine/Plugins/Experimental/MeshModelingToolsetExp/Source/MeshModelingToolsEditorOnlyExp/Public/SkeletalMesh/SkeletonEditingTool.h
+++ b/Engine/Plugins/Experimental/MeshModelingToolsetExp/Source/MeshModelingToolsEditorOnlyExp/Public/SkeletalMesh/SkeletonEditingTool.h
@@ -241,6 +241,17 @@ private:
 
 	void NormalizeSelection();
 
+	// actions
+	void RegisterCreateAction(FInteractiveToolActionSet& InOutActionSet, const int32 InActionId);
+	void RegisterDeleteAction(FInteractiveToolActionSet& InOutActionSet, const int32 InActionId);
+	void RegisterSelectAction(FInteractiveToolActionSet& InOutActionSet, const int32 InActionId);
+	void RegisterParentAction(FInteractiveToolActionSet& InOutActionSet, const int32 InActionId);
+	void RegisterUnParentAction(FInteractiveToolActionSet& InOutActionSet, const int32 InActionId);
+	void RegisterCopyAction(FInteractiveToolActionSet& InOutActionSet, const int32 InActionId);
+	void RegisterPasteAction(FInteractiveToolActionSet& InOutActionSet, const int32 InActionId);
+	void RegisterDuplicateAction(FInteractiveToolActionSet& InOutActionSet, const int32 InActionId);
+
+	
 	TValueWatcher<TArray<FName>> SelectionWatcher;
 	TValueWatcher<EToolContextCoordinateSystem> CoordinateSystemWatcher; 
 	
diff --git a/Engine/Plugins/Experimental/MeshModelingToolsetExp/Source/SkeletalMeshModifiers/Private/SkeletonClipboard.cpp b/Engine/Plugins/Experimental/MeshModelingToolsetExp/Source/SkeletalMeshModifiers/Private/SkeletonClipboard.cpp
new file mode 100644
--- /dev/null
+++ b/Engine/Plugins/Experimental/MeshModelingToolsetExp/Source/SkeletalMeshModifiers/Private/SkeletonClipboard.cpp
@@ -0,0 +1,146 @@
+﻿// Copyright Epic Games, Inc. All Rights Reserved.
+
+
+#include "SkeletonClipboard.h"
+
+#include "SkeletonModifier.h"
+
+#include "HAL/PlatformApplicationMisc.h"
+
+namespace SkeletonClipboard
+{
+
+void CopyToClipboard(const USkeletonModifier& InModifier, const TArray<FName>& InBonesToCopy)
+{
+	FHierarchyClipboardData ClipboardData;
+	ClipboardData.Bones.Reserve(InBonesToCopy.Num());
+
+	const FReferenceSkeleton& RefSkeleton = InModifier.GetReferenceSkeleton();
+
+	// ensure ordering
+	TArray<FName> BonesToCopy(InBonesToCopy);
+	BonesToCopy.StableSort([&](const FName Bone0, const FName Bone1)
+	{
+		const int32 Index0 = RefSkeleton.FindBoneIndex(Bone0);
+		const int32 Index1 = RefSkeleton.FindBoneIndex(Bone1);
+		return Index0 < Index1;
+	});
+
+	// convert bone names to clipboard data
+	for (const FName& BoneName: BonesToCopy)
+	{
+		const FName ParentName = InModifier.GetParentName(BoneName);
+		const int32 ParentIndex = ParentName != NAME_None ? BonesToCopy.IndexOfByKey(ParentName) : INDEX_NONE;
+		const FTransform& Global = InModifier.GetBoneTransform(BoneName, true);
+		ClipboardData.Bones.Emplace(BoneName, Global, ParentIndex);
+	}
+	
+	// convert data to text
+	FString ClipboardText;
+	static const FHierarchyClipboardData DefaultData;
+	FHierarchyClipboardData::StaticStruct()->ExportText(ClipboardText, &ClipboardData, &DefaultData, nullptr, PPF_None, nullptr);
+
+	// copy to clipboard
+	FPlatformApplicationMisc::ClipboardCopy(*ClipboardText);
+}
+
+class FHierarchyImportErrorContext : public FOutputDevice
+{
+public:
+
+    int32 NumErrors;
+
+    FHierarchyImportErrorContext(const bool bInNotify)
+    	: FOutputDevice()
+    	, NumErrors(0)
+		, bNotify(bInNotify)
+    {}
+
+    virtual void Serialize(const TCHAR* V, ELogVerbosity::Type Verbosity, const FName& Category) override
+    {
+    	if (bNotify)
+    	{
+    		UE_LOG(LogTemp, Error, TEXT("Error Importing Bones: %s"), V);
+    	}
+    	NumErrors++;
+    }
+
+private:
+	bool bNotify = false;
+};
+
+TArray<FName> PasteFromClipboard(USkeletonModifier& InOutModifier, const FName& InDefaultParent)
+{
+	// Get the text from the clipboard.
+	FString ClipboardText;
+	FPlatformApplicationMisc::ClipboardPaste(ClipboardText);
+
+	FHierarchyClipboardData ClipboardData;
+	FHierarchyImportErrorContext ErrorPipe(true);
+	
+	UScriptStruct* DataStruct = FHierarchyClipboardData::StaticStruct();
+	DataStruct->ImportText(*ClipboardText, &ClipboardData, nullptr, PPF_None, &ErrorPipe, DataStruct->GetName(), true);
+
+	if (ErrorPipe.NumErrors > 0 || ClipboardData.Bones.IsEmpty())
+	{
+		return {};
+	}
+
+	// get new names
+	TArray<FName> NewBones;
+	NewBones.Reserve(ClipboardData.Bones.Num());
+
+	Algo::Transform(ClipboardData.Bones, NewBones, [&](const FBoneClipboardData& BoneData)
+	{
+		return InOutModifier.GetUniqueName(BoneData.BoneName, NewBones);
+	});
+
+	// get parent names & local transforms
+	TArray<FName> NewParents;
+	NewParents.Reserve(ClipboardData.Bones.Num());
+
+	TArray<FTransform> LocalTransforms;
+	LocalTransforms.Reserve(ClipboardData.Bones.Num());
+
+	for (const FBoneClipboardData& BoneData: ClipboardData.Bones)
+	{
+		const bool bIsParentFromClipboard = NewBones.IsValidIndex(BoneData.ParentIndex);
+		if (bIsParentFromClipboard)
+		{ // compute from clipboard
+			NewParents.Add(NewBones[BoneData.ParentIndex]);
+
+			const FTransform& ParentGlobal = ClipboardData.Bones[BoneData.ParentIndex].Global;
+			LocalTransforms.Add(BoneData.Global.GetRelativeTransform(ParentGlobal));
+		}
+		else
+		{ // compute from modifier
+			const bool bUseDefaultParent = InDefaultParent != NAME_None && InDefaultParent != BoneData.BoneName;
+			const FName ParentName = bUseDefaultParent ? InDefaultParent : InOutModifier.GetParentName(BoneData.BoneName);
+			NewParents.Add(ParentName);
+
+			const FTransform& ParentGlobal = InOutModifier.GetBoneTransform(ParentName, true);
+			LocalTransforms.Add(BoneData.Global.GetRelativeTransform(ParentGlobal));
+		}
+	}
+
+	// add bones
+	InOutModifier.AddBones(NewBones, NewParents, LocalTransforms);
+
+	return NewBones;
+}
+	
+bool IsClipboardValid()
+{
+	FString ClipboardText;
+	FPlatformApplicationMisc::ClipboardPaste(ClipboardText);
+
+	FHierarchyClipboardData ClipboardData;
+	FHierarchyImportErrorContext ErrorPipe(false);
+	
+	UScriptStruct* DataStruct = FHierarchyClipboardData::StaticStruct();
+	DataStruct->ImportText(*ClipboardText, &ClipboardData, nullptr, PPF_None, &ErrorPipe, DataStruct->GetName(), true);
+
+	return ErrorPipe.NumErrors == 0 && !ClipboardData.Bones.IsEmpty();
+}
+	
+}
diff --git a/Engine/Plugins/Experimental/MeshModelingToolsetExp/Source/SkeletalMeshModifiers/Public/SkeletonClipboard.h b/Engine/Plugins/Experimental/MeshModelingToolsetExp/Source/SkeletalMeshModifiers/Public/SkeletonClipboard.h
new file mode 100644
--- /dev/null
+++ b/Engine/Plugins/Experimental/MeshModelingToolsetExp/Source/SkeletalMeshModifiers/Public/SkeletonClipboard.h
@@ -0,0 +1,61 @@
+﻿// Copyright Epic Games, Inc. All Rights Reserved.
+
+#pragma once
+
+#include "Containers/Array.h"
+#include "UObject/NameTypes.h"
+#include "Math/MathFwd.h"
+
+#include "SkeletonClipboard.generated.h"
+
+class USkeletonModifier;
+
+/**
+ * The functions below are helper functions that simplify usage of copying/pasting bones on a FReferenceSkeleton
+ * using a USkeletonModifier.
+ */
+
+namespace SkeletonClipboard
+{
+	/**
+	 * Copies InBonesToCopy data (name, parent, transform) into the clipboard.
+	 * Hierarchy info & transforms are retrieved via the InModifier.
+	 */
+	SKELETALMESHMODIFIERS_API void CopyToClipboard(const USkeletonModifier& InModifier, const TArray<FName>& InBonesToCopy);
+
+	/**
+	 * Paste new bones into the InOutModifier the clipboard InBonesToCopy data (name, parent, transform) into the clipboard.
+	 * InDefaultParent is used as the default parent if a proper parent can't be retrieve from the clipboard data. 
+	 */
+	SKELETALMESHMODIFIERS_API TArray<FName> PasteFromClipboard(USkeletonModifier& InOutModifier, const FName& InDefaultParent);
+	
+	/**
+	 * Checks whether the clipboard data is valid. 
+	 */
+	SKELETALMESHMODIFIERS_API bool IsClipboardValid();
+}
+
+USTRUCT()
+struct FBoneClipboardData
+{
+	GENERATED_USTRUCT_BODY()
+
+	UPROPERTY()
+	FName BoneName = NAME_None;
+
+	UPROPERTY()
+	FTransform Global = FTransform::Identity;
+
+	UPROPERTY()
+	int32 ParentIndex = INDEX_NONE;
+};
+
+USTRUCT()
+struct FHierarchyClipboardData
+{
+	GENERATED_USTRUCT_BODY()
+	
+	UPROPERTY()
+	TArray<FBoneClipboardData> Bones;
+};
+
diff --git a/Engine/Plugins/Experimental/MeshModelingToolsetExp/Source/SkeletalMeshModifiers/SkeletalMeshModifiers.build.cs b/Engine/Plugins/Experimental/MeshModelingToolsetExp/Source/SkeletalMeshModifiers/SkeletalMeshModifiers.build.cs
--- a/Engine/Plugins/Experimental/MeshModelingToolsetExp/Source/SkeletalMeshModifiers/SkeletalMeshModifiers.build.cs
+++ b/Engine/Plugins/Experimental/MeshModelingToolsetExp/Source/SkeletalMeshModifiers/SkeletalMeshModifiers.build.cs
@@ -25,7 +25,8 @@ public class SkeletalMeshModifiers : ModuleRules
 				"SkeletalMeshUtilitiesCommon",
 				"RenderCore",
 				"CoreUObject",
-				"Engine"
+				"Engine",
+				"ApplicationCore"
 			}
 			);
 
