Fixed CIS build error

#rnx

diff --git a/Engine/Plugins/Experimental/MeshModelingToolsetExp/Source/SkeletalMeshModifiers/Private/SkeletonClipboard.cpp b/Engine/Plugins/Experimental/MeshModelingToolsetExp/Source/SkeletalMeshModifiers/Private/SkeletonClipboard.cpp
--- a/Engine/Plugins/Experimental/MeshModelingToolsetExp/Source/SkeletalMeshModifiers/Private/SkeletonClipboard.cpp
+++ b/Engine/Plugins/Experimental/MeshModelingToolsetExp/Source/SkeletalMeshModifiers/Private/SkeletonClipboard.cpp
@@ -30,9 +30,13 @@ void CopyToClipboard(const USkeletonModifier& InModifier, const TArray<FName>& I
 	for (const FName& BoneName: BonesToCopy)
 	{
 		const FName ParentName = InModifier.GetParentName(BoneName);
-		const int32 ParentIndex = ParentName != NAME_None ? BonesToCopy.IndexOfByKey(ParentName) : INDEX_NONE;
-		const FTransform& Global = InModifier.GetBoneTransform(BoneName, true);
-		ClipboardData.Bones.Emplace(BoneName, Global, ParentIndex);
+
+		FBoneClipboardData BoneData;
+			BoneData.BoneName = BoneName;
+			BoneData.ParentIndex = ParentName != NAME_None ? BonesToCopy.IndexOfByKey(ParentName) : INDEX_NONE;
+			BoneData.Global = InModifier.GetBoneTransform(BoneName, true);
+		
+		ClipboardData.Bones.Add(MoveTemp(BoneData));
 	}
 	
 	// convert data to text
