Fix source mesh being invisible when plane cut tool is started with show preview unchecked

diff --git a/Engine/Plugins/Experimental/MeshModelingToolsetExp/Source/MeshModelingToolsExp/Private/PlaneCutTool.cpp b/Engine/Plugins/Experimental/MeshModelingToolsetExp/Source/MeshModelingToolsExp/Private/PlaneCutTool.cpp
--- a/Engine/Plugins/Experimental/MeshModelingToolsetExp/Source/MeshModelingToolsExp/Private/PlaneCutTool.cpp
+++ b/Engine/Plugins/Experimental/MeshModelingToolsetExp/Source/MeshModelingToolsExp/Private/PlaneCutTool.cpp
@@ -192,8 +192,8 @@ void UPlaneCutTool::SetupPreviews()
 		// set initial preview to un-processed mesh, so stuff doesn't just disappear if the first cut takes a while
 		Preview->PreviewMesh->UpdatePreview(MeshesToCut[PreviewIdx]->GetMesh().Get());
 		Preview->PreviewMesh->SetTransform((FTransform)UE::ToolTarget::GetLocalToWorldTransform(Targets[PreviewIdx]));
-		Preview->SetVisibility(BasicProperties->bShowPreview);
 	}
+	UpdateVisibility();
 }
 
 
@@ -345,18 +345,24 @@ void UPlaneCutTool::PostEditChangeProperty(FPropertyChangedEvent& PropertyChange
 #endif
 
 
+void UPlaneCutTool::UpdateVisibility()
+{
+	for (int Idx = 0; Idx < Targets.Num(); Idx++)
+	{
+		UE::ToolTarget::SetSourceObjectVisible(Targets[Idx], !BasicProperties->bShowPreview);
+	}
+	for (UMeshOpPreviewWithBackgroundCompute* Preview : Previews)
+	{
+		Preview->SetVisibility(BasicProperties->bShowPreview);
+	}
+}
+
+
 void UPlaneCutTool::OnPropertyModified(UObject* PropertySet, FProperty* Property)
 {
 	if (Property && (Property->GetFName() == GET_MEMBER_NAME_CHECKED(UPlaneCutToolProperties, bShowPreview)))
 	{
-		for (int Idx = 0; Idx < Targets.Num(); Idx++)
-		{
-			UE::ToolTarget::SetSourceObjectVisible(Targets[Idx], !BasicProperties->bShowPreview);
-		}
-		for (UMeshOpPreviewWithBackgroundCompute* Preview : Previews)
-		{
-			Preview->SetVisibility(BasicProperties->bShowPreview);
-		}
+		UpdateVisibility();
 	}
 
 	InvalidatePreviews();
diff --git a/Engine/Plugins/Experimental/MeshModelingToolsetExp/Source/MeshModelingToolsExp/Public/PlaneCutTool.h b/Engine/Plugins/Experimental/MeshModelingToolsetExp/Source/MeshModelingToolsExp/Public/PlaneCutTool.h
--- a/Engine/Plugins/Experimental/MeshModelingToolsetExp/Source/MeshModelingToolsExp/Public/PlaneCutTool.h
+++ b/Engine/Plugins/Experimental/MeshModelingToolsetExp/Source/MeshModelingToolsExp/Public/PlaneCutTool.h
@@ -183,4 +183,6 @@ protected:
 	void InvalidatePreviews();
 
 	void GenerateAsset(const TArray<FDynamicMeshOpResult>& Results);
+
+	void UpdateVisibility();
 };
