Mesh To Collision: Add support for Geometry Selection

diff --git a/Engine/Plugins/Experimental/MeshModelingToolsetExp/Source/MeshModelingToolsExp/Private/Physics/SetCollisionGeometryTool.cpp b/Engine/Plugins/Experimental/MeshModelingToolsetExp/Source/MeshModelingToolsExp/Private/Physics/SetCollisionGeometryTool.cpp
--- a/Engine/Plugins/Experimental/MeshModelingToolsetExp/Source/MeshModelingToolsExp/Private/Physics/SetCollisionGeometryTool.cpp
+++ b/Engine/Plugins/Experimental/MeshModelingToolsetExp/Source/MeshModelingToolsExp/Private/Physics/SetCollisionGeometryTool.cpp
@@ -7,6 +7,11 @@
 #include "DynamicMesh/MeshTransforms.h"
 #include "DynamicMeshEditor.h"
 #include "Selections/MeshConnectedComponents.h"
+#include "Selections/GeometrySelectionUtil.h"
+#include "Selection/GeometrySelectionVisualization.h"
+#include "Selection/StoredMeshSelectionUtil.h"
+#include "PropertySets/GeometrySelectionVisualizationProperties.h"
+#include "GroupTopology.h"
 #include "DynamicSubmesh3.h"
 #include "Polygroups/PolygroupUtil.h"
 #include "Util/ColorConstants.h"
@@ -192,6 +197,26 @@ bool USetCollisionGeometryToolBuilder::CanBuildTool(const FToolBuilderState& Sce
 }
 
 
+void USetCollisionGeometryToolBuilder::InitializeNewTool(UMultiSelectionMeshEditingTool* Tool, const FToolBuilderState& SceneState) const
+{
+	const TArray<TObjectPtr<UToolTarget>> Targets = SceneState.TargetManager->BuildAllSelectedTargetable(SceneState, GetTargetRequirements());
+	Tool->SetTargets(Targets);
+	Tool->SetWorld(SceneState.World);
+
+	if (USetCollisionGeometryTool* CollisionTool = Cast<USetCollisionGeometryTool>(Tool))
+	{
+		if (Targets.Num() == 1) // Can only have a selection when there is one target
+		{
+			FGeometrySelection Selection;
+			if (GetCurrentGeometrySelectionForTarget(SceneState, Targets[0], Selection))
+			{
+				CollisionTool->SetGeometrySelection(MoveTemp(Selection));
+			}
+		}
+	}
+}
+
+
 UMultiSelectionMeshEditingTool* USetCollisionGeometryToolBuilder::CreateNewTool(const FToolBuilderState& SceneState) const
 {
 	return NewObject<USetCollisionGeometryTool>(SceneState.ToolManager);
@@ -310,6 +335,30 @@ void USetCollisionGeometryTool::Setup()
 	CollisionProps = NewObject<UPhysicsObjectToolPropertySet>(this);
 	AddToolPropertySource(CollisionProps);
 
+	if (InputGeometrySelection.IsEmpty() == false)
+	{
+		GeometrySelectionVizProperties = NewObject<UGeometrySelectionVisualizationProperties>(this);
+		GeometrySelectionVizProperties->RestoreProperties(this);
+		AddToolPropertySource(GeometrySelectionVizProperties);
+		GeometrySelectionVizProperties->Initialize(this);
+		GeometrySelectionVizProperties->SelectionElementType = static_cast<EGeometrySelectionElementType>(InputGeometrySelection.ElementType);
+		GeometrySelectionVizProperties->SelectionTopologyType = static_cast<EGeometrySelectionTopologyType>(InputGeometrySelection.TopologyType);
+
+		// Compute group topology if the selection has Polygroup topology, and do nothing otherwise
+		FGroupTopology GroupTopology(&InitialSourceMeshes[0], InputGeometrySelection.TopologyType == EGeometryTopologyType::Polygroup);
+
+		FTransformSRT3d ApplyTransform(UE::ToolTarget::GetLocalToWorldTransform(CollisionTarget));
+
+		GeometrySelectionViz = NewObject<UPreviewGeometry>(this);
+		GeometrySelectionViz->CreateInWorld(GetTargetWorld(), ApplyTransform);
+		InitializeGeometrySelectionVisualization(
+			GeometrySelectionViz,
+			GeometrySelectionVizProperties,
+			InitialSourceMeshes[0],
+			InputGeometrySelection,
+			&GroupTopology);
+	}
+
 	SetToolDisplayName(LOCTEXT("ToolName", "Mesh To Collision"));
 	GetToolManager()->DisplayMessage(
 		LOCTEXT("OnStartTool", "Initialize Simple Collision geometry for a Mesh from one or more input Meshes (including itself)."),
@@ -319,6 +368,11 @@ void USetCollisionGeometryTool::Setup()
 	bInputMeshesValid = false;
 }
 
+void USetCollisionGeometryTool::SetGeometrySelection(FGeometrySelection&& SelectionIn)
+{
+	InputGeometrySelection = MoveTemp(SelectionIn);
+}
+
 
 TUniquePtr<UE::Geometry::TGenericDataOperator<FPhysicsDataCollection>> USetCollisionGeometryTool::MakeNewOperator()
 {
@@ -404,6 +458,16 @@ void USetCollisionGeometryTool::OnShutdown(EToolShutdownType ShutdownType)
 
 	PreviewGeom->Disconnect();
 
+	if (GeometrySelectionViz)
+	{
+		GeometrySelectionViz->Disconnect();
+	}
+
+	if (GeometrySelectionVizProperties)
+	{
+		GeometrySelectionVizProperties->SaveProperties(this);
+	}
+
 	// show hidden sources
 	if (bSourcesHidden)
 	{
@@ -508,6 +572,13 @@ void USetCollisionGeometryTool::OnTick(float DeltaTime)
 		InvalidateCompute();
 	}
 
+	FText DisplayMessage;
+	if (!Settings->bAppendToExisting && !InputGeometrySelection.IsEmpty())
+	{
+		DisplayMessage = LOCTEXT("GeometrySelectionWithoutAppendToExisting", "The tool was invoked with a selection so you may want to enable 'Append to Existing'");
+	}
+	GetToolManager()->DisplayMessage(DisplayMessage, EToolMessageLevel::UserWarning);
+
 	if (Compute)
 	{
 		Compute->Tick(DeltaTime);
@@ -534,6 +605,11 @@ void USetCollisionGeometryTool::OnTick(float DeltaTime)
 	}
 
 	UE::PhysicsTools::UpdateCollisionGeometryVisualization(PreviewGeom, VizSettings);
+
+	if (GeometrySelectionViz)
+	{
+		UpdateGeometrySelectionVisualization(GeometrySelectionViz, GeometrySelectionVizProperties);
+	}
 }
 
 
@@ -576,15 +652,8 @@ void USetCollisionGeometryTool::OnSelectedGroupLayerChanged()
 }
 
 
-void USetCollisionGeometryTool::UpdateActiveGroupLayer()
+void USetCollisionGeometryTool::UpdateActiveGroupLayer(FDynamicMesh3* GroupLayersMesh)
 {
-	if (InitialSourceMeshes.Num() != 1)
-	{
-		ensure(false);		// should not get here
-		return;
-	}
-	FDynamicMesh3* GroupLayersMesh = &InitialSourceMeshes[0];
-
 	if (PolygroupLayerProperties->HasSelectedPolygroup() == false)
 	{
 		ActiveGroupSet = MakeUnique<UE::Geometry::FPolygroupSet>(GroupLayersMesh);
@@ -674,28 +743,41 @@ TArray<const T*> MakeRawPointerList(const TArray<TSharedPtr<T, ESPMode::ThreadSa
 
 void USetCollisionGeometryTool::PrecomputeInputMeshes()
 {
-	if (InitialSourceMeshes.Num() == 1)
-	{
-		UpdateActiveGroupLayer();
-	}
-
 	UToolTarget* CollisionTarget = Targets[Targets.Num() - 1];
 	FTransformSRT3d TargetTransform(UE::ToolTarget::GetLocalToWorldTransform(CollisionTarget));
 
+	// build input meshes.
 	InputMeshes.Reset();
 	InputMeshes.SetNum(SourceObjectIndices.Num());
-	ParallelFor(SourceObjectIndices.Num(), [&](int32 k)
+	if (!InputGeometrySelection.IsEmpty())
+	{
+		TSet<int> TriangleROI;
+		UE::Geometry::EnumerateSelectionTriangles(
+			InputGeometrySelection,
+			InitialSourceMeshes[0],
+			[&TriangleROI](int32 TriangleID) { TriangleROI.Add(TriangleID); });
+
+		// We dont discard attributes in the Submesh, we need them when building per-group input meshes
+		FDynamicSubmesh3 Submesh(&InitialSourceMeshes[0], TriangleROI.Array());
+
+		InputMeshes[0] = MakeShared<FDynamicMesh3, ESPMode::ThreadSafe>(MoveTemp(Submesh.GetSubmesh()));
+	}
+	else
 	{
-		FDynamicMesh3 SourceMesh = InitialSourceMeshes[k];
-		if (Settings->bUseWorldSpace)
+		ParallelFor(SourceObjectIndices.Num(), [&](int32 k)
 		{
-			FTransformSRT3d ToWorld(UE::ToolTarget::GetLocalToWorldTransform(Targets[k]));
-			MeshTransforms::ApplyTransform(SourceMesh, ToWorld, true);
-			MeshTransforms::ApplyTransformInverse(SourceMesh, TargetTransform, true);
-		}
-		SourceMesh.DiscardAttributes();
-		InputMeshes[k] = MakeShared<FDynamicMesh3, ESPMode::ThreadSafe>(MoveTemp(SourceMesh));
-	});
+			FDynamicMesh3 SourceMesh = InitialSourceMeshes[k];
+			if (Settings->bUseWorldSpace)
+			{
+				FTransformSRT3d ToWorld(UE::ToolTarget::GetLocalToWorldTransform(Targets[k]));
+				MeshTransforms::ApplyTransform(SourceMesh, ToWorld, true);
+				MeshTransforms::ApplyTransformInverse(SourceMesh, TargetTransform, true);
+			}
+			SourceMesh.DiscardAttributes();
+
+			InputMeshes[k] = MakeShared<FDynamicMesh3, ESPMode::ThreadSafe>(MoveTemp(SourceMesh));
+		});
+	}
 	InputMeshesApproximator = MakeShared<FMeshSimpleShapeApproximation, ESPMode::ThreadSafe>();
 	InputMeshesApproximator->InitializeSourceMeshes(MakeRawPointerList<FDynamicMesh3>(InputMeshes));
 
@@ -718,22 +800,39 @@ void USetCollisionGeometryTool::PrecomputeInputMeshes()
 	// build separated input meshes
 	SeparatedInputMeshes.Reset();
 	InitializeDerivedMeshSet(InputMeshes, SeparatedInputMeshes, 
-		[&](const FDynamicMesh3* Mesh, int32 Tri0, int32 Tri1) { return true; });
+		[](const FDynamicMesh3* Mesh, int32 Tri0, int32 Tri1)
+		{
+			return true;
+		});
 	SeparatedMeshesApproximator = MakeShared<FMeshSimpleShapeApproximation, ESPMode::ThreadSafe>();
 	SeparatedMeshesApproximator->InitializeSourceMeshes(MakeRawPointerList<FDynamicMesh3>(SeparatedInputMeshes));
 
 	// build per-group input meshes
 	PerGroupInputMeshes.Reset();
-	if (ActiveGroupSet.IsValid())
+	if (InputMeshes.Num() == 1)
 	{
-		check(InputMeshes.Num() == 1);
+		FDynamicMesh3* UseGroupLayerMesh = &InitialSourceMeshes[0];
+		if (!InputGeometrySelection.IsEmpty())
+		{
+			UseGroupLayerMesh = InputMeshes[0].Get();
+		}
+		UpdateActiveGroupLayer(UseGroupLayerMesh);
+
+		// Use the active polygroup layer when there is only one input
 		InitializeDerivedMeshSet(InputMeshes, PerGroupInputMeshes,
-			[&](const FDynamicMesh3* Mesh, int32 Tri0, int32 Tri1) { return ActiveGroupSet->GetTriangleGroup(Tri0) == ActiveGroupSet->GetTriangleGroup(Tri1); });
+			[this](const FDynamicMesh3* Mesh, int32 Tri0, int32 Tri1)
+			{
+				return ActiveGroupSet->GetTriangleGroup(Tri0) == ActiveGroupSet->GetTriangleGroup(Tri1);
+			});
 	}
 	else
 	{
+		// Use the default polygroup layer when there is more than one input
 		InitializeDerivedMeshSet(InputMeshes, PerGroupInputMeshes,
-			[&](const FDynamicMesh3* Mesh, int32 Tri0, int32 Tri1) { return Mesh->GetTriangleGroup(Tri0) == Mesh->GetTriangleGroup(Tri1); });
+			[](const FDynamicMesh3* Mesh, int32 Tri0, int32 Tri1)
+			{
+				return Mesh->GetTriangleGroup(Tri0) == Mesh->GetTriangleGroup(Tri1);
+			});
 	}
 	PerGroupMeshesApproximator = MakeShared<FMeshSimpleShapeApproximation, ESPMode::ThreadSafe>();
 	PerGroupMeshesApproximator->InitializeSourceMeshes(MakeRawPointerList<FDynamicMesh3>(PerGroupInputMeshes));
diff --git a/Engine/Plugins/Experimental/MeshModelingToolsetExp/Source/MeshModelingToolsExp/Public/Physics/SetCollisionGeometryTool.h b/Engine/Plugins/Experimental/MeshModelingToolsetExp/Source/MeshModelingToolsExp/Public/Physics/SetCollisionGeometryTool.h
--- a/Engine/Plugins/Experimental/MeshModelingToolsetExp/Source/MeshModelingToolsExp/Public/Physics/SetCollisionGeometryTool.h
+++ b/Engine/Plugins/Experimental/MeshModelingToolsetExp/Source/MeshModelingToolsExp/Public/Physics/SetCollisionGeometryTool.h
@@ -13,12 +13,14 @@
 #include "Physics/PhysicsDataCollection.h"
 #include "PropertySets/PolygroupLayersProperties.h"
 #include "Polygroups/PolygroupSet.h"
+#include "Selections/GeometrySelection.h"
 #include "TransformSequence.h"
 #include "ModelingOperators.h"
 #include "MeshOpPreviewHelpers.h"
 #include "SetCollisionGeometryTool.generated.h"
 
 class UPreviewGeometry;
+class UGeometrySelectionVisualizationProperties;
 PREDECLARE_GEOMETRY(class FMeshSimpleShapeApproximation)
 PREDECLARE_USE_GEOMETRY_CLASS(FDynamicMesh3);
 
@@ -29,6 +31,7 @@ class MESHMODELINGTOOLSEXP_API USetCollisionGeometryToolBuilder : public UMultiS
 
 public:
 	virtual bool CanBuildTool(const FToolBuilderState& SceneState) const override;
+	virtual void InitializeNewTool(UMultiSelectionMeshEditingTool* Tool, const FToolBuilderState& SceneState) const;
 	virtual UMultiSelectionMeshEditingTool* CreateNewTool(const FToolBuilderState& SceneState) const override;
 
 protected:
@@ -41,9 +44,17 @@ protected:
 UENUM()
 enum class ESetCollisionGeometryInputMode
 {
+	// Compute collision geometry using a combined mesh of all input objects
 	CombineAll = 0,
+
+	// Compute collision geometry for each input object
+	// Note: A Geometry Selection always counts as one input object
 	PerInputObject = 1,
+
+	// Compute collision geometry for each connected component of each input object
 	PerMeshComponent = 2,
+
+	// Compute collision geometry for each PolyGroup of each input object
 	PerMeshGroup = 3
 };
 
@@ -93,6 +104,7 @@ public:
 	UPROPERTY(EditAnywhere, Category = Options)
 	ECollisionGeometryType GeometryType = ECollisionGeometryType::AlignedBoxes;
 
+	// If true/false, Accept will append to/overwrite any existing collision geometry
 	UPROPERTY(EditAnywhere, Category = Options)
 	bool bAppendToExisting = false;
 
@@ -106,6 +118,7 @@ public:
 	UPROPERTY(EditAnywhere, Category = Options)
 	bool bRemoveContained = true;
 
+	// If true, discard all but MaxCount collision geometries with the largest volume
 	UPROPERTY(EditAnywhere, Category = Options)
 	bool bEnableMaxCount = true;
 
@@ -206,6 +219,8 @@ public:
 	virtual TUniquePtr<UE::Geometry::TGenericDataOperator<FPhysicsDataCollection>> MakeNewOperator() override;
 	// End IGenericDataOperatorFactory interface
 
+	void SetGeometrySelection(UE::Geometry::FGeometrySelection&& SelectionIn);
+
 protected:
 
 	UPROPERTY()
@@ -280,7 +295,7 @@ protected:
 
 	TUniquePtr<UE::Geometry::FPolygroupSet> ActiveGroupSet;
 	void OnSelectedGroupLayerChanged();
-	void UpdateActiveGroupLayer();
+	void UpdateActiveGroupLayer(FDynamicMesh3* GroupLayersMesh);
 
 	FTransform OrigTargetTransform;
 	UE::Geometry::FTransformSequence3d TargetInverseTransform;
@@ -290,4 +305,16 @@ protected:
 	TSharedPtr<FPhysicsDataCollection, ESPMode::ThreadSafe> GeneratedCollision;
 	TSharedPtr<TArray<FPhysicsDataCollection>, ESPMode::ThreadSafe> OtherInputsCollision;
 	TSharedPtr<TArray<FTransform3d>, ESPMode::ThreadSafe> OtherInputsTransforms;
+
+	//
+	// Geometry Selection
+	//
+
+	UE::Geometry::FGeometrySelection InputGeometrySelection;
+
+	UPROPERTY()
+	TObjectPtr<UGeometrySelectionVisualizationProperties> GeometrySelectionVizProperties = nullptr;
+
+	UPROPERTY()
+	TObjectPtr<UPreviewGeometry> GeometrySelectionViz = nullptr;
 };
