Make the Simple Collision Editor Tool set collision on a dynamic mesh component target via the function for doing so, rather than by writing it directly to the BodySetup.  Fixes an issue where the updated collision could be missed or lost.

diff --git a/Engine/Plugins/Experimental/MeshModelingToolsetExp/Source/MeshModelingToolsExp/Private/Physics/SimpleCollisionEditorTool.cpp b/Engine/Plugins/Experimental/MeshModelingToolsetExp/Source/MeshModelingToolsExp/Private/Physics/SimpleCollisionEditorTool.cpp
--- a/Engine/Plugins/Experimental/MeshModelingToolsetExp/Source/MeshModelingToolsExp/Private/Physics/SimpleCollisionEditorTool.cpp
+++ b/Engine/Plugins/Experimental/MeshModelingToolsetExp/Source/MeshModelingToolsExp/Private/Physics/SimpleCollisionEditorTool.cpp
@@ -163,10 +163,8 @@ void USimpleCollisionEditorTool::OnShutdown(EToolShutdownType ShutdownType)
 		}
 		else if (UDynamicMeshComponent* DynamicMeshComponent = Cast<UDynamicMeshComponent>(Component))
 		{
-			if (UBodySetup* BodySetup = DynamicMeshComponent->GetBodySetup())
-			{
-				UpdateBodySetup(BodySetup);
-			}
+			DynamicMeshComponent->SetSimpleCollisionShapes(PhysicsInfos->AggGeom, true);
+			DynamicMeshComponent->MarkRenderStateDirty();
 		}
 
 		// post the undo transaction
