Skin Weights Paint Tool: fixed crash when selecting invalid vertices ids

    we have to make sure that the vertex ids are safe to use as PolygonSelectionMechanic does not act on the mesh description but on the dynamic mesh that can duplicate vertices when dealing with degenerate triangles.
    cf. FMeshDescriptionToDynamicMesh::Convert for more details.

#rnx

diff --git a/Engine/Plugins/Experimental/MeshModelingToolsetExp/Source/MeshModelingToolsEditorOnlyExp/Private/SkeletalMesh/SkinWeightsPaintTool.cpp b/Engine/Plugins/Experimental/MeshModelingToolsetExp/Source/MeshModelingToolsEditorOnlyExp/Private/SkeletalMesh/SkinWeightsPaintTool.cpp
--- a/Engine/Plugins/Experimental/MeshModelingToolsetExp/Source/MeshModelingToolsEditorOnlyExp/Private/SkeletalMesh/SkinWeightsPaintTool.cpp
+++ b/Engine/Plugins/Experimental/MeshModelingToolsetExp/Source/MeshModelingToolsEditorOnlyExp/Private/SkeletalMesh/SkinWeightsPaintTool.cpp
@@ -2043,7 +2043,15 @@ void USkinWeightsPaintTool::GetSelectedVertices(TArray<int32>& OutVertexIndices)
 	const FGroupTopologySelection& Selection = PolygonSelectionMechanic->GetActiveSelection();
 	if (!Selection.SelectedCornerIDs.IsEmpty())
 	{
-		OutVertexIndices = Selection.SelectedCornerIDs.Array();
+		// we have to make sure that the vertex ids are safe to use as PolygonSelectionMechanic does not act on the
+		// mesh description but on the dynamic mesh that can duplicate vertices when dealing with degenerate triangles.
+		// cf. FMeshDescriptionToDynamicMesh::Convert for more details.
+		const FVertexArray& Vertices = EditedMesh->Vertices();
+		OutVertexIndices.Empty();
+		Algo::CopyIf(Selection.SelectedCornerIDs, OutVertexIndices, [&](int VertexID)
+		{
+			return Vertices.IsValid(VertexID);	
+		});
 	}
 }
 
