Skeleton Editing Tool: fixed adding new bone creating two transactions

#rnx

diff --git a/Engine/Plugins/Experimental/MeshModelingToolsetExp/Source/MeshModelingToolsEditorOnlyExp/Private/SkeletalMesh/SkeletonEditingTool.cpp b/Engine/Plugins/Experimental/MeshModelingToolsetExp/Source/MeshModelingToolsEditorOnlyExp/Private/SkeletalMesh/SkeletonEditingTool.cpp
--- a/Engine/Plugins/Experimental/MeshModelingToolsetExp/Source/MeshModelingToolsEditorOnlyExp/Private/SkeletalMesh/SkeletonEditingTool.cpp
+++ b/Engine/Plugins/Experimental/MeshModelingToolsetExp/Source/MeshModelingToolsEditorOnlyExp/Private/SkeletalMesh/SkeletonEditingTool.cpp
@@ -558,7 +558,6 @@ void USkeletonEditingTool::CreateNewBone()
 		Selection = {BoneName};
 		Properties->Name = BoneName;
 
-		EndChange();
 		return;
 	}
 
@@ -752,7 +751,10 @@ void USkeletonEditingTool::OnClickPress(const FInputDeviceRay& InPressPos)
 		PendingFunction();
 		PendingFunction.Reset();
 	}
-	BeginChange();
+	else
+	{ // make sure that the PendingFunction handles BeginChange if it needs to
+		BeginChange();
+	}
 }
 
 void USkeletonEditingTool::OnClickDrag(const FInputDeviceRay& InDragPos)
@@ -956,6 +958,11 @@ void USkeletonEditingTool::NormalizeSelection()
 
 void USkeletonEditingTool::OnClickRelease(const FInputDeviceRay& InReleasePos)
 {
+	if (Operation == EEditingOperation::Create)
+	{
+		return EndChange();
+	}
+	
 	TGuardValue OperationGuard(Operation, EEditingOperation::Transform);
 	EndChange();
 }
@@ -1315,8 +1322,8 @@ void USkeletonEditingTool::EndChange()
 	ActiveChange->StoreSkeleton(this);
 
 	static const UEnum* OperationEnum = StaticEnum<EEditingOperation>();
-	const FName OperationName = OperationEnum->GetNameByValue(static_cast<int64>(Operation));
-	const FText TransactionDesc = FText::Format(LOCTEXT("RefSkeletonChanged", "Skeleton Edit - {0}"), FText::FromName(OperationName));
+	const FString OperationString = OperationEnum->GetNameStringByValue(static_cast<int64>(Operation));
+	const FText TransactionDesc = FText::Format(LOCTEXT("RefSkeletonChanged", "Skeleton Edit - {0}"), FText::FromString(OperationString));
 	
 	UInteractiveToolManager* ToolManager = GetToolManager();
 	ToolManager->BeginUndoTransaction(TransactionDesc);
