Skeleton Editing Tools: Add negative axis for bone orient options

diff --git a/Engine/Plugins/Experimental/Animation/SkeletalMeshModelingTools/Source/SkeletalMeshModelingTools/Private/DetailCustomization/SkeletonEditingToolPropertyCustomizations.cpp b/Engine/Plugins/Experimental/Animation/SkeletalMeshModelingTools/Source/SkeletalMeshModelingTools/Private/DetailCustomization/SkeletonEditingToolPropertyCustomizations.cpp
--- a/Engine/Plugins/Experimental/Animation/SkeletalMeshModelingTools/Source/SkeletalMeshModelingTools/Private/DetailCustomization/SkeletonEditingToolPropertyCustomizations.cpp
+++ b/Engine/Plugins/Experimental/Animation/SkeletalMeshModelingTools/Source/SkeletalMeshModelingTools/Private/DetailCustomization/SkeletonEditingToolPropertyCustomizations.cpp
@@ -631,6 +631,10 @@ void FOrientingPropertiesDetailCustomization::CustomizeDetails(IDetailLayoutBuil
 			.HAlign(HAlign_Center)
 			.VAlign(VAlign_Center)
 			.Text(LOCTEXT("OrientButtonLabel", "Orient"))
+			.IsEnabled_Lambda([Properties]()
+			{
+				return Properties->Options.Primary != EOrientAxis::None; 
+			})
 			.OnClicked_Lambda([Properties]()
 			{
 				Properties->OrientBones();
diff --git a/Engine/Plugins/Experimental/MeshModelingToolsetExp/Source/MeshModelingToolsEditorOnlyExp/Private/SkeletalMesh/SkeletonEditingTool.cpp b/Engine/Plugins/Experimental/MeshModelingToolsetExp/Source/MeshModelingToolsEditorOnlyExp/Private/SkeletalMesh/SkeletonEditingTool.cpp
--- a/Engine/Plugins/Experimental/MeshModelingToolsetExp/Source/MeshModelingToolsEditorOnlyExp/Private/SkeletalMesh/SkeletonEditingTool.cpp
+++ b/Engine/Plugins/Experimental/MeshModelingToolsetExp/Source/MeshModelingToolsEditorOnlyExp/Private/SkeletalMesh/SkeletonEditingTool.cpp
@@ -18,8 +18,6 @@
 #include "BaseBehaviors/ClickDragBehavior.h"
 #include "BaseGizmos/GizmoViewContext.h"
 
-#include "Algo/Count.h"
-#include "BaseGizmos/TransformGizmoUtil.h"
 #include "SkeletalMesh/SkeletonTransformProxy.h"
 #include "TargetInterfaces/PrimitiveComponentBackedTarget.h"
 
@@ -1385,23 +1383,45 @@ void UOrientingProperties::PostEditChangeProperty(FPropertyChangedEvent &Propert
 
 	if (PropertyChangedEvent.ChangeType == EPropertyChangeType::ValueSet)
 	{
-		auto CheckAxis = [&](const TEnumAsByte<EAxis::Type>& InRef, TEnumAsByte<EAxis::Type>& OutOther)
+		auto CheckAxis = [&](const EOrientAxis& InRef, EOrientAxis& OutOther)
 		{
-			if (OutOther != InRef)
-			{
-				return;
-			}
-
 			switch (InRef)
 			{
-			case EAxis::X:
-				OutOther = EAxis::Y;
+			case EOrientAxis::PositiveX:
+				if (OutOther == InRef || OutOther == EOrientAxis::NegativeX)
+				{
+					OutOther = EOrientAxis::PositiveY;
+				}
+				break;
+			case EOrientAxis::PositiveY:
+				if (OutOther == InRef || OutOther == EOrientAxis::NegativeY)
+				{
+					OutOther = EOrientAxis::PositiveZ;
+				}
+				break;
+			case EOrientAxis::PositiveZ:
+				if (OutOther == InRef || OutOther == EOrientAxis::NegativeZ)
+				{
+					OutOther = EOrientAxis::PositiveX;
+				}
 				break;
-			case EAxis::Y:
-				OutOther = EAxis::Z;
+			case EOrientAxis::NegativeX:
+				if (OutOther == InRef || OutOther == EOrientAxis::PositiveX)
+				{
+					OutOther = EOrientAxis::PositiveY;
+				}
 				break;
-			case EAxis::Z:
-				OutOther = EAxis::X;
+			case EOrientAxis::NegativeY:
+				if (OutOther == InRef || OutOther == EOrientAxis::PositiveY)
+				{
+					OutOther = EOrientAxis::PositiveZ;
+				}
+				break;
+			case EOrientAxis::NegativeZ:
+				if (OutOther == InRef || OutOther == EOrientAxis::PositiveZ)
+				{
+					OutOther = EOrientAxis::PositiveX;
+				}
 				break;
 			default:
 				break;
@@ -1411,12 +1431,6 @@ void UOrientingProperties::PostEditChangeProperty(FPropertyChangedEvent &Propert
 		const FName PropertyName = PropertyChangedEvent.GetPropertyName();
 		if (PropertyName == GET_MEMBER_NAME_CHECKED(FOrientOptions, Primary))
 		{
-			if (Options.Primary == EAxis::None)
-			{
-				Options.Primary = EAxis::X;
-				Options.Secondary = EAxis::Y;
-				return;
-			}
 			CheckAxis(Options.Primary, Options.Secondary);
 			return;
 		}
diff --git a/Engine/Plugins/Experimental/MeshModelingToolsetExp/Source/SkeletalMeshModifiers/Private/SkeletonModifier.cpp b/Engine/Plugins/Experimental/MeshModelingToolsetExp/Source/SkeletalMeshModifiers/Private/SkeletonModifier.cpp
--- a/Engine/Plugins/Experimental/MeshModelingToolsetExp/Source/SkeletalMeshModifiers/Private/SkeletonModifier.cpp
+++ b/Engine/Plugins/Experimental/MeshModelingToolsetExp/Source/SkeletalMeshModifiers/Private/SkeletonModifier.cpp
@@ -65,49 +65,70 @@ FVector FMirrorOptions::MirrorVector(const FVector& InVector) const
 
 FTransform FOrientOptions::OrientTransform(const FVector& InPrimaryTarget, const FTransform& InTransform) const
 {
-	if (Primary == EAxis::None || InPrimaryTarget.IsNearlyZero())
+	if (Primary == EOrientAxis::None || InPrimaryTarget.IsNearlyZero())
 	{
 		return InTransform;
 	}
-
+	
+	auto GetOrientVector = [](const EOrientAxis& InOrientAxis)
+	{
+		switch(InOrientAxis)
+		{
+		case EOrientAxis::None:	return FVector::ZeroVector;
+		case EOrientAxis::PositiveX: return FVector::XAxisVector;
+		case EOrientAxis::PositiveY: return FVector::YAxisVector;
+		case EOrientAxis::PositiveZ: return FVector::ZAxisVector;
+		case EOrientAxis::NegativeX: return -FVector::XAxisVector;
+		case EOrientAxis::NegativeY: return -FVector::YAxisVector;
+		case EOrientAxis::NegativeZ: return -FVector::ZAxisVector;
+		default: break;
+		}
+		return FVector::ZeroVector;
+	};
+	
 	FTransform Transform = InTransform;
 
-	FVector PrimaryAxis(ForceInitToZero); PrimaryAxis.SetComponentForAxis(Primary, 1.f);
+	const FVector PrimaryOrientVector = GetOrientVector(Primary);
+	const FVector PrimaryAxis = Transform.TransformVectorNoScale(PrimaryOrientVector).GetSafeNormal();
 	const FVector PrimaryTarget = InPrimaryTarget.GetSafeNormal();
 	
 	// orient primary axis towards InPrimaryTarget
 	{
-		const FVector Axis = Transform.TransformVectorNoScale(PrimaryAxis).GetSafeNormal();
-		const FQuat Rotation = FQuat::FindBetweenNormals(Axis, PrimaryTarget);
+		const FQuat Rotation = FQuat::FindBetweenNormals(PrimaryAxis, PrimaryTarget);
 		const FQuat NewRotation = (Rotation * Transform.GetRotation()).GetNormalized();
 		Transform.SetRotation(NewRotation);
 	}
 
-	if (Secondary == EAxis::None || SecondaryTarget.IsNearlyZero())
-	{ // no need to use secondary axis
+	// no need to use secondary axis
+	if (Secondary == Primary || Secondary == EOrientAxis::None || SecondaryTarget.IsNearlyZero())
+	{
 		return Transform;
 	}
 
-	FVector Target = SecondaryTarget.GetSafeNormal();
-	if (FMath::IsNearlyEqual(FMath::Abs(FVector::DotProduct(PrimaryTarget, Target)), 1.0f))
-	{ // both targets are parallel
+	FVector SecondTarget = SecondaryTarget.GetSafeNormal();
+	if (FMath::IsNearlyEqual(FMath::Abs(FVector::DotProduct(PrimaryTarget, SecondTarget)), 1.0f))
+	{
+		// both targets are parallel
 		return Transform;
 	}
 
 	// orient secondary axis towards SecondaryDirection
 	{
-		FVector SecondaryAxis(ForceInitToZero); SecondaryAxis.SetComponentForAxis(Secondary, 1.f);
+		// project on primary 
+		SecondTarget = SecondTarget - FVector::DotProduct(SecondTarget, PrimaryTarget) * PrimaryTarget;
 		
+		if (!SecondTarget.IsNearlyZero())
 		{
-			const FVector Axis = Transform.TransformVectorNoScale(PrimaryAxis).GetSafeNormal();
-			Target = Target - FVector::DotProduct(Target, Axis) * Axis;
-		}
+			SecondTarget = SecondTarget.GetSafeNormal();
 
-		if (!Target.IsNearlyZero())
-		{
-			Target = Target.GetSafeNormal();
-			const FVector Axis = Transform.TransformVectorNoScale(SecondaryAxis).GetSafeNormal();
-			const FQuat Rotation = FQuat::FindBetweenNormals(Axis, Target);
+			const FVector SecondaryOrientVector = GetOrientVector(Secondary);
+			const FVector SecondaryAxis = Transform.TransformVectorNoScale(SecondaryOrientVector).GetSafeNormal();
+
+			// if they are opposites, we only need to rotate 180 degrees around PrimaryTarget
+			const double DotProduct = FVector::DotProduct(SecondaryAxis, SecondTarget);
+			const bool bAreOpposites = (DotProduct + 1.0) < KINDA_SMALL_NUMBER;
+			
+			const FQuat Rotation = bAreOpposites ? FQuat(PrimaryTarget, PI) : FQuat::FindBetweenNormals(SecondaryAxis, SecondTarget);
 			const FQuat NewRotation = (Rotation * Transform.GetRotation()).GetNormalized();
 			Transform.SetRotation(NewRotation);
 		}
diff --git a/Engine/Plugins/Experimental/MeshModelingToolsetExp/Source/SkeletalMeshModifiers/Public/SkeletonModifier.h b/Engine/Plugins/Experimental/MeshModelingToolsetExp/Source/SkeletalMeshModifiers/Public/SkeletonModifier.h
--- a/Engine/Plugins/Experimental/MeshModelingToolsetExp/Source/SkeletalMeshModifiers/Public/SkeletonModifier.h
+++ b/Engine/Plugins/Experimental/MeshModelingToolsetExp/Source/SkeletalMeshModifiers/Public/SkeletonModifier.h
@@ -48,24 +48,36 @@ struct SKELETALMESHMODIFIERS_API FMirrorOptions
  * FOrientOptions
  */
 
+UENUM()
+enum class EOrientAxis : uint8
+{
+	None,
+	PositiveX	UMETA(DisplayName = "+X", ToolTip = "Orients axis in the positive X direction"),
+	PositiveY	UMETA(DisplayName = "+Y", ToolTip = "Orients axis in the positive Y direction"),
+	PositiveZ	UMETA(DisplayName = "+Z", ToolTip = "Orients axis in the positive Z direction"),
+	NegativeX	UMETA(DisplayName = "-X", ToolTip = "Orients axis in the negative X direction"),
+	NegativeY	UMETA(DisplayName = "-Y", ToolTip = "Orients axis in the negative Y direction"),
+	NegativeZ	UMETA(DisplayName = "-Z", ToolTip = "Orients axis in the negative Z direction"),
+};
+
 USTRUCT(BlueprintType)
 struct SKELETALMESHMODIFIERS_API FOrientOptions
 {
 	GENERATED_USTRUCT_BODY()
 
 	UPROPERTY(BlueprintReadWrite, EditAnywhere, Category = "Orient")
-	TEnumAsByte<EAxis::Type> Primary = EAxis::X;
+	EOrientAxis Primary = EOrientAxis::PositiveX;
 
-	UPROPERTY(BlueprintReadWrite, EditAnywhere, Category = "Orient")
-	TEnumAsByte<EAxis::Type> Secondary = EAxis::Y;
+	UPROPERTY(BlueprintReadWrite, EditAnywhere, Category = "Orient", meta=(EditCondition="Primary != EOrientAxis::None"))
+	EOrientAxis Secondary = EOrientAxis::PositiveY;
 
-	UPROPERTY(BlueprintReadWrite, EditAnywhere, Category = "Orient")
+	UPROPERTY(BlueprintReadWrite, EditAnywhere, Category = "Orient", meta=(EditCondition="Primary != EOrientAxis::None"))
 	bool bUsePlaneAsSecondary = true;
 	
-	UPROPERTY(BlueprintReadWrite, EditAnywhere, Category = "Orient")
+	UPROPERTY(BlueprintReadWrite, EditAnywhere, Category = "Orient", meta=(EditCondition="Primary != EOrientAxis::None && !bUsePlaneAsSecondary"))
 	FVector SecondaryTarget = FVector::YAxisVector;
 
-	UPROPERTY(BlueprintReadWrite, EditAnywhere, Category = "Orient")
+	UPROPERTY(BlueprintReadWrite, EditAnywhere, Category = "Orient", meta=(EditCondition="Primary != EOrientAxis::None"))
 	bool bOrientChildren = true;
 
 	FTransform OrientTransform(const FVector& InPrimaryTarget, const FTransform& InTransform) const;
