Make Simple Collision Editor and Mesh To Collision tools call Modify() on the dynamic mesh component and corresponding bodysetup before updating the dynamic mesh collision, so that the undo transaction is properly registered.

diff --git a/Engine/Plugins/Experimental/MeshModelingToolsetExp/Source/MeshModelingToolsExp/Private/Physics/SetCollisionGeometryTool.cpp b/Engine/Plugins/Experimental/MeshModelingToolsetExp/Source/MeshModelingToolsExp/Private/Physics/SetCollisionGeometryTool.cpp
--- a/Engine/Plugins/Experimental/MeshModelingToolsetExp/Source/MeshModelingToolsExp/Private/Physics/SetCollisionGeometryTool.cpp
+++ b/Engine/Plugins/Experimental/MeshModelingToolsetExp/Source/MeshModelingToolsExp/Private/Physics/SetCollisionGeometryTool.cpp
@@ -549,6 +549,11 @@ void USetCollisionGeometryTool::OnShutdown(EToolShutdownType ShutdownType)
 		}
 		else if (UDynamicMeshComponent* DynamicMeshComponent = Cast<UDynamicMeshComponent>(Component))
 		{
+			DynamicMeshComponent->Modify();
+			if (UBodySetup* BodySetup = DynamicMeshComponent->GetBodySetup())
+			{
+				BodySetup->Modify();
+			}
 			DynamicMeshComponent->CollisionType = (ECollisionTraceFlag)(int32)Settings->SetCollisionType;
 			DynamicMeshComponent->SetSimpleCollisionShapes(GeneratedCollision->AggGeom, true);		
 			DynamicMeshComponent->MarkRenderStateDirty();
diff --git a/Engine/Plugins/Experimental/MeshModelingToolsetExp/Source/MeshModelingToolsExp/Private/Physics/SimpleCollisionEditorTool.cpp b/Engine/Plugins/Experimental/MeshModelingToolsetExp/Source/MeshModelingToolsExp/Private/Physics/SimpleCollisionEditorTool.cpp
--- a/Engine/Plugins/Experimental/MeshModelingToolsetExp/Source/MeshModelingToolsExp/Private/Physics/SimpleCollisionEditorTool.cpp
+++ b/Engine/Plugins/Experimental/MeshModelingToolsetExp/Source/MeshModelingToolsExp/Private/Physics/SimpleCollisionEditorTool.cpp
@@ -163,6 +163,11 @@ void USimpleCollisionEditorTool::OnShutdown(EToolShutdownType ShutdownType)
 		}
 		else if (UDynamicMeshComponent* DynamicMeshComponent = Cast<UDynamicMeshComponent>(Component))
 		{
+			DynamicMeshComponent->Modify();
+			if (UBodySetup* BodySetup = DynamicMeshComponent->GetBodySetup())
+			{
+				BodySetup->Modify();
+			}
 			DynamicMeshComponent->SetSimpleCollisionShapes(PhysicsInfos->AggGeom, true);
 			DynamicMeshComponent->MarkRenderStateDirty();
 		}
