fix cancel drag being turned into accept on the MeshSurfacePointTool-derived classes

diff --git a/Engine/Plugins/Experimental/HairModelingToolset/Source/HairModelingToolset/Private/GroomCardsEditorTool.cpp b/Engine/Plugins/Experimental/HairModelingToolset/Source/HairModelingToolset/Private/GroomCardsEditorTool.cpp
--- a/Engine/Plugins/Experimental/HairModelingToolset/Source/HairModelingToolset/Private/GroomCardsEditorTool.cpp
+++ b/Engine/Plugins/Experimental/HairModelingToolset/Source/HairModelingToolset/Private/GroomCardsEditorTool.cpp
@@ -342,6 +342,11 @@ void UGroomCardsEditorTool::OnEndDrag(const FRay& Ray)
 	check(false);
 }
 
+void UGroomCardsEditorTool::OnCancelDrag()
+{
+	check(false);
+}
+
 
 bool UGroomCardsEditorTool::OnUpdateHover(const FInputDeviceRay& DevicePos)
 {
diff --git a/Engine/Plugins/Experimental/HairModelingToolset/Source/HairModelingToolset/Public/GroomCardsEditorTool.h b/Engine/Plugins/Experimental/HairModelingToolset/Source/HairModelingToolset/Public/GroomCardsEditorTool.h
--- a/Engine/Plugins/Experimental/HairModelingToolset/Source/HairModelingToolset/Public/GroomCardsEditorTool.h
+++ b/Engine/Plugins/Experimental/HairModelingToolset/Source/HairModelingToolset/Public/GroomCardsEditorTool.h
@@ -168,6 +168,7 @@ public:
 	virtual void OnBeginDrag(const FRay& Ray) override;
 	virtual void OnUpdateDrag(const FRay& Ray) override;
 	virtual void OnEndDrag(const FRay& Ray) override;
+	virtual void OnCancelDrag() override;
 	virtual bool OnUpdateHover(const FInputDeviceRay& DevicePos) override;
 	virtual void OnEndHover() override;
 
diff --git a/Engine/Plugins/Experimental/MeshModelingToolsetExp/Source/MeshModelingToolsExp/Private/DeformMeshPolygonsTool.cpp b/Engine/Plugins/Experimental/MeshModelingToolsetExp/Source/MeshModelingToolsExp/Private/DeformMeshPolygonsTool.cpp
--- a/Engine/Plugins/Experimental/MeshModelingToolsetExp/Source/MeshModelingToolsExp/Private/DeformMeshPolygonsTool.cpp
+++ b/Engine/Plugins/Experimental/MeshModelingToolsetExp/Source/MeshModelingToolsExp/Private/DeformMeshPolygonsTool.cpp
@@ -1292,6 +1292,20 @@ void UDeformMeshPolygonsTool::OnEndDrag(const FRay& Ray)
 	}
 }
 
+void UDeformMeshPolygonsTool::OnCancelDrag()
+{
+	bInDrag = false;
+	bUpdatePending = false;
+
+	HilightSelection.Clear();
+	TopoSelector.Invalidate(true, false);
+	QuickAxisRotator.Reset();
+	QuickAxisTranslater.Reset();
+
+	delete ActiveVertexChange;
+	ActiveVertexChange = nullptr;
+}
+
 
 bool UDeformMeshPolygonsTool::OnUpdateHover(const FInputDeviceRay& DevicePos)
 {
diff --git a/Engine/Plugins/Experimental/MeshModelingToolsetExp/Source/MeshModelingToolsExp/Private/DynamicMeshSculptTool.cpp b/Engine/Plugins/Experimental/MeshModelingToolsetExp/Source/MeshModelingToolsExp/Private/DynamicMeshSculptTool.cpp
--- a/Engine/Plugins/Experimental/MeshModelingToolsetExp/Source/MeshModelingToolsExp/Private/DynamicMeshSculptTool.cpp
+++ b/Engine/Plugins/Experimental/MeshModelingToolsetExp/Source/MeshModelingToolsExp/Private/DynamicMeshSculptTool.cpp
@@ -1441,6 +1441,18 @@ void UDynamicMeshSculptTool::OnEndDrag(const FRay& Ray)
 	}
 }
 
+void UDynamicMeshSculptTool::OnCancelDrag()
+{
+	bInDrag = false;
+	bStampPending = false;
+	bRemeshPending = false;
+
+	CancelChange();
+
+	// destroy active remesher
+	ActiveRemesher = nullptr;
+}
+
 
 FInputRayHit UDynamicMeshSculptTool::BeginHoverSequenceHitTest(const FInputDeviceRay& PressPos)
 {
@@ -2440,6 +2452,15 @@ void UDynamicMeshSculptTool::EndChange()
 	}
 }
 
+void UDynamicMeshSculptTool::CancelChange()
+{
+	delete ActiveVertexChange;
+	ActiveVertexChange = nullptr;
+	
+	delete ActiveMeshChange;
+	ActiveMeshChange = nullptr;
+}
+
 void UDynamicMeshSculptTool::SaveActiveROI()
 {
 	if (ActiveMeshChange != nullptr)
diff --git a/Engine/Plugins/Experimental/MeshModelingToolsetExp/Source/MeshModelingToolsExp/Private/EditUVIslandsTool.cpp b/Engine/Plugins/Experimental/MeshModelingToolsetExp/Source/MeshModelingToolsExp/Private/EditUVIslandsTool.cpp
--- a/Engine/Plugins/Experimental/MeshModelingToolsetExp/Source/MeshModelingToolsExp/Private/EditUVIslandsTool.cpp
+++ b/Engine/Plugins/Experimental/MeshModelingToolsetExp/Source/MeshModelingToolsExp/Private/EditUVIslandsTool.cpp
@@ -282,7 +282,10 @@ void UEditUVIslandsTool::OnEndDrag(const FRay& Ray)
 	check(false);
 }
 
-
+void UEditUVIslandsTool::OnCancelDrag()
+{
+	check(false);
+}
 
 void UEditUVIslandsTool::OnGizmoTransformBegin(UTransformProxy*)
 {
diff --git a/Engine/Plugins/Experimental/MeshModelingToolsetExp/Source/MeshModelingToolsExp/Private/MeshGroupPaintTool.cpp b/Engine/Plugins/Experimental/MeshModelingToolsetExp/Source/MeshModelingToolsExp/Private/MeshGroupPaintTool.cpp
--- a/Engine/Plugins/Experimental/MeshModelingToolsetExp/Source/MeshModelingToolsExp/Private/MeshGroupPaintTool.cpp
+++ b/Engine/Plugins/Experimental/MeshModelingToolsetExp/Source/MeshModelingToolsExp/Private/MeshGroupPaintTool.cpp
@@ -418,7 +418,12 @@ void UMeshGroupPaintTool::OnEndStroke()
 	EndChange();
 }
 
-
+void UMeshGroupPaintTool::OnCancelStroke()
+{
+	GetActiveBrushOp()->CancelStroke();
+	ActiveGroupEditBuilder.Reset();
+	bDrawGroupsDataValid = false;
+}
 
 
 void UMeshGroupPaintTool::UpdateROI(const FSculptBrushStamp& BrushStamp)
diff --git a/Engine/Plugins/Experimental/MeshModelingToolsetExp/Source/MeshModelingToolsExp/Private/MeshVertexPaintTool.cpp b/Engine/Plugins/Experimental/MeshModelingToolsetExp/Source/MeshModelingToolsExp/Private/MeshVertexPaintTool.cpp
--- a/Engine/Plugins/Experimental/MeshModelingToolsetExp/Source/MeshModelingToolsExp/Private/MeshVertexPaintTool.cpp
+++ b/Engine/Plugins/Experimental/MeshModelingToolsetExp/Source/MeshModelingToolsExp/Private/MeshVertexPaintTool.cpp
@@ -575,6 +575,11 @@ void UMeshVertexPaintTool::OnEndStroke()
 	EndChange();
 }
 
+void UMeshVertexPaintTool::OnCancelStroke()
+{
+	GetActiveBrushOp()->CancelStroke();
+	ActiveChangeBuilder.Reset();
+}
 
 
 
diff --git a/Engine/Plugins/Experimental/MeshModelingToolsetExp/Source/MeshModelingToolsExp/Private/MeshVertexSculptTool.cpp b/Engine/Plugins/Experimental/MeshModelingToolsetExp/Source/MeshModelingToolsExp/Private/MeshVertexSculptTool.cpp
--- a/Engine/Plugins/Experimental/MeshModelingToolsetExp/Source/MeshModelingToolsExp/Private/MeshVertexSculptTool.cpp
+++ b/Engine/Plugins/Experimental/MeshModelingToolsetExp/Source/MeshModelingToolsExp/Private/MeshVertexSculptTool.cpp
@@ -500,6 +500,15 @@ void UMeshVertexSculptTool::OnEndStroke()
 }
 
 
+void UMeshVertexSculptTool::OnCancelStroke()
+{
+	GetActiveBrushOp()->CancelStroke();
+
+	delete ActiveVertexChange;
+	ActiveVertexChange = nullptr;
+}
+
+
 
 void UMeshVertexSculptTool::UpdateROI(const FVector3d& BrushPos)
 {
diff --git a/Engine/Plugins/Experimental/MeshModelingToolsetExp/Source/MeshModelingToolsExp/Private/Sculpting/MeshSculptToolBase.cpp b/Engine/Plugins/Experimental/MeshModelingToolsetExp/Source/MeshModelingToolsExp/Private/Sculpting/MeshSculptToolBase.cpp
--- a/Engine/Plugins/Experimental/MeshModelingToolsetExp/Source/MeshModelingToolsExp/Private/Sculpting/MeshSculptToolBase.cpp
+++ b/Engine/Plugins/Experimental/MeshModelingToolsetExp/Source/MeshModelingToolsExp/Private/Sculpting/MeshSculptToolBase.cpp
@@ -574,10 +574,17 @@ void UMeshSculptToolBase::OnEndDrag(const FRay& Ray)
 	bIsStampPending = false;
 
 	OnEndStroke();
-
 }
 
+void UMeshSculptToolBase::OnCancelDrag()
+{
+	bInStroke = false;
+
+	// cancel any outstanding stamps
+	bIsStampPending = false;
 
+	OnCancelStroke();
+}
 
 FRay3d UMeshSculptToolBase::GetLocalRay(const FRay& WorldRay) const
 {
diff --git a/Engine/Plugins/Experimental/MeshModelingToolsetExp/Source/MeshModelingToolsExp/Public/DeformMeshPolygonsTool.h b/Engine/Plugins/Experimental/MeshModelingToolsetExp/Source/MeshModelingToolsExp/Public/DeformMeshPolygonsTool.h
--- a/Engine/Plugins/Experimental/MeshModelingToolsetExp/Source/MeshModelingToolsExp/Public/DeformMeshPolygonsTool.h
+++ b/Engine/Plugins/Experimental/MeshModelingToolsetExp/Source/MeshModelingToolsExp/Public/DeformMeshPolygonsTool.h
@@ -168,6 +168,7 @@ public:
 	virtual void OnBeginDrag(const FRay& Ray) override;
 	virtual void OnUpdateDrag(const FRay& Ray) override;
 	virtual void OnEndDrag(const FRay& Ray) override;
+	virtual void OnCancelDrag() override;
 	virtual bool OnUpdateHover(const FInputDeviceRay& DevicePos) override;
 
 public:
diff --git a/Engine/Plugins/Experimental/MeshModelingToolsetExp/Source/MeshModelingToolsExp/Public/DynamicMeshSculptTool.h b/Engine/Plugins/Experimental/MeshModelingToolsetExp/Source/MeshModelingToolsExp/Public/DynamicMeshSculptTool.h
--- a/Engine/Plugins/Experimental/MeshModelingToolsetExp/Source/MeshModelingToolsExp/Public/DynamicMeshSculptTool.h
+++ b/Engine/Plugins/Experimental/MeshModelingToolsetExp/Source/MeshModelingToolsExp/Public/DynamicMeshSculptTool.h
@@ -280,6 +280,7 @@ public:
 	virtual void OnBeginDrag(const FRay& Ray) override;
 	virtual void OnUpdateDrag(const FRay& Ray) override;
 	virtual void OnEndDrag(const FRay& Ray) override;
+	virtual void OnCancelDrag() override;
 
 	virtual FInputRayHit BeginHoverSequenceHitTest(const FInputDeviceRay& PressPos) override;
 	virtual bool OnUpdateHover(const FInputDeviceRay& DevicePos) override;
@@ -498,6 +499,7 @@ private:
 	UE::Geometry::FDynamicMeshChangeTracker* ActiveMeshChange = nullptr;
 	void BeginChange(bool bIsVertexChange);
 	void EndChange();
+	void CancelChange();
 	void SaveActiveROI();
 
 	double EstimateIntialSafeTargetLength(const FDynamicMesh3& Mesh, int MinTargetTriCount);
diff --git a/Engine/Plugins/Experimental/MeshModelingToolsetExp/Source/MeshModelingToolsExp/Public/EditUVIslandsTool.h b/Engine/Plugins/Experimental/MeshModelingToolsetExp/Source/MeshModelingToolsExp/Public/EditUVIslandsTool.h
--- a/Engine/Plugins/Experimental/MeshModelingToolsetExp/Source/MeshModelingToolsExp/Public/EditUVIslandsTool.h
+++ b/Engine/Plugins/Experimental/MeshModelingToolsetExp/Source/MeshModelingToolsExp/Public/EditUVIslandsTool.h
@@ -82,6 +82,7 @@ public:
 	virtual void OnBeginDrag(const FRay& Ray) override;
 	virtual void OnUpdateDrag(const FRay& Ray) override;
 	virtual void OnEndDrag(const FRay& Ray) override;
+	virtual void OnCancelDrag() override;
 	virtual bool OnUpdateHover(const FInputDeviceRay& DevicePos) override;
 	virtual void OnEndHover() override;
 
diff --git a/Engine/Plugins/Experimental/MeshModelingToolsetExp/Source/MeshModelingToolsExp/Public/MeshGroupPaintTool.h b/Engine/Plugins/Experimental/MeshModelingToolsetExp/Source/MeshModelingToolsExp/Public/MeshGroupPaintTool.h
--- a/Engine/Plugins/Experimental/MeshModelingToolsetExp/Source/MeshModelingToolsExp/Public/MeshGroupPaintTool.h
+++ b/Engine/Plugins/Experimental/MeshModelingToolsetExp/Source/MeshModelingToolsExp/Public/MeshGroupPaintTool.h
@@ -365,6 +365,7 @@ protected:
 
 	virtual void OnBeginStroke(const FRay& WorldRay) override;
 	virtual void OnEndStroke() override;
+	virtual void OnCancelStroke() override;
 
 	virtual TUniquePtr<FMeshSculptBrushOp>& GetActiveBrushOp();
 	// end UMeshSculptToolBase API
diff --git a/Engine/Plugins/Experimental/MeshModelingToolsetExp/Source/MeshModelingToolsExp/Public/MeshVertexPaintTool.h b/Engine/Plugins/Experimental/MeshModelingToolsetExp/Source/MeshModelingToolsExp/Public/MeshVertexPaintTool.h
--- a/Engine/Plugins/Experimental/MeshModelingToolsetExp/Source/MeshModelingToolsExp/Public/MeshVertexPaintTool.h
+++ b/Engine/Plugins/Experimental/MeshModelingToolsetExp/Source/MeshModelingToolsExp/Public/MeshVertexPaintTool.h
@@ -529,6 +529,7 @@ protected:
 
 	virtual void OnBeginStroke(const FRay& WorldRay) override;
 	virtual void OnEndStroke() override;
+	virtual void OnCancelStroke() override;
 
 	virtual TUniquePtr<FMeshSculptBrushOp>& GetActiveBrushOp();
 	// end UMeshSculptToolBase API
diff --git a/Engine/Plugins/Experimental/MeshModelingToolsetExp/Source/MeshModelingToolsExp/Public/MeshVertexSculptTool.h b/Engine/Plugins/Experimental/MeshModelingToolsetExp/Source/MeshModelingToolsExp/Public/MeshVertexSculptTool.h
--- a/Engine/Plugins/Experimental/MeshModelingToolsetExp/Source/MeshModelingToolsExp/Public/MeshVertexSculptTool.h
+++ b/Engine/Plugins/Experimental/MeshModelingToolsetExp/Source/MeshModelingToolsExp/Public/MeshVertexSculptTool.h
@@ -267,6 +267,7 @@ protected:
 
 	virtual void OnBeginStroke(const FRay& WorldRay) override;
 	virtual void OnEndStroke() override;
+	virtual void OnCancelStroke() override;
 	// end UMeshSculptToolBase API
 
 protected:
diff --git a/Engine/Plugins/Experimental/MeshModelingToolsetExp/Source/MeshModelingToolsExp/Public/Sculpting/MeshBrushOpBase.h b/Engine/Plugins/Experimental/MeshModelingToolsetExp/Source/MeshModelingToolsExp/Public/Sculpting/MeshBrushOpBase.h
--- a/Engine/Plugins/Experimental/MeshModelingToolsetExp/Source/MeshModelingToolsExp/Public/Sculpting/MeshBrushOpBase.h
+++ b/Engine/Plugins/Experimental/MeshModelingToolsetExp/Source/MeshModelingToolsExp/Public/Sculpting/MeshBrushOpBase.h
@@ -124,6 +124,7 @@ public:
 
 	virtual void BeginStroke(const FDynamicMesh3* Mesh, const FSculptBrushStamp& Stamp, const TArray<int32>& InitialVertices) {}
 	virtual void EndStroke(const FDynamicMesh3* Mesh, const FSculptBrushStamp& Stamp, const TArray<int32>& FinalVertices) {}
+	virtual void CancelStroke() {}
 	virtual void ApplyStamp(const FDynamicMesh3* Mesh, const FSculptBrushStamp& Stamp, const TArray<int32>& Vertices, TArray<FVector3d>& NewValuesOut) = 0;
 
 
diff --git a/Engine/Plugins/Experimental/MeshModelingToolsetExp/Source/MeshModelingToolsExp/Public/Sculpting/MeshSculptToolBase.h b/Engine/Plugins/Experimental/MeshModelingToolsetExp/Source/MeshModelingToolsExp/Public/Sculpting/MeshSculptToolBase.h
--- a/Engine/Plugins/Experimental/MeshModelingToolsetExp/Source/MeshModelingToolsExp/Public/Sculpting/MeshSculptToolBase.h
+++ b/Engine/Plugins/Experimental/MeshModelingToolsetExp/Source/MeshModelingToolsExp/Public/Sculpting/MeshSculptToolBase.h
@@ -250,6 +250,7 @@ public:
 	virtual void OnBeginDrag(const FRay& Ray) override;
 	virtual void OnUpdateDrag(const FRay& Ray) override;
 	virtual void OnEndDrag(const FRay& Ray) override;
+	virtual void OnCancelDrag() override;
 	// end UMeshSurfacePointTool API
 
 protected:
@@ -259,6 +260,7 @@ protected:
 	virtual void OnCompleteSetup();
 	virtual void OnBeginStroke(const FRay& WorldRay) { check(false); }
 	virtual void OnEndStroke() { check(false); }
+	virtual void OnCancelStroke() { check(false); }
 
 public:
 	/** Properties that control brush size/etc */
diff --git a/Engine/Plugins/Experimental/SampleToolsEditorMode/Source/Private/SampleTools/DrawCurveOnMeshSampleTool.cpp b/Engine/Plugins/Experimental/SampleToolsEditorMode/Source/Private/SampleTools/DrawCurveOnMeshSampleTool.cpp
--- a/Engine/Plugins/Experimental/SampleToolsEditorMode/Source/Private/SampleTools/DrawCurveOnMeshSampleTool.cpp
+++ b/Engine/Plugins/Experimental/SampleToolsEditorMode/Source/Private/SampleTools/DrawCurveOnMeshSampleTool.cpp
@@ -101,7 +101,10 @@ void UDrawCurveOnMeshSampleTool::OnEndDrag(const FRay& Ray)
 
 }
 
+void UDrawCurveOnMeshSampleTool::OnCancelDrag()
+{
 
+}
 
 
 
diff --git a/Engine/Plugins/Experimental/SampleToolsEditorMode/Source/Private/SampleTools/DrawCurveOnMeshSampleTool.h b/Engine/Plugins/Experimental/SampleToolsEditorMode/Source/Private/SampleTools/DrawCurveOnMeshSampleTool.h
--- a/Engine/Plugins/Experimental/SampleToolsEditorMode/Source/Private/SampleTools/DrawCurveOnMeshSampleTool.h
+++ b/Engine/Plugins/Experimental/SampleToolsEditorMode/Source/Private/SampleTools/DrawCurveOnMeshSampleTool.h
@@ -74,6 +74,7 @@ public:
 	virtual void OnBeginDrag(const FRay& Ray) override;
 	virtual void OnUpdateDrag(const FRay& Ray) override;
 	virtual void OnEndDrag(const FRay& Ray) override;
+	virtual void OnCancelDrag() override;
 
 protected:
 	UPROPERTY()
diff --git a/Engine/Source/Runtime/InteractiveToolsFramework/Private/BaseTools/BaseBrushTool.cpp b/Engine/Source/Runtime/InteractiveToolsFramework/Private/BaseTools/BaseBrushTool.cpp
--- a/Engine/Source/Runtime/InteractiveToolsFramework/Private/BaseTools/BaseBrushTool.cpp
+++ b/Engine/Source/Runtime/InteractiveToolsFramework/Private/BaseTools/BaseBrushTool.cpp
@@ -441,6 +441,11 @@ void UBaseBrushTool::OnEndDrag(const FRay& Ray)
 	bInBrushStroke = false;
 }
 
+void UBaseBrushTool::OnCancelDrag()
+{
+	bInBrushStroke = false;
+}
+
 bool UBaseBrushTool::OnUpdateHover(const FInputDeviceRay& DevicePos)
 {
 	if (BrushAdjusterBehavior.IsValid() && BrushAdjusterBehavior->IsBrushBeingAdjusted())
diff --git a/Engine/Source/Runtime/InteractiveToolsFramework/Private/BaseTools/MeshSurfacePointTool.cpp b/Engine/Source/Runtime/InteractiveToolsFramework/Private/BaseTools/MeshSurfacePointTool.cpp
--- a/Engine/Source/Runtime/InteractiveToolsFramework/Private/BaseTools/MeshSurfacePointTool.cpp
+++ b/Engine/Source/Runtime/InteractiveToolsFramework/Private/BaseTools/MeshSurfacePointTool.cpp
@@ -184,7 +184,7 @@ void UMeshSurfacePointTool::OnClickRelease(const FInputDeviceRay& ReleasePos)
 
 void UMeshSurfacePointTool::OnTerminateDragSequence()
 {
-	OnEndDrag(LastWorldRay);
+	OnCancelDrag();
 }
 
 
diff --git a/Engine/Source/Runtime/InteractiveToolsFramework/Public/BaseTools/BaseBrushTool.h b/Engine/Source/Runtime/InteractiveToolsFramework/Public/BaseTools/BaseBrushTool.h
--- a/Engine/Source/Runtime/InteractiveToolsFramework/Public/BaseTools/BaseBrushTool.h
+++ b/Engine/Source/Runtime/InteractiveToolsFramework/Public/BaseTools/BaseBrushTool.h
@@ -172,6 +172,7 @@ public:
 	INTERACTIVETOOLSFRAMEWORK_API virtual void OnBeginDrag(const FRay& Ray) override;
 	INTERACTIVETOOLSFRAMEWORK_API virtual void OnUpdateDrag(const FRay& Ray) override;
 	INTERACTIVETOOLSFRAMEWORK_API virtual void OnEndDrag(const FRay& Ray) override;
+	INTERACTIVETOOLSFRAMEWORK_API virtual void OnCancelDrag() override;
 	INTERACTIVETOOLSFRAMEWORK_API virtual bool OnUpdateHover(const FInputDeviceRay& DevicePos) override;
 
 public:
diff --git a/Engine/Source/Runtime/InteractiveToolsFramework/Public/BaseTools/MeshSurfacePointTool.h b/Engine/Source/Runtime/InteractiveToolsFramework/Public/BaseTools/MeshSurfacePointTool.h
--- a/Engine/Source/Runtime/InteractiveToolsFramework/Public/BaseTools/MeshSurfacePointTool.h
+++ b/Engine/Source/Runtime/InteractiveToolsFramework/Public/BaseTools/MeshSurfacePointTool.h
@@ -95,8 +95,11 @@ public:
 	 */
 	INTERACTIVETOOLSFRAMEWORK_API virtual void OnEndDrag(const FRay& Ray);
 
-
-
+	/**
+	 * This function is called when the user's drag is cancelled, for example due to the whole tool being shut down.
+	 */
+	virtual void OnCancelDrag() {}
+	
 
 	// IClickDragBehaviorTarget implementation
 	INTERACTIVETOOLSFRAMEWORK_API virtual FInputRayHit CanBeginClickDragSequence(const FInputDeviceRay& PressPos) override;
