fix tool builder type for DrawPolyPathTool (was using the mesh surface point tool builder, but it is not a mesh surface point tool)

diff --git a/Engine/Plugins/Experimental/MeshModelingToolsetExp/Source/MeshModelingToolsExp/Public/DrawPolyPathTool.h b/Engine/Plugins/Experimental/MeshModelingToolsetExp/Source/MeshModelingToolsExp/Public/DrawPolyPathTool.h
--- a/Engine/Plugins/Experimental/MeshModelingToolsetExp/Source/MeshModelingToolsExp/Public/DrawPolyPathTool.h
+++ b/Engine/Plugins/Experimental/MeshModelingToolsetExp/Source/MeshModelingToolsExp/Public/DrawPolyPathTool.h
@@ -4,7 +4,7 @@
 
 #include "CoreMinimal.h"
 #include "UObject/NoExportTypes.h"
-#include "BaseTools/MeshSurfacePointTool.h"
+#include "InteractiveToolBuilder.h"
 #include "Mechanics/PlaneDistanceFromHitMechanic.h"
 #include "Mechanics/SpatialCurveDistanceMechanic.h"
 #include "Mechanics/CollectSurfacePathMechanic.h"
@@ -21,7 +21,7 @@ class UTransformProxy;
  * ToolBuilder
  */
 UCLASS()
-class MESHMODELINGTOOLSEXP_API UDrawPolyPathToolBuilder : public UMeshSurfacePointToolBuilder
+class MESHMODELINGTOOLSEXP_API UDrawPolyPathToolBuilder : public UInteractiveToolBuilder
 {
 	GENERATED_BODY()
 public:
