add an additional tooltip comment to revolve spline tool

diff --git a/Engine/Plugins/Experimental/MeshModelingToolsetExp/Source/MeshModelingToolsExp/Public/RevolveSplineTool.h b/Engine/Plugins/Experimental/MeshModelingToolsetExp/Source/MeshModelingToolsExp/Public/RevolveSplineTool.h
--- a/Engine/Plugins/Experimental/MeshModelingToolsetExp/Source/MeshModelingToolsExp/Public/RevolveSplineTool.h
+++ b/Engine/Plugins/Experimental/MeshModelingToolsetExp/Source/MeshModelingToolsExp/Public/RevolveSplineTool.h
@@ -29,6 +29,7 @@ class UWorld;
 UENUM()
 enum class ERevolveSplineSampleMode : uint8
 {
+	// Place points only at the spline control points
 	ControlPointsOnly,
 
 	// Place points along the spline such that the resulting polyline has no more than 
