Fix crash on undo during dynamic mesh sculpt brush by creating a dummy transaction while the brush stroke is in progress

diff --git a/Engine/Plugins/Experimental/MeshModelingToolsetExp/Source/MeshModelingToolsExp/Private/DynamicMeshSculptTool.cpp b/Engine/Plugins/Experimental/MeshModelingToolsetExp/Source/MeshModelingToolsExp/Private/DynamicMeshSculptTool.cpp
--- a/Engine/Plugins/Experimental/MeshModelingToolsetExp/Source/MeshModelingToolsExp/Private/DynamicMeshSculptTool.cpp
+++ b/Engine/Plugins/Experimental/MeshModelingToolsetExp/Source/MeshModelingToolsExp/Private/DynamicMeshSculptTool.cpp
@@ -293,6 +293,8 @@ void UDynamicMeshSculptTool::Setup()
 
 void UDynamicMeshSculptTool::Shutdown(EToolShutdownType ShutdownType)
 {
+	ClosePlaceholderTransaction(); // End any placeholder transaction, which could arise from brush stroke actions
+
 	if (ShutdownType == EToolShutdownType::Accept && AreAllTargetsValid() == false)
 	{
 		UE_LOG(LogTemp, Error, TEXT("Tool Target has become Invalid (possibly it has been Force Deleted). Aborting Tool."));
@@ -2419,6 +2421,7 @@ void UDynamicMeshSculptTool::BeginChange(bool bIsVertexChange)
 {
 	check(ActiveVertexChange == nullptr);
 	check(ActiveMeshChange == nullptr);
+	OpenPlaceholderTransaction(LOCTEXT("MeshSculptChange", "Brush Stroke")); // Create a placeholder change to prevent undo during the operation
 	if (bIsVertexChange)
 	{
 		ActiveVertexChange = new FMeshVertexChangeBuilder();
@@ -2432,6 +2435,12 @@ void UDynamicMeshSculptTool::BeginChange(bool bIsVertexChange)
 
 void UDynamicMeshSculptTool::EndChange()
 {
+	// End any placeholder transaction so we can emit a real transaction
+	// Note we close the placeholder transaction before emitting the FChange-based one due to the comment on 
+	// ApplyChange in ToolContextInterfaces.h indicating that we should not have an open transaction when we emit a change.
+	// TODO: We could revisit whether that comment is accurate, and if not, could move this Close call to the bottom of this function.
+	// (This does not affect the behavior in practice, since the empty placeholder transaction does not persist.)
+	ClosePlaceholderTransaction();
 	if (ActiveVertexChange != nullptr)
 	{
 		GetToolManager()->EmitObjectChange(DynamicMeshComponent, MoveTemp(ActiveVertexChange->Change), LOCTEXT("MeshSculptChange", "Brush Stroke"));
@@ -2454,6 +2463,7 @@ void UDynamicMeshSculptTool::EndChange()
 
 void UDynamicMeshSculptTool::CancelChange()
 {
+	ClosePlaceholderTransaction();
 	delete ActiveVertexChange;
 	ActiveVertexChange = nullptr;
 	
@@ -2461,6 +2471,23 @@ void UDynamicMeshSculptTool::CancelChange()
 	ActiveMeshChange = nullptr;
 }
 
+void UDynamicMeshSculptTool::OpenPlaceholderTransaction(FText Name)
+{
+	if (!bHasActivePlaceholderTransaction)
+	{
+		GetToolManager()->BeginUndoTransaction(Name);
+		bHasActivePlaceholderTransaction = true;
+	}
+}
+void UDynamicMeshSculptTool::ClosePlaceholderTransaction()
+{
+	if (bHasActivePlaceholderTransaction)
+	{
+		GetToolManager()->EndUndoTransaction();
+		bHasActivePlaceholderTransaction = false;
+	}
+}
+
 void UDynamicMeshSculptTool::SaveActiveROI()
 {
 	if (ActiveMeshChange != nullptr)
diff --git a/Engine/Plugins/Experimental/MeshModelingToolsetExp/Source/MeshModelingToolsExp/Public/DynamicMeshSculptTool.h b/Engine/Plugins/Experimental/MeshModelingToolsetExp/Source/MeshModelingToolsExp/Public/DynamicMeshSculptTool.h
--- a/Engine/Plugins/Experimental/MeshModelingToolsetExp/Source/MeshModelingToolsExp/Public/DynamicMeshSculptTool.h
+++ b/Engine/Plugins/Experimental/MeshModelingToolsetExp/Source/MeshModelingToolsExp/Public/DynamicMeshSculptTool.h
@@ -497,6 +497,9 @@ private:
 
 	FMeshVertexChangeBuilder* ActiveVertexChange = nullptr;
 	UE::Geometry::FDynamicMeshChangeTracker* ActiveMeshChange = nullptr;
+	bool bHasActivePlaceholderTransaction = false;
+	void OpenPlaceholderTransaction(FText Name);
+	void ClosePlaceholderTransaction();
 	void BeginChange(bool bIsVertexChange);
 	void EndChange();
 	void CancelChange();
