Added more comments/tooltips in the Mesh to Collision tool properties.
Also remove Simplify Polygon property (which had no effect) and prevent the Simplify Hulls toggle from affecting whether the Swept Hull polygon is simplified

diff --git a/Engine/Plugins/Experimental/MeshModelingToolsetExp/Source/MeshModelingToolsExp/Private/Physics/SetCollisionGeometryTool.cpp b/Engine/Plugins/Experimental/MeshModelingToolsetExp/Source/MeshModelingToolsExp/Private/Physics/SetCollisionGeometryTool.cpp
--- a/Engine/Plugins/Experimental/MeshModelingToolsetExp/Source/MeshModelingToolsExp/Private/Physics/SetCollisionGeometryTool.cpp
+++ b/Engine/Plugins/Experimental/MeshModelingToolsetExp/Source/MeshModelingToolsExp/Private/Physics/SetCollisionGeometryTool.cpp
@@ -311,7 +311,6 @@ void USetCollisionGeometryTool::Setup()
 	Settings->WatchProperty(Settings->ConvexDecompositionSearchFactor, [this](int32) { InvalidateCompute(); });
 	Settings->WatchProperty(Settings->AddHullsErrorTolerance, [this](int32) { InvalidateCompute(); });
 	Settings->WatchProperty(Settings->MinPartThickness, [this](int32) { InvalidateCompute(); });
-	Settings->WatchProperty(Settings->bSimplifyPolygons, [this](bool) { InvalidateCompute(); });
 	Settings->WatchProperty(Settings->HullTolerance, [this](float) { InvalidateCompute(); });
 	Settings->WatchProperty(Settings->SweepAxis, [this](EProjectedHullAxis) { InvalidateCompute(); });
 	Settings->WatchProperty(Settings->LevelSetResolution, [this](int32) { InvalidateCompute(); });
@@ -416,7 +415,8 @@ TUniquePtr<UE::Geometry::TGenericDataOperator<FPhysicsDataCollection>> USetColli
 	Op->UseShapeGenerator->bDetectBoxes = Settings->bDetectBoxes;
 	Op->UseShapeGenerator->bDetectCapsules = Settings->bDetectCapsules;
 	Op->UseShapeGenerator->MinDimension = Settings->MinThickness;
-	Op->UseShapeGenerator->bSimplifyHulls = Settings->bSimplifyHulls;
+	// SimplifyHulls on the shape generator controls simplification on both swept and convex hull paths, but for Swept Hulls UI we leave simplification always enabled
+	Op->UseShapeGenerator->bSimplifyHulls = Settings->GeometryType == ECollisionGeometryType::SweptHulls || Settings->bSimplifyHulls;
 	Op->UseShapeGenerator->HullTargetFaceCount = Settings->HullTargetFaceCount;
 	Op->UseShapeGenerator->ConvexDecompositionMaxPieces = Settings->MaxHullsPerMesh;
 	Op->UseShapeGenerator->ConvexDecompositionSearchFactor = Settings->ConvexDecompositionSearchFactor;
diff --git a/Engine/Plugins/Experimental/MeshModelingToolsetExp/Source/MeshModelingToolsExp/Public/Physics/SetCollisionGeometryTool.h b/Engine/Plugins/Experimental/MeshModelingToolsetExp/Source/MeshModelingToolsExp/Public/Physics/SetCollisionGeometryTool.h
--- a/Engine/Plugins/Experimental/MeshModelingToolsetExp/Source/MeshModelingToolsExp/Public/Physics/SetCollisionGeometryTool.h
+++ b/Engine/Plugins/Experimental/MeshModelingToolsetExp/Source/MeshModelingToolsExp/Public/Physics/SetCollisionGeometryTool.h
@@ -66,14 +66,21 @@ enum class ECollisionGeometryType
 	// Copy the existing collision geometry shapes from the inputs to the target. With a single-selection,
 	// always does the same thing as Empty with Append To Existing set to true.
 	CopyFromInputs = 0,
-
+	// Fit axis-aligned bounding boxes to the inputs
 	AlignedBoxes = 1,
+	// Fit oriented bounding boxes to the inputs
 	OrientedBoxes = 2,
+	// Fit spheres to the inputs
 	MinimalSpheres = 3,
+	// Fit capsules to the inputs
 	Capsules = 4,
+	// Fit convex hulls to the inputs
 	ConvexHulls = 5,
+	// Fit convex hulls to 2D projections of the inputs, and sweep these 2D hulls along the projection dimension
 	SweptHulls = 6,
+	// Fit level sets to the inputs
 	LevelSets = 7,
+	// Fit the boxes, spheres, and capsules to the inputs, and keep the best fitting of these shapes based on volume
 	MinVolume = 10,
 
 	// Do not produce new collision for inputs. If Append To Existing is false, this gives a way
@@ -87,10 +94,15 @@ enum class ECollisionGeometryType
 UENUM()
 enum class EProjectedHullAxis
 {
+	// Project along the X axis
 	X = 0,
+	// Project along the Y axis
 	Y = 1,
+	// Project along the Z axis
 	Z = 2,
+	// Project along the bounding box's shortest axis
 	SmallestBoxDimension = 3,
+	// Project along each major axis, and take the result with the smallest volume
 	SmallestVolume = 4
 };
 
@@ -101,27 +113,32 @@ class MESHMODELINGTOOLSEXP_API USetCollisionGeometryToolProperties : public UInt
 	GENERATED_BODY()
 public:
 
+	/** What kind of shapes to fit to the input. Note: Will be overridden by any enabled 'Auto Detect' settings, if close-fitting 'Auto Detect' shapes are found. */
 	UPROPERTY(EditAnywhere, Category = Options)
 	ECollisionGeometryType GeometryType = ECollisionGeometryType::AlignedBoxes;
 
-	// If true/false, Accept will append to/overwrite any existing collision geometry
+	/** Whether to keep the existing collision shapes, and append new shapes to that set. Otherwise, existing collision will be cleared. */
 	UPROPERTY(EditAnywhere, Category = Options)
 	bool bAppendToExisting = false;
 
+	/** When using multiple inputs to generate our collision, whether to use the world-space position of those input. If false, inputs will be considered as if they were all centered at the same location. */
 	UPROPERTY(EditAnywhere, Category = Options, meta = (EditCondition = "bUsingMultipleInputs", EditConditionHides, HideEditConditionToggle))
 	bool bUseWorldSpace = false;
 
+	/** What parts of the input should be separately fit with collision shapes */
 	UPROPERTY(EditAnywhere, Category = Options, meta = (
 		EditCondition = "GeometryType != ECollisionGeometryType::Empty && GeometryType != ECollisionGeometryType::CopyFromInputs"))
 	ESetCollisionGeometryInputMode InputMode = ESetCollisionGeometryInputMode::PerInputObject;
 
+	/** Whether to attempt to detect and remove collision shapes that are fully contained inside other collision shapes */
 	UPROPERTY(EditAnywhere, Category = Options)
 	bool bRemoveContained = true;
 
-	// If true, discard all but MaxCount collision geometries with the largest volume
+	/** Whether to discard all but MaxCount collision geometries with the largest volume */
 	UPROPERTY(EditAnywhere, Category = Options)
 	bool bEnableMaxCount = true;
 
+	/** The maximum number of collision shapes to generate. If necessary, the shapes with smallest volume will be discarded to meet this count. */
 	UPROPERTY(EditAnywhere, Category = Options, meta = (UIMin = "1", UIMax = "100", ClampMin = "1", ClampMax = "9999999", EditCondition = "bEnableMaxCount"))
 	int32 MaxCount = 50;
 
@@ -129,20 +146,25 @@ public:
 	UPROPERTY(EditAnywhere, Category = Options, AdvancedDisplay, meta = (ClampMin = "0", UIMax = "10", EditCondition = "GeometryType != ECollisionGeometryType::LevelSets && (GeometryType != ECollisionGeometryType::ConvexHulls || MaxHullsPerMesh == 1)"))
 	float MinThickness = 0.01;
 
-	UPROPERTY(EditAnywhere, Category = AutoDetect)
+	/** Whether to override the requested Geometry Type with a box whenever a box closely fits the input shape */
+	UPROPERTY(EditAnywhere, Category = AutoDetectGeometryOverrides)
 	bool bDetectBoxes = true;
 
-	UPROPERTY(EditAnywhere, Category = AutoDetect)
+	/** Whether to override the requested Geometry Type with a sphere whenever a sphere closely fits the input shape */
+	UPROPERTY(EditAnywhere, Category = AutoDetectGeometryOverrides)
 	bool bDetectSpheres = true;
 
-	UPROPERTY(EditAnywhere, Category = AutoDetect)
+	/** Whether to override the requested Geometry Type with a capsule whenever a capsule closely fits the input shape */
+	UPROPERTY(EditAnywhere, Category = AutoDetectGeometryOverrides)
 	bool bDetectCapsules = true;
 
+	/** Whether to simplify the convex hull */
 	UPROPERTY(EditAnywhere, Category = ConvexHulls, meta = (EditConditionHides, EditCondition = "GeometryType == ECollisionGeometryType::ConvexHulls"))
 	bool bSimplifyHulls = true;
 
+	/** Target number of faces in the simplified hull */
 	UPROPERTY(EditAnywhere, Category = ConvexHulls, meta = (UIMin = "4", UIMax = "100", ClampMin = "4", ClampMax = "9999999",
-		EditConditionHides, EditCondition = "GeometryType == ECollisionGeometryType::ConvexHulls"))
+		EditConditionHides, EditCondition = "GeometryType == ECollisionGeometryType::ConvexHulls && bSimplifyHulls"))
 	int32 HullTargetFaceCount = 20;
 
 	/** How many convex hulls can be used to approximate each mesh */
@@ -160,29 +182,27 @@ public:
 		EditConditionHides, EditCondition = "GeometryType == ECollisionGeometryType::ConvexHulls && MaxHullsPerMesh > 1"))
 	float AddHullsErrorTolerance = 0;
 
-	/** Minimum part thickness for convex decomposition (in cm); hulls thinner than this will be merged into adjacent hulls, if possible. */
+	/** Minimum part thickness for convex decomposition, in cm; hulls thinner than this will be merged into adjacent hulls, if possible. */
 	UPROPERTY(EditAnywhere, Category = ConvexHulls, meta = (UIMin = "0", UIMax = "1", ClampMin = "0",
 		EditConditionHides, EditCondition = "GeometryType == ECollisionGeometryType::ConvexHulls && MaxHullsPerMesh > 1"))
 	float MinPartThickness = 0.1;
 
-	UPROPERTY(EditAnywhere, Category = SweptHulls, meta = (EditConditionHides, EditCondition = "GeometryType == ECollisionGeometryType::SweptHulls"))
-	bool bSimplifyPolygons = true;
-
+	/** If > 0, the polygon used to generate the swept hull will be simplified up to this distance tolerance, in cm */
 	UPROPERTY(EditAnywhere, Category = SweptHulls, meta = (UIMin = "0", UIMax = "10", ClampMin = "0", ClampMax = "100000",
 		EditConditionHides, EditCondition = "GeometryType == ECollisionGeometryType::SweptHulls"))
 	float HullTolerance = 0.1;
 
+	/** How to choose which direction to sweep when creating a swept hull */
 	UPROPERTY(EditAnywhere, Category = SweptHulls, meta = (UIMin = "0", UIMax = "10", ClampMin = "0", ClampMax = "100000",
 		EditConditionHides, EditCondition = "GeometryType == ECollisionGeometryType::SweptHulls"))
 	EProjectedHullAxis SweepAxis = EProjectedHullAxis::SmallestVolume;
 
-	// Level Set settings
-
 	/** Level set grid resolution along longest grid axis */
 	UPROPERTY(EditAnywhere, Category = LevelSets, meta = (UIMin = "3", UIMax = "100", ClampMin = "3", ClampMax = "1000",
 		EditConditionHides, EditCondition = "GeometryType == ECollisionGeometryType::LevelSets"))
 	int32 LevelSetResolution = 10;
 
+	/** Set how the physics system should interpret collision shapes on the output mesh. Does not affect what collision shapes are generated by this tool. */
 	UPROPERTY(EditAnywhere, Category = OutputOptions)
 	ECollisionGeometryMode SetCollisionType = ECollisionGeometryMode::SimpleAndComplex;
 
