make ctrl+click to position custom plane work more consistently in draw spline tool

diff --git a/Engine/Plugins/Experimental/MeshModelingToolsetExp/Source/MeshModelingToolsEditorOnlyExp/Private/DrawSplineTool.cpp b/Engine/Plugins/Experimental/MeshModelingToolsetExp/Source/MeshModelingToolsEditorOnlyExp/Private/DrawSplineTool.cpp
--- a/Engine/Plugins/Experimental/MeshModelingToolsetExp/Source/MeshModelingToolsEditorOnlyExp/Private/DrawSplineTool.cpp
+++ b/Engine/Plugins/Experimental/MeshModelingToolsetExp/Source/MeshModelingToolsEditorOnlyExp/Private/DrawSplineTool.cpp
@@ -295,6 +295,9 @@ void UDrawSplineTool::Setup()
 	ClickOrDragBehavior->Initialize(this, this);
 	AddInputBehavior(ClickOrDragBehavior);
 
+	// Make sure the plane mechanic captures clicks first, to ensure it sees ctrl+clicks to reposition the plane
+	PlaneMechanic->UpdateClickPriority(ClickOrDragBehavior->GetPriority().MakeHigher());
+
 	Settings->WatchProperty(Settings->bLoop, [this](bool) {
 		if (ensure(WorkingSpline.IsValid()))
 		{
