add more tooltip comments to the draw spline tool

diff --git a/Engine/Plugins/Experimental/MeshModelingToolsetExp/Source/MeshModelingToolsEditorOnlyExp/Public/DrawSplineTool.h b/Engine/Plugins/Experimental/MeshModelingToolsetExp/Source/MeshModelingToolsEditorOnlyExp/Public/DrawSplineTool.h
--- a/Engine/Plugins/Experimental/MeshModelingToolsetExp/Source/MeshModelingToolsEditorOnlyExp/Public/DrawSplineTool.h
+++ b/Engine/Plugins/Experimental/MeshModelingToolsetExp/Source/MeshModelingToolsEditorOnlyExp/Public/DrawSplineTool.h
@@ -34,7 +34,9 @@ enum class EDrawSplineDrawMode : uint8
 UENUM()
 enum class ESplineOffsetMethod : uint8
 {
+	// Spline points will be offset along the normal direction of the clicked surface
 	HitNormal,
+	// Spline points will be offset along a manually-chosen direction
 	Custom
 };
 
@@ -125,9 +127,11 @@ public:
 	UPROPERTY(EditAnywhere, Category = Spline, meta = (UIMin = 0, UIMax = 100))
 	double ClickOffset = 0;
 
+	/** How to choose the direction to offset points from the clicked surface */
 	UPROPERTY(EditAnywhere, Category = Spline, meta = (EditCondition = "ClickOffset > 0", EditConditionHides))
 	ESplineOffsetMethod OffsetMethod = ESplineOffsetMethod::HitNormal;
 
+	/** Manually-specified click offset direction. Note: Will be normalized. If it is a zero vector, a default Up vector will be used instead. */
 	UPROPERTY(EditAnywhere, Category = Spline, meta = (EditCondition = "ClickOffset > 0 && OffsetMethod == ESplineOffsetMethod::Custom", EditConditionHides))
 	FVector OffsetDirection = FVector::UpVector;
 
@@ -154,12 +158,15 @@ public:
 		EditCondition = "OutputMode != EDrawSplineOutputMode::EmptyActor"))
 	bool bPreviewUsingActorCopy = true;
 
+	/** Whether to place spline points on the surface of objects in the world */
 	UPROPERTY(EditAnywhere, Category = RaycastTargets, meta = (DisplayName = "World"))
 	bool bHitWorld = true;
 
+	/** Whether to place spline points on a custom, user-adjustable plane */
 	UPROPERTY(EditAnywhere, Category = RaycastTargets, meta = (DisplayName = "Custom Plane"))
 	bool bHitCustomPlane = false;
 
+	/** Whether to place spline points on a plane through the origin aligned with the Z axis in perspective views, or facing the camera in othographic views */
 	UPROPERTY(EditAnywhere, Category = RaycastTargets, meta = (DisplayName = "Ground Planes"))
 	bool bHitGroundPlanes = true;
 
