PatternTool: Fix undo/redo bugs due to ResetTransformGizmoPosition setting the transform gizmo's position in a way that creates a transaction.

We really only want manual changes to the gizmo to be recorded as transactions (click+drag changes). Calls to ResetTransformGizmoPosition generally mean something has changed in such a way that we must reset the gizmo's transform to keep it consistent with how it already is (i.e. we enabled 'centered', which means to keep the same pattern extent we need to move the gizmo), which shouldn't be viewed as a change from an undo/redo perspective.

This fixes an ensure related to undoing changes which result in ResetTransformGizmoPosition being called, and also fixes some bugs I discovered while trouble shooting that are not crashes but result in the gizmo behaving very oddly.

diff --git a/Engine/Plugins/Experimental/MeshModelingToolsetExp/Source/MeshModelingToolsExp/Private/PatternTool.cpp b/Engine/Plugins/Experimental/MeshModelingToolsetExp/Source/MeshModelingToolsExp/Private/PatternTool.cpp
--- a/Engine/Plugins/Experimental/MeshModelingToolsetExp/Source/MeshModelingToolsExp/Private/PatternTool.cpp
+++ b/Engine/Plugins/Experimental/MeshModelingToolsetExp/Source/MeshModelingToolsExp/Private/PatternTool.cpp
@@ -1190,7 +1190,7 @@ void UPatternTool::ResetTransformGizmoPosition()
 
 	const FVector3d ProxyTranslation = CurrentStartFrameWorld.Origin + FRotator(CurrentStartFrameWorld.Rotation).RotateVector(OffsetFromOrigin);
 
-	PatternGizmo->SetNewGizmoTransform(FTransform(FQuat(CurrentStartFrameWorld.Rotation * FQuaterniond(RotationInLocalSpace)), ProxyTranslation));
+	PatternGizmo->ReinitializeGizmoTransform(FTransform(FQuat(CurrentStartFrameWorld.Rotation * FQuaterniond(RotationInLocalSpace)), ProxyTranslation));
 }
 
 void UPatternTool::ReconstructTransformGizmos()
