Update comments to make it clear that fchange can be called inside other transactions, and move the dynamic sculpt close-transaction call to after the fchanges so we don't need to rely on the placeholder transaction being empty and deleting itself

diff --git a/Engine/Plugins/Experimental/MeshModelingToolsetExp/Source/MeshModelingToolsExp/Private/DynamicMeshSculptTool.cpp b/Engine/Plugins/Experimental/MeshModelingToolsetExp/Source/MeshModelingToolsExp/Private/DynamicMeshSculptTool.cpp
--- a/Engine/Plugins/Experimental/MeshModelingToolsetExp/Source/MeshModelingToolsExp/Private/DynamicMeshSculptTool.cpp
+++ b/Engine/Plugins/Experimental/MeshModelingToolsetExp/Source/MeshModelingToolsExp/Private/DynamicMeshSculptTool.cpp
@@ -2435,12 +2435,6 @@ void UDynamicMeshSculptTool::BeginChange(bool bIsVertexChange)
 
 void UDynamicMeshSculptTool::EndChange()
 {
-	// End any placeholder transaction so we can emit a real transaction
-	// Note we close the placeholder transaction before emitting the FChange-based one due to the comment on 
-	// ApplyChange in ToolContextInterfaces.h indicating that we should not have an open transaction when we emit a change.
-	// TODO: We could revisit whether that comment is accurate, and if not, could move this Close call to the bottom of this function.
-	// (This does not affect the behavior in practice, since the empty placeholder transaction does not persist.)
-	ClosePlaceholderTransaction();
 	if (ActiveVertexChange != nullptr)
 	{
 		GetToolManager()->EmitObjectChange(DynamicMeshComponent, MoveTemp(ActiveVertexChange->Change), LOCTEXT("MeshSculptChange", "Brush Stroke"));
@@ -2459,6 +2453,8 @@ void UDynamicMeshSculptTool::EndChange()
 		delete ActiveMeshChange;
 		ActiveMeshChange = nullptr;
 	}
+
+	ClosePlaceholderTransaction();
 }
 
 void UDynamicMeshSculptTool::CancelChange()
diff --git a/Engine/Source/Runtime/InteractiveToolsFramework/Public/ToolContextInterfaces.h b/Engine/Source/Runtime/InteractiveToolsFramework/Public/ToolContextInterfaces.h
--- a/Engine/Source/Runtime/InteractiveToolsFramework/Public/ToolContextInterfaces.h
+++ b/Engine/Source/Runtime/InteractiveToolsFramework/Public/ToolContextInterfaces.h
@@ -317,8 +317,8 @@ public:
 
 	/**
 	 * Insert an FChange into the transaction history in the current Context. 
-	 * This cannot be called between Begin/EndUndoTransaction, the FChange should be 
-	 * automatically inserted into a Transaction.
+	 * It is safe but not necessary to call this between Begin/EndUndoTransaction;
+	 * the FChange will be automatically inserted into a Transaction.
 	 * @param TargetObject The UObject this Change is applied to
 	 * @param Change The Change implementation
 	 * @param Description text description of the transaction that could be shown to user
