Add geometry selection api overload to tools that won't affect selection, so the modeling mode can restore any previous geometry selection after use

diff --git a/Engine/Plugins/Experimental/MeshModelingToolsetExp/Source/MeshModelingToolsEditorOnlyExp/Public/AttributeEditorTool.h b/Engine/Plugins/Experimental/MeshModelingToolsetExp/Source/MeshModelingToolsEditorOnlyExp/Public/AttributeEditorTool.h
--- a/Engine/Plugins/Experimental/MeshModelingToolsetExp/Source/MeshModelingToolsEditorOnlyExp/Public/AttributeEditorTool.h
+++ b/Engine/Plugins/Experimental/MeshModelingToolsetExp/Source/MeshModelingToolsEditorOnlyExp/Public/AttributeEditorTool.h
@@ -328,7 +328,7 @@ public:
  * Mesh Attribute Editor Tool
  */
 UCLASS()
-class MESHMODELINGTOOLSEDITORONLYEXP_API UAttributeEditorTool : public UMultiSelectionMeshEditingTool
+class MESHMODELINGTOOLSEDITORONLYEXP_API UAttributeEditorTool : public UMultiSelectionMeshEditingTool, public IInteractiveToolManageGeometrySelectionAPI
 {
 	GENERATED_BODY()
 
@@ -346,6 +346,12 @@ public:
 
 	virtual void RequestAction(EAttributeEditorToolActions ActionType);
 
+	// IInteractiveToolManageGeometrySelectionAPI -- this tool won't update external geometry selection or change selection-relevant mesh IDs
+	virtual bool IsInputSelectionValidOnOutput() override
+	{
+		return true;
+	}
+
 protected:
 
 	UPROPERTY()
diff --git a/Engine/Plugins/Experimental/MeshModelingToolsetExp/Source/MeshModelingToolsExp/Public/BakeMeshAttributeTool.h b/Engine/Plugins/Experimental/MeshModelingToolsetExp/Source/MeshModelingToolsExp/Public/BakeMeshAttributeTool.h
--- a/Engine/Plugins/Experimental/MeshModelingToolsetExp/Source/MeshModelingToolsExp/Public/BakeMeshAttributeTool.h
+++ b/Engine/Plugins/Experimental/MeshModelingToolsetExp/Source/MeshModelingToolsExp/Public/BakeMeshAttributeTool.h
@@ -80,7 +80,7 @@ static constexpr EBakeMapType ENUM_EBAKEMAPTYPE_ALL[] =
  * Base Mesh Bake tool
  */
 UCLASS()
-class MESHMODELINGTOOLSEXP_API UBakeMeshAttributeTool : public UMultiSelectionMeshEditingTool, public IInteractiveToolExclusiveToolAPI
+class MESHMODELINGTOOLSEXP_API UBakeMeshAttributeTool : public UMultiSelectionMeshEditingTool, public IInteractiveToolExclusiveToolAPI, public IInteractiveToolManageGeometrySelectionAPI
 {
 	GENERATED_BODY()
 
@@ -95,6 +95,12 @@ public:
 	virtual bool CanAccept() const override { return true; }
 	// End UInteractiveTool interface
 
+	// IInteractiveToolManageGeometrySelectionAPI -- this tool won't update external geometry selection or change selection-relevant mesh IDs
+	virtual bool IsInputSelectionValidOnOutput() override
+	{
+		return true;
+	}
+
 protected:
 	//
 	// Bake tool property sets
diff --git a/Engine/Plugins/Experimental/MeshModelingToolsetExp/Source/MeshModelingToolsExp/Public/BakeTransformTool.h b/Engine/Plugins/Experimental/MeshModelingToolsetExp/Source/MeshModelingToolsExp/Public/BakeTransformTool.h
--- a/Engine/Plugins/Experimental/MeshModelingToolsetExp/Source/MeshModelingToolsExp/Public/BakeTransformTool.h
+++ b/Engine/Plugins/Experimental/MeshModelingToolsetExp/Source/MeshModelingToolsExp/Public/BakeTransformTool.h
@@ -80,7 +80,7 @@ public:
  * Simple tool to bake scene transform on meshes into the mesh assets
  */
 UCLASS()
-class MESHMODELINGTOOLSEXP_API UBakeTransformTool : public UMultiSelectionMeshEditingTool
+class MESHMODELINGTOOLSEXP_API UBakeTransformTool : public UMultiSelectionMeshEditingTool, public IInteractiveToolManageGeometrySelectionAPI
 {
 	GENERATED_BODY()
 
@@ -93,6 +93,12 @@ public:
 	virtual bool HasCancel() const override { return true; }
 	virtual bool HasAccept() const override { return true; }
 
+	// IInteractiveToolManageGeometrySelectionAPI -- this tool won't update external geometry selection or change selection-relevant mesh IDs
+	virtual bool IsInputSelectionValidOnOutput() override
+	{
+		return true;
+	}
+
 protected:
 
 	UPROPERTY()
diff --git a/Engine/Plugins/Experimental/MeshModelingToolsetExp/Source/MeshModelingToolsExp/Public/DeformMeshPolygonsTool.h b/Engine/Plugins/Experimental/MeshModelingToolsetExp/Source/MeshModelingToolsExp/Public/DeformMeshPolygonsTool.h
--- a/Engine/Plugins/Experimental/MeshModelingToolsetExp/Source/MeshModelingToolsExp/Public/DeformMeshPolygonsTool.h
+++ b/Engine/Plugins/Experimental/MeshModelingToolsetExp/Source/MeshModelingToolsExp/Public/DeformMeshPolygonsTool.h
@@ -145,7 +145,7 @@ public:
  *
  */
 UCLASS()
-class MESHMODELINGTOOLSEXP_API UDeformMeshPolygonsTool : public UMeshSurfacePointTool
+class MESHMODELINGTOOLSEXP_API UDeformMeshPolygonsTool : public UMeshSurfacePointTool, public IInteractiveToolManageGeometrySelectionAPI
 {
 	GENERATED_BODY()
 
@@ -171,6 +171,12 @@ public:
 	virtual void OnCancelDrag() override;
 	virtual bool OnUpdateHover(const FInputDeviceRay& DevicePos) override;
 
+	// IInteractiveToolManageGeometrySelectionAPI -- this tool won't update external geometry selection or change selection-relevant mesh IDs
+	virtual bool IsInputSelectionValidOnOutput() override
+	{
+		return true;
+	}
+
 public:
 	virtual void NextTransformTypeAction();
 
diff --git a/Engine/Plugins/Experimental/MeshModelingToolsetExp/Source/MeshModelingToolsExp/Public/EditPivotTool.h b/Engine/Plugins/Experimental/MeshModelingToolsetExp/Source/MeshModelingToolsExp/Public/EditPivotTool.h
--- a/Engine/Plugins/Experimental/MeshModelingToolsetExp/Source/MeshModelingToolsExp/Public/EditPivotTool.h
+++ b/Engine/Plugins/Experimental/MeshModelingToolsetExp/Source/MeshModelingToolsExp/Public/EditPivotTool.h
@@ -151,7 +151,7 @@ public:
  *
  */
 UCLASS()
-class MESHMODELINGTOOLSEXP_API UEditPivotTool : public UMultiSelectionMeshEditingTool, public IClickDragBehaviorTarget
+class MESHMODELINGTOOLSEXP_API UEditPivotTool : public UMultiSelectionMeshEditingTool, public IClickDragBehaviorTarget, public IInteractiveToolManageGeometrySelectionAPI
 {
 	GENERATED_BODY()
 
@@ -179,6 +179,12 @@ public:
 	virtual void OnClickRelease(const FInputDeviceRay& ReleasePos) override;
 	virtual void OnTerminateDragSequence() override;
 
+	// IInteractiveToolManageGeometrySelectionAPI -- this tool won't update external geometry selection or change selection-relevant mesh IDs
+	virtual bool IsInputSelectionValidOnOutput() override
+	{
+		return true;
+	}
+
 public:
 	UPROPERTY()
 	TObjectPtr<UEditPivotToolProperties> TransformProps;
diff --git a/Engine/Plugins/Experimental/MeshModelingToolsetExp/Source/MeshModelingToolsExp/Public/EditUVIslandsTool.h b/Engine/Plugins/Experimental/MeshModelingToolsetExp/Source/MeshModelingToolsExp/Public/EditUVIslandsTool.h
--- a/Engine/Plugins/Experimental/MeshModelingToolsetExp/Source/MeshModelingToolsExp/Public/EditUVIslandsTool.h
+++ b/Engine/Plugins/Experimental/MeshModelingToolsetExp/Source/MeshModelingToolsExp/Public/EditUVIslandsTool.h
@@ -59,7 +59,7 @@ public:
  *
  */
 UCLASS()
-class MESHMODELINGTOOLSEXP_API UEditUVIslandsTool : public UMeshSurfacePointTool
+class MESHMODELINGTOOLSEXP_API UEditUVIslandsTool : public UMeshSurfacePointTool, public IInteractiveToolManageGeometrySelectionAPI
 {
 	GENERATED_BODY()
 
@@ -89,6 +89,11 @@ public:
 	// IClickDragBehaviorTarget API
 	virtual FInputRayHit CanBeginClickDragSequence(const FInputDeviceRay& PressPos) override;
 
+	// IInteractiveToolManageGeometrySelectionAPI -- this tool won't update external geometry selection or change selection-relevant mesh IDs
+	virtual bool IsInputSelectionValidOnOutput() override
+	{
+		return true;
+	}
 
 public:
 	UPROPERTY()
diff --git a/Engine/Plugins/Experimental/MeshModelingToolsetExp/Source/MeshModelingToolsExp/Public/LatticeDeformerTool.h b/Engine/Plugins/Experimental/MeshModelingToolsetExp/Source/MeshModelingToolsExp/Public/LatticeDeformerTool.h
--- a/Engine/Plugins/Experimental/MeshModelingToolsetExp/Source/MeshModelingToolsExp/Public/LatticeDeformerTool.h
+++ b/Engine/Plugins/Experimental/MeshModelingToolsetExp/Source/MeshModelingToolsExp/Public/LatticeDeformerTool.h
@@ -127,7 +127,7 @@ public:
 
 /** Deform a mesh using a regular hexahedral lattice */
 UCLASS()
-class MESHMODELINGTOOLSEXP_API ULatticeDeformerTool : public USingleSelectionMeshEditingTool
+class MESHMODELINGTOOLSEXP_API ULatticeDeformerTool : public USingleSelectionMeshEditingTool, public IInteractiveToolManageGeometrySelectionAPI
 {
 	GENERATED_BODY()
 
@@ -149,6 +149,12 @@ public:
 
 	UE::Geometry::FVector3i GetLatticeResolution() const;
 
+	// IInteractiveToolManageGeometrySelectionAPI -- this tool won't update external geometry selection or change selection-relevant mesh IDs
+	virtual bool IsInputSelectionValidOnOutput() override
+	{
+		return true;
+	}
+
 protected:
 
 	// Input mesh
diff --git a/Engine/Plugins/Experimental/MeshModelingToolsetExp/Source/MeshModelingToolsExp/Public/MeshAttributePaintTool.h b/Engine/Plugins/Experimental/MeshModelingToolsetExp/Source/MeshModelingToolsExp/Public/MeshAttributePaintTool.h
--- a/Engine/Plugins/Experimental/MeshModelingToolsetExp/Source/MeshModelingToolsExp/Public/MeshAttributePaintTool.h
+++ b/Engine/Plugins/Experimental/MeshModelingToolsetExp/Source/MeshModelingToolsExp/Public/MeshAttributePaintTool.h
@@ -201,7 +201,7 @@ public:
  * 
  */
 UCLASS()
-class MESHMODELINGTOOLSEXP_API UMeshAttributePaintTool : public UDynamicMeshBrushTool
+class MESHMODELINGTOOLSEXP_API UMeshAttributePaintTool : public UDynamicMeshBrushTool, public IInteractiveToolManageGeometrySelectionAPI
 {
 	GENERATED_BODY()
 
@@ -226,6 +226,12 @@ public:
 
 	virtual void SetColorMap(TUniquePtr<FFloatAttributeColorMapper> ColorMap);
 
+	// IInteractiveToolManageGeometrySelectionAPI -- this tool won't update external geometry selection or change selection-relevant mesh IDs
+	virtual bool IsInputSelectionValidOnOutput() override
+	{
+		return true;
+	}
+
 protected:
 	virtual void ApplyStamp(const FBrushStampData& Stamp);
 
diff --git a/Engine/Plugins/Experimental/MeshModelingToolsetExp/Source/MeshModelingToolsExp/Public/MeshSpaceDeformerTool.h b/Engine/Plugins/Experimental/MeshModelingToolsetExp/Source/MeshModelingToolsExp/Public/MeshSpaceDeformerTool.h
--- a/Engine/Plugins/Experimental/MeshModelingToolsetExp/Source/MeshModelingToolsExp/Public/MeshSpaceDeformerTool.h
+++ b/Engine/Plugins/Experimental/MeshModelingToolsetExp/Source/MeshModelingToolsExp/Public/MeshSpaceDeformerTool.h
@@ -196,7 +196,7 @@ public:
  * Applies non-linear deformations to a mesh 
  */
 UCLASS()
-class MESHMODELINGTOOLSEXP_API UMeshSpaceDeformerTool : public USingleSelectionMeshEditingTool
+class MESHMODELINGTOOLSEXP_API UMeshSpaceDeformerTool : public USingleSelectionMeshEditingTool, public IInteractiveToolManageGeometrySelectionAPI
 {
 	GENERATED_BODY()
 public:
@@ -219,6 +219,12 @@ public:
 	// sync the parameters owned by the MeshSpaceDeformerOp 
 	void UpdateOpParameters(UE::Geometry::FMeshSpaceDeformerOp& MeshSpaceDeformerOp) const;
 
+	// IInteractiveToolManageGeometrySelectionAPI -- this tool won't update external geometry selection or change selection-relevant mesh IDs
+	virtual bool IsInputSelectionValidOnOutput() override
+	{
+		return true;
+	}
+
 protected:
 
 	UPROPERTY()
diff --git a/Engine/Plugins/Experimental/MeshModelingToolsetExp/Source/MeshModelingToolsExp/Public/MeshVertexPaintTool.h b/Engine/Plugins/Experimental/MeshModelingToolsetExp/Source/MeshModelingToolsExp/Public/MeshVertexPaintTool.h
--- a/Engine/Plugins/Experimental/MeshModelingToolsetExp/Source/MeshModelingToolsExp/Public/MeshVertexPaintTool.h
+++ b/Engine/Plugins/Experimental/MeshModelingToolsetExp/Source/MeshModelingToolsExp/Public/MeshVertexPaintTool.h
@@ -458,7 +458,7 @@ public:
  * Mesh Vertex Color Painting TOol
  */
 UCLASS()
-class MESHMODELINGTOOLSEXP_API UMeshVertexPaintTool : public UMeshSculptToolBase
+class MESHMODELINGTOOLSEXP_API UMeshVertexPaintTool : public UMeshSculptToolBase, public IInteractiveToolManageGeometrySelectionAPI
 {
 	GENERATED_BODY()
 
@@ -483,6 +483,11 @@ public:
 
 	virtual void CommitResult(UBaseDynamicMeshComponent* Component, bool bModifiedTopology) override;
 
+	// IInteractiveToolManageGeometrySelectionAPI -- this tool won't update external geometry selection or change selection-relevant mesh IDs
+	virtual bool IsInputSelectionValidOnOutput() override
+	{
+		return true;
+	}
 
 public:
 
diff --git a/Engine/Plugins/Experimental/MeshModelingToolsetExp/Source/MeshModelingToolsExp/Public/MeshVertexSculptTool.h b/Engine/Plugins/Experimental/MeshModelingToolsetExp/Source/MeshModelingToolsExp/Public/MeshVertexSculptTool.h
--- a/Engine/Plugins/Experimental/MeshModelingToolsetExp/Source/MeshModelingToolsExp/Public/MeshVertexSculptTool.h
+++ b/Engine/Plugins/Experimental/MeshModelingToolsetExp/Source/MeshModelingToolsExp/Public/MeshVertexSculptTool.h
@@ -197,7 +197,7 @@ public:
  * Mesh Vertex Sculpt Tool Class
  */
 UCLASS()
-class MESHMODELINGTOOLSEXP_API UMeshVertexSculptTool : public UMeshSculptToolBase
+class MESHMODELINGTOOLSEXP_API UMeshVertexSculptTool : public UMeshSculptToolBase, public IInteractiveToolManageGeometrySelectionAPI
 {
 	GENERATED_BODY()
 public:
@@ -214,6 +214,12 @@ public:
 
 	virtual void OnPropertyModified(UObject* PropertySet, FProperty* Property) override;
 
+	// IInteractiveToolManageGeometrySelectionAPI -- this tool won't update external geometry selection or change selection-relevant mesh IDs
+	virtual bool IsInputSelectionValidOnOutput() override
+	{
+		return true;
+	}
+
 public:
 
 	/** Properties that control sculpting*/
diff --git a/Engine/Plugins/Experimental/MeshModelingToolsetExp/Source/MeshModelingToolsExp/Public/Physics/PhysicsInspectorTool.h b/Engine/Plugins/Experimental/MeshModelingToolsetExp/Source/MeshModelingToolsExp/Public/Physics/PhysicsInspectorTool.h
--- a/Engine/Plugins/Experimental/MeshModelingToolsetExp/Source/MeshModelingToolsExp/Public/Physics/PhysicsInspectorTool.h
+++ b/Engine/Plugins/Experimental/MeshModelingToolsetExp/Source/MeshModelingToolsExp/Public/Physics/PhysicsInspectorTool.h
@@ -29,7 +29,7 @@ protected:
  * Mesh Inspector Tool for visualizing mesh information
  */
 UCLASS()
-class MESHMODELINGTOOLSEXP_API UPhysicsInspectorTool : public UMultiSelectionMeshEditingTool
+class MESHMODELINGTOOLSEXP_API UPhysicsInspectorTool : public UMultiSelectionMeshEditingTool, public IInteractiveToolManageGeometrySelectionAPI
 {
 	GENERATED_BODY()
 public:
@@ -42,6 +42,12 @@ public:
 	virtual bool HasAccept() const override { return false; }
 	virtual bool CanAccept() const override { return false; }
 
+	// IInteractiveToolManageGeometrySelectionAPI -- this tool won't update external geometry selection or change selection-relevant mesh IDs
+	virtual bool IsInputSelectionValidOnOutput() override
+	{
+		return true;
+	}
+
 protected:
 
 	UPROPERTY()
diff --git a/Engine/Plugins/Experimental/MeshModelingToolsetExp/Source/MeshModelingToolsExp/Public/Physics/SetCollisionGeometryTool.h b/Engine/Plugins/Experimental/MeshModelingToolsetExp/Source/MeshModelingToolsExp/Public/Physics/SetCollisionGeometryTool.h
--- a/Engine/Plugins/Experimental/MeshModelingToolsetExp/Source/MeshModelingToolsExp/Public/Physics/SetCollisionGeometryTool.h
+++ b/Engine/Plugins/Experimental/MeshModelingToolsetExp/Source/MeshModelingToolsExp/Public/Physics/SetCollisionGeometryTool.h
@@ -219,7 +219,7 @@ public:
  * Mesh Inspector Tool for visualizing mesh information
  */
 UCLASS()
-class MESHMODELINGTOOLSEXP_API USetCollisionGeometryTool : public UMultiSelectionMeshEditingTool, public UE::Geometry::IGenericDataOperatorFactory<FPhysicsDataCollection>
+class MESHMODELINGTOOLSEXP_API USetCollisionGeometryTool : public UMultiSelectionMeshEditingTool, public UE::Geometry::IGenericDataOperatorFactory<FPhysicsDataCollection>, public IInteractiveToolManageGeometrySelectionAPI
 {
 	GENERATED_BODY()
 public:
@@ -242,6 +242,12 @@ public:
 
 	void SetGeometrySelection(UE::Geometry::FGeometrySelection&& SelectionIn);
 
+	// IInteractiveToolManageGeometrySelectionAPI -- this tool won't update external geometry selection or change selection-relevant mesh IDs
+	virtual bool IsInputSelectionValidOnOutput() override
+	{
+		return true;
+	}
+
 protected:
 
 	UPROPERTY()
diff --git a/Engine/Plugins/Experimental/MeshModelingToolsetExp/Source/MeshModelingToolsExp/Public/Physics/SimpleCollisionEditorTool.h b/Engine/Plugins/Experimental/MeshModelingToolsetExp/Source/MeshModelingToolsExp/Public/Physics/SimpleCollisionEditorTool.h
--- a/Engine/Plugins/Experimental/MeshModelingToolsetExp/Source/MeshModelingToolsExp/Public/Physics/SimpleCollisionEditorTool.h
+++ b/Engine/Plugins/Experimental/MeshModelingToolsetExp/Source/MeshModelingToolsExp/Public/Physics/SimpleCollisionEditorTool.h
@@ -94,7 +94,7 @@ public:
  * Simple Collision Editing tool for updating the simple collision geometry on meshes
  */
 UCLASS()
-class MESHMODELINGTOOLSEXP_API USimpleCollisionEditorTool : public USingleSelectionMeshEditingTool
+class MESHMODELINGTOOLSEXP_API USimpleCollisionEditorTool : public USingleSelectionMeshEditingTool, public IInteractiveToolManageGeometrySelectionAPI
 {
 	GENERATED_BODY()
 public:
@@ -110,6 +110,12 @@ public:
 	virtual bool HasAccept() const override { return true; }
 	virtual bool CanAccept() const override { return true; }
 
+	// IInteractiveToolManageGeometrySelectionAPI -- this tool won't update external geometry selection or change selection-relevant mesh IDs
+	virtual bool IsInputSelectionValidOnOutput() override
+	{
+		return true;
+	}
+
 protected:
 
 	UPROPERTY()
diff --git a/Engine/Plugins/Experimental/MeshModelingToolsetExp/Source/MeshModelingToolsExp/Public/SeamSculptTool.h b/Engine/Plugins/Experimental/MeshModelingToolsetExp/Source/MeshModelingToolsExp/Public/SeamSculptTool.h
--- a/Engine/Plugins/Experimental/MeshModelingToolsetExp/Source/MeshModelingToolsExp/Public/SeamSculptTool.h
+++ b/Engine/Plugins/Experimental/MeshModelingToolsetExp/Source/MeshModelingToolsExp/Public/SeamSculptTool.h
@@ -49,7 +49,7 @@ public:
 
 
 UCLASS(Transient)
-class MESHMODELINGTOOLSEXP_API USeamSculptTool : public UDynamicMeshBrushTool
+class MESHMODELINGTOOLSEXP_API USeamSculptTool : public UDynamicMeshBrushTool, public IInteractiveToolManageGeometrySelectionAPI
 {
 	GENERATED_BODY()
 public:
@@ -70,6 +70,12 @@ public:
 	virtual void OnEndDrag(const FRay& Ray) override;
 	virtual bool OnUpdateHover(const FInputDeviceRay& DevicePos) override;
 
+	// IInteractiveToolManageGeometrySelectionAPI -- this tool won't update external geometry selection or change selection-relevant mesh IDs
+	virtual bool IsInputSelectionValidOnOutput() override
+	{
+		return true;
+	}
+
 
 protected:
 	virtual void ApplyStamp(const FBrushStampData& Stamp);
diff --git a/Engine/Plugins/Experimental/MeshModelingToolsetExp/Source/MeshModelingToolsExp/Public/SmoothMeshTool.h b/Engine/Plugins/Experimental/MeshModelingToolsetExp/Source/MeshModelingToolsExp/Public/SmoothMeshTool.h
--- a/Engine/Plugins/Experimental/MeshModelingToolsetExp/Source/MeshModelingToolsExp/Public/SmoothMeshTool.h
+++ b/Engine/Plugins/Experimental/MeshModelingToolsetExp/Source/MeshModelingToolsExp/Public/SmoothMeshTool.h
@@ -125,7 +125,7 @@ public:
  * Mesh Smoothing Tool
  */
 UCLASS()
-class MESHMODELINGTOOLSEXP_API USmoothMeshTool : public UBaseMeshProcessingTool
+class MESHMODELINGTOOLSEXP_API USmoothMeshTool : public UBaseMeshProcessingTool, public IInteractiveToolManageGeometrySelectionAPI
 {
 	GENERATED_BODY()
 
@@ -143,6 +143,12 @@ public:
 	virtual FText GetToolMessageString() const override;
 	virtual FText GetAcceptTransactionName() const override;
 
+	// IInteractiveToolManageGeometrySelectionAPI -- this tool won't update external geometry selection or change selection-relevant mesh IDs
+	virtual bool IsInputSelectionValidOnOutput() override
+	{
+		return true;
+	}
+
 protected:
 	UPROPERTY()
 	TObjectPtr<USmoothMeshToolProperties> SmoothProperties = nullptr;
diff --git a/Engine/Plugins/Experimental/MeshModelingToolsetExp/Source/MeshModelingToolsExp/Public/TransformMeshesTool.h b/Engine/Plugins/Experimental/MeshModelingToolsetExp/Source/MeshModelingToolsExp/Public/TransformMeshesTool.h
--- a/Engine/Plugins/Experimental/MeshModelingToolsetExp/Source/MeshModelingToolsExp/Public/TransformMeshesTool.h
+++ b/Engine/Plugins/Experimental/MeshModelingToolsetExp/Source/MeshModelingToolsExp/Public/TransformMeshesTool.h
@@ -148,7 +148,7 @@ struct FTransformMeshesTarget
  *
  */
 UCLASS()
-class MESHMODELINGTOOLSEXP_API UTransformMeshesTool : public UMultiSelectionMeshEditingTool, public IClickDragBehaviorTarget
+class MESHMODELINGTOOLSEXP_API UTransformMeshesTool : public UMultiSelectionMeshEditingTool, public IClickDragBehaviorTarget, public IInteractiveToolManageGeometrySelectionAPI
 {
 	GENERATED_BODY()
 
@@ -174,6 +174,12 @@ public:
 	virtual void OnClickRelease(const FInputDeviceRay& ReleasePos) override;
 	virtual void OnTerminateDragSequence() override;
 
+	// IInteractiveToolManageGeometrySelectionAPI -- this tool won't update external geometry selection or change selection-relevant mesh IDs
+	virtual bool IsInputSelectionValidOnOutput() override
+	{
+		return true;
+	}
+
 protected:
 	UPROPERTY()
 	TObjectPtr<UTransformMeshesToolProperties> TransformProps;
diff --git a/Engine/Plugins/Runtime/MeshModelingToolset/Source/MeshModelingTools/Public/RecomputeUVsTool.h b/Engine/Plugins/Runtime/MeshModelingToolset/Source/MeshModelingTools/Public/RecomputeUVsTool.h
--- a/Engine/Plugins/Runtime/MeshModelingToolset/Source/MeshModelingTools/Public/RecomputeUVsTool.h
+++ b/Engine/Plugins/Runtime/MeshModelingToolset/Source/MeshModelingTools/Public/RecomputeUVsTool.h
@@ -42,7 +42,7 @@ public:
  * URecomputeUVsTool Recomputes UVs based on existing segmentations of the mesh
  */
 UCLASS()
-class MESHMODELINGTOOLS_API URecomputeUVsTool : public USingleSelectionMeshEditingTool
+class MESHMODELINGTOOLS_API URecomputeUVsTool : public USingleSelectionMeshEditingTool, public IInteractiveToolManageGeometrySelectionAPI
 {
 	GENERATED_BODY()
 
@@ -59,6 +59,12 @@ public:
 
 	virtual void OnPropertyModified(UObject* PropertySet, FProperty* Property) override;
 
+	// IInteractiveToolManageGeometrySelectionAPI -- this tool won't update external geometry selection or change selection-relevant mesh IDs
+	virtual bool IsInputSelectionValidOnOutput() override
+	{
+		return true;
+	}
+
 protected:
 	UPROPERTY()
 	TObjectPtr<UMeshUVChannelProperties> UVChannelProperties = nullptr;
diff --git a/Engine/Plugins/Runtime/MeshModelingToolset/Source/MeshModelingTools/Public/UVLayoutTool.h b/Engine/Plugins/Runtime/MeshModelingToolset/Source/MeshModelingTools/Public/UVLayoutTool.h
--- a/Engine/Plugins/Runtime/MeshModelingToolset/Source/MeshModelingTools/Public/UVLayoutTool.h
+++ b/Engine/Plugins/Runtime/MeshModelingToolset/Source/MeshModelingTools/Public/UVLayoutTool.h
@@ -42,7 +42,7 @@ public:
  * The level editor version of the UV layout tool.
  */
 UCLASS()
-class MESHMODELINGTOOLS_API UUVLayoutTool : public UMultiSelectionMeshEditingTool
+class MESHMODELINGTOOLS_API UUVLayoutTool : public UMultiSelectionMeshEditingTool, public IInteractiveToolManageGeometrySelectionAPI
 {
 	GENERATED_BODY()
 
@@ -64,6 +64,12 @@ public:
 
 	int32 GetSelectedUVChannel() const;
 
+	// IInteractiveToolManageGeometrySelectionAPI -- this tool won't update external geometry selection or change selection-relevant mesh IDs
+	virtual bool IsInputSelectionValidOnOutput() override
+	{
+		return true;
+	}
+
 protected:
 
 	UPROPERTY()
diff --git a/Engine/Plugins/Runtime/MeshModelingToolset/Source/MeshModelingToolsEditorOnly/Public/ParameterizeMeshTool.h b/Engine/Plugins/Runtime/MeshModelingToolset/Source/MeshModelingToolsEditorOnly/Public/ParameterizeMeshTool.h
--- a/Engine/Plugins/Runtime/MeshModelingToolset/Source/MeshModelingToolsEditorOnly/Public/ParameterizeMeshTool.h
+++ b/Engine/Plugins/Runtime/MeshModelingToolset/Source/MeshModelingToolsEditorOnly/Public/ParameterizeMeshTool.h
@@ -39,7 +39,7 @@ public:
  * and then packs the resulting charts
  */
 UCLASS()
-class MESHMODELINGTOOLSEDITORONLY_API UParameterizeMeshTool : public USingleSelectionMeshEditingTool
+class MESHMODELINGTOOLSEDITORONLY_API UParameterizeMeshTool : public USingleSelectionMeshEditingTool, public IInteractiveToolManageGeometrySelectionAPI
 {
 	GENERATED_BODY()
 
@@ -56,6 +56,12 @@ public:
 
 	virtual void OnPropertyModified(UObject* PropertySet, FProperty* Property) override;
 
+	// IInteractiveToolManageGeometrySelectionAPI -- this tool won't update external geometry selection or change selection-relevant mesh IDs
+	virtual bool IsInputSelectionValidOnOutput() override
+	{
+		return true;
+	}
+
 protected:
 	UPROPERTY()
 	TObjectPtr<UMeshUVChannelProperties> UVChannelProperties = nullptr;
