Remove reset to default button on "Preview Output Type" options in baking tools, this is actually a workaround rather than a fix to UE-192561. Also only show the baked channels as options in the BakeRC tool

diff --git a/Engine/Plugins/Experimental/MeshModelingToolsetExp/Source/MeshModelingToolsEditorOnlyExp/Private/BakeRenderCaptureTool.cpp b/Engine/Plugins/Experimental/MeshModelingToolsetExp/Source/MeshModelingToolsEditorOnlyExp/Private/BakeRenderCaptureTool.cpp
--- a/Engine/Plugins/Experimental/MeshModelingToolsetExp/Source/MeshModelingToolsEditorOnlyExp/Private/BakeRenderCaptureTool.cpp
+++ b/Engine/Plugins/Experimental/MeshModelingToolsetExp/Source/MeshModelingToolsEditorOnlyExp/Private/BakeRenderCaptureTool.cpp
@@ -314,7 +314,7 @@ void UBakeRenderCaptureTool::Setup()
 	AddToolPropertySource(Settings);
 
 	Settings->MapPreview = BaseColorTexParamName;
-	Settings->WatchProperty(Settings->MapPreview, [this](FString) { UpdateVisualization(); });
+	MapPreviewWatcherIndex = Settings->WatchProperty(Settings->MapPreview, [this](FString) { UpdateVisualization(); });
 	Settings->WatchProperty(Settings->SamplesPerPixel, [this](EBakeTextureSamplesPerPixel) { OpState |= EBakeOpState::Evaluate; });
 	Settings->WatchProperty(Settings->TextureSize, [this](EBakeTextureResolution) { OpState |= EBakeOpState::Evaluate; });
 	Settings->WatchProperty(Settings->ValidSampleDepthThreshold, [this](float ValidSampleDepthThreshold)
@@ -357,18 +357,6 @@ void UBakeRenderCaptureTool::Setup()
 	UpdateUVLayerNames(InputMeshSettings->TargetUVLayer, InputMeshSettings->TargetUVLayerNamesList, *TargetMesh);
 	InputMeshSettings->WatchProperty(InputMeshSettings->TargetUVLayer, [this](FString) { OpState |= EBakeOpState::Evaluate; });
 	
-	{
-		Settings->MapPreviewNamesList.Add(BaseColorTexParamName);
-		Settings->MapPreviewNamesList.Add(NormalTexParamName);
-		Settings->MapPreviewNamesList.Add(PackedMRSTexParamName);
-		Settings->MapPreviewNamesList.Add(MetallicTexParamName);
-		Settings->MapPreviewNamesList.Add(RoughnessTexParamName);
-		Settings->MapPreviewNamesList.Add(SpecularTexParamName);
-		Settings->MapPreviewNamesList.Add(EmissiveTexParamName);
-		Settings->MapPreviewNamesList.Add(OpacityTexParamName);
-		Settings->MapPreviewNamesList.Add(SubsurfaceColorTexParamName);
-	}
-
 	ResultSettings = NewObject<UBakeRenderCaptureResults>(this);
 	ResultSettings->RestoreProperties(this);
 	AddToolPropertySource(ResultSettings);
@@ -1056,6 +1044,68 @@ void UBakeRenderCaptureTool::UpdateResult()
 		return;
 	}
 
+	{
+		Settings->MapPreviewNamesList.Reset();
+
+		if (RenderCaptureProperties->bBaseColorMap)
+		{
+			Settings->MapPreviewNamesList.Add(BaseColorTexParamName);
+		}
+		if (RenderCaptureProperties->bNormalMap)
+		{
+			Settings->MapPreviewNamesList.Add(NormalTexParamName);
+		}
+		if (RenderCaptureProperties->bPackedMRSMap)
+		{
+			Settings->MapPreviewNamesList.Add(PackedMRSTexParamName);
+		}
+		if (RenderCaptureProperties->bMetallicMap)
+		{
+			Settings->MapPreviewNamesList.Add(MetallicTexParamName);
+		}
+		if (RenderCaptureProperties->bRoughnessMap)
+		{
+			Settings->MapPreviewNamesList.Add(RoughnessTexParamName);
+		}
+		if (RenderCaptureProperties->bSpecularMap)
+		{
+			Settings->MapPreviewNamesList.Add(SpecularTexParamName);
+		}
+		if (RenderCaptureProperties->bEmissiveMap)
+		{
+			Settings->MapPreviewNamesList.Add(EmissiveTexParamName);
+		}
+		if (RenderCaptureProperties->bOpacityMap)
+		{
+			Settings->MapPreviewNamesList.Add(OpacityTexParamName);
+		}
+		if (RenderCaptureProperties->bSubsurfaceColorMap)
+		{
+			Settings->MapPreviewNamesList.Add(SubsurfaceColorTexParamName);
+		}
+
+		if (Settings->MapPreviewNamesList.IsEmpty())
+		{
+			// Display an empty string when MapPreview is disabled
+			Settings->MapPreview = TEXT("");
+			Settings->SilentUpdateWatcherAtIndex(MapPreviewWatcherIndex);
+
+			Settings->bEnableMapPreview = false;
+		}
+		else
+		{
+			// If the current MapPreview channel is disabled, switch to the first enabled channel in the list
+			if (Settings->MapPreviewNamesList.Find(Settings->MapPreview) == INDEX_NONE)
+			{
+				Settings->MapPreview = Settings->MapPreviewNamesList[0];
+				Settings->SilentUpdateWatcherAtIndex(MapPreviewWatcherIndex);
+			}
+
+			Settings->bEnableMapPreview = true;
+		}
+		NotifyOfPropertyChangeByTool(Settings);
+	}
+
 	// The bake operation, Compute, stores a pointer to the SceneCapture so that must not be modified while baking
 	const bool bComputeInProgress = (Compute && (Compute->GetElapsedComputeTime() > 0.f));
 	if (bComputeInProgress)
diff --git a/Engine/Plugins/Experimental/MeshModelingToolsetExp/Source/MeshModelingToolsEditorOnlyExp/Public/BakeRenderCaptureTool.h b/Engine/Plugins/Experimental/MeshModelingToolsetExp/Source/MeshModelingToolsEditorOnlyExp/Public/BakeRenderCaptureTool.h
--- a/Engine/Plugins/Experimental/MeshModelingToolsetExp/Source/MeshModelingToolsEditorOnlyExp/Public/BakeRenderCaptureTool.h
+++ b/Engine/Plugins/Experimental/MeshModelingToolsetExp/Source/MeshModelingToolsEditorOnlyExp/Public/BakeRenderCaptureTool.h
@@ -180,16 +180,8 @@ class MESHMODELINGTOOLSEDITORONLYEXP_API UBakeRenderCaptureToolProperties : publ
 public:
 	
 	/** The map type to preview */
-	UPROPERTY(EditAnywhere, Category = BakeOutput, meta = (DisplayName="Preview Output Type", TransientToolProperty, GetOptions = GetMapPreviewNamesFunc))
-	FString MapPreview;
-	
-	UFUNCTION()
-	const TArray<FString>& GetMapPreviewNamesFunc()
-	{
-		return MapPreviewNamesList;
-	}
-	UPROPERTY(meta = (TransientToolProperty))
-	TArray<FString> MapPreviewNamesList;
+	UPROPERTY(EditAnywhere, Category = BakeOutput, meta = (DisplayName="Preview Output Type", EditCondition = "bEnableMapPreview", HideEditConditionToggle, TransientToolProperty, GetOptions = GetMapPreviewNamesFunc, NoResetToDefault))
+	FString MapPreview = TEXT("");
 
 	/** Number of samples per pixel */
 	UPROPERTY(EditAnywhere, Category = BakeOutput)
@@ -213,6 +205,20 @@ public:
 	 */
 	UPROPERTY(EditAnywhere, Category = BakeOutput, meta = (ClampMin="0", UIMin="0"), DisplayName="Cleanup Threshold")
 	float ValidSampleDepthThreshold = 0.f;
+
+	//~ Hidden properties
+
+	UPROPERTY(Transient, meta = (TransientToolProperty))
+	bool bEnableMapPreview = false;
+
+	UFUNCTION()
+	const TArray<FString>& GetMapPreviewNamesFunc()
+	{
+		return MapPreviewNamesList;
+	}
+
+	UPROPERTY(meta = (TransientToolProperty))
+	TArray<FString> MapPreviewNamesList;
 };
 
 
@@ -310,6 +316,7 @@ protected:
 
 	UPROPERTY()
 	TObjectPtr<UBakeRenderCaptureToolProperties> Settings;
+	int32 MapPreviewWatcherIndex = -1;
 
 	UPROPERTY()
 	TObjectPtr<URenderCaptureProperties> RenderCaptureProperties;
diff --git a/Engine/Plugins/Experimental/MeshModelingToolsetExp/Source/MeshModelingToolsExp/Public/BakeMeshAttributeMapsTool.h b/Engine/Plugins/Experimental/MeshModelingToolsetExp/Source/MeshModelingToolsExp/Public/BakeMeshAttributeMapsTool.h
--- a/Engine/Plugins/Experimental/MeshModelingToolsetExp/Source/MeshModelingToolsExp/Public/BakeMeshAttributeMapsTool.h
+++ b/Engine/Plugins/Experimental/MeshModelingToolsetExp/Source/MeshModelingToolsExp/Public/BakeMeshAttributeMapsTool.h
@@ -59,7 +59,7 @@ public:
 
 	/** The baked output type used for preview in the viewport */
 	UPROPERTY(EditAnywhere, Category = BakeOutput, meta = (DisplayName = "Preview Output Type", GetOptions = GetMapPreviewNamesFunc, TransientToolProperty,
-		EditCondition = "MapTypes != 0"))
+		EditCondition = "MapTypes != 0", NoResetToDefault))
 	FString MapPreview;
 
 	/** The pixel resolution of the generated textures */
diff --git a/Engine/Plugins/Experimental/MeshModelingToolsetExp/Source/MeshModelingToolsExp/Public/BakeMultiMeshAttributeMapsTool.h b/Engine/Plugins/Experimental/MeshModelingToolsetExp/Source/MeshModelingToolsExp/Public/BakeMultiMeshAttributeMapsTool.h
--- a/Engine/Plugins/Experimental/MeshModelingToolsetExp/Source/MeshModelingToolsExp/Public/BakeMultiMeshAttributeMapsTool.h
+++ b/Engine/Plugins/Experimental/MeshModelingToolsetExp/Source/MeshModelingToolsExp/Public/BakeMultiMeshAttributeMapsTool.h
@@ -46,7 +46,7 @@ public:
 
 	/** The map type index to preview */
 	UPROPERTY(EditAnywhere, Category = BakeOutput, meta=(DisplayName="Preview Output Type", TransientToolProperty, GetOptions = GetMapPreviewNamesFunc,
-		EditCondition = "MapTypes != 0"))
+		EditCondition = "MapTypes != 0", NoResetToDefault))
 	FString MapPreview;
 
 	/** The pixel resolution of the generated map */
