Make the Sculpt Tool Base derived tools and the Dynamic Sculpt Tool use the LongTransactions utility to prevent undo while creating brush/sculpt changes

diff --git a/Engine/Plugins/Experimental/ChaosClothAssetEditor/Source/ChaosClothAssetEditorTools/Private/ChaosClothAsset/ClothWeightMapPaintTool.cpp b/Engine/Plugins/Experimental/ChaosClothAssetEditor/Source/ChaosClothAssetEditorTools/Private/ChaosClothAsset/ClothWeightMapPaintTool.cpp
--- a/Engine/Plugins/Experimental/ChaosClothAssetEditor/Source/ChaosClothAssetEditorTools/Private/ChaosClothAsset/ClothWeightMapPaintTool.cpp
+++ b/Engine/Plugins/Experimental/ChaosClothAssetEditor/Source/ChaosClothAssetEditorTools/Private/ChaosClothAsset/ClothWeightMapPaintTool.cpp
@@ -1722,6 +1722,7 @@ void UClothEditorWeightMapPaintTool::BeginChange()
 	
 	ActiveWeightEditChangeTracker = MakeUnique<FDynamicMeshChangeTracker>(GetSculptMesh());
 	ActiveWeightEditChangeTracker->BeginChange();
+	LongTransactions.Open(LOCTEXT("WeightPaintChange", "Weight Stroke"), GetToolManager());
 }
 
 void UClothEditorWeightMapPaintTool::EndChange()
@@ -1744,6 +1745,7 @@ void UClothEditorWeightMapPaintTool::EndChange()
 	};
 
 	GetToolManager()->EmitObjectChange(DynamicMeshComponent, MoveTemp(NewChange), LOCTEXT("WeightPaintChange", "Weight Stroke"));
+	LongTransactions.Close(GetToolManager());
 }
 
 
diff --git a/Engine/Plugins/Experimental/MeshModelingToolsetExp/Source/MeshModelingToolsExp/Private/DynamicMeshSculptTool.cpp b/Engine/Plugins/Experimental/MeshModelingToolsetExp/Source/MeshModelingToolsExp/Private/DynamicMeshSculptTool.cpp
--- a/Engine/Plugins/Experimental/MeshModelingToolsetExp/Source/MeshModelingToolsExp/Private/DynamicMeshSculptTool.cpp
+++ b/Engine/Plugins/Experimental/MeshModelingToolsetExp/Source/MeshModelingToolsExp/Private/DynamicMeshSculptTool.cpp
@@ -293,7 +293,7 @@ void UDynamicMeshSculptTool::Setup()
 
 void UDynamicMeshSculptTool::Shutdown(EToolShutdownType ShutdownType)
 {
-	ClosePlaceholderTransaction(); // End any placeholder transaction, which could arise from brush stroke actions
+	LongTransactions.CloseAll(GetToolManager());
 
 	if (ShutdownType == EToolShutdownType::Accept && AreAllTargetsValid() == false)
 	{
@@ -2421,7 +2421,7 @@ void UDynamicMeshSculptTool::BeginChange(bool bIsVertexChange)
 {
 	check(ActiveVertexChange == nullptr);
 	check(ActiveMeshChange == nullptr);
-	OpenPlaceholderTransaction(LOCTEXT("MeshSculptChange", "Brush Stroke")); // Create a placeholder change to prevent undo during the operation
+	LongTransactions.Open(LOCTEXT("MeshSculptChange", "Brush Stroke"), GetToolManager()); // Open brush stroke transaction to prevent undo during the operation
 	if (bIsVertexChange)
 	{
 		ActiveVertexChange = new FMeshVertexChangeBuilder();
@@ -2454,12 +2454,12 @@ void UDynamicMeshSculptTool::EndChange()
 		ActiveMeshChange = nullptr;
 	}
 
-	ClosePlaceholderTransaction();
+	LongTransactions.Close(GetToolManager());
 }
 
 void UDynamicMeshSculptTool::CancelChange()
 {
-	ClosePlaceholderTransaction();
+	LongTransactions.Close(GetToolManager());
 	delete ActiveVertexChange;
 	ActiveVertexChange = nullptr;
 	
@@ -2467,23 +2467,6 @@ void UDynamicMeshSculptTool::CancelChange()
 	ActiveMeshChange = nullptr;
 }
 
-void UDynamicMeshSculptTool::OpenPlaceholderTransaction(FText Name)
-{
-	if (!bHasActivePlaceholderTransaction)
-	{
-		GetToolManager()->BeginUndoTransaction(Name);
-		bHasActivePlaceholderTransaction = true;
-	}
-}
-void UDynamicMeshSculptTool::ClosePlaceholderTransaction()
-{
-	if (bHasActivePlaceholderTransaction)
-	{
-		GetToolManager()->EndUndoTransaction();
-		bHasActivePlaceholderTransaction = false;
-	}
-}
-
 void UDynamicMeshSculptTool::SaveActiveROI()
 {
 	if (ActiveMeshChange != nullptr)
diff --git a/Engine/Plugins/Experimental/MeshModelingToolsetExp/Source/MeshModelingToolsExp/Private/MeshGroupPaintTool.cpp b/Engine/Plugins/Experimental/MeshModelingToolsetExp/Source/MeshModelingToolsExp/Private/MeshGroupPaintTool.cpp
--- a/Engine/Plugins/Experimental/MeshModelingToolsetExp/Source/MeshModelingToolsExp/Private/MeshGroupPaintTool.cpp
+++ b/Engine/Plugins/Experimental/MeshModelingToolsetExp/Source/MeshModelingToolsExp/Private/MeshGroupPaintTool.cpp
@@ -1350,6 +1350,7 @@ void UMeshGroupPaintTool::BeginChange()
 {
 	check(ActiveGroupEditBuilder == nullptr);
 	ActiveGroupEditBuilder = MakeUnique<FDynamicMeshGroupEditBuilder>(ActiveGroupSet.Get());
+	LongTransactions.Open(LOCTEXT("GroupPaintChange", "Group Stroke"), GetToolManager());
 }
 
 void UMeshGroupPaintTool::EndChange()
@@ -1367,6 +1368,7 @@ void UMeshGroupPaintTool::EndChange()
 	};
 
 	GetToolManager()->EmitObjectChange(DynamicMeshComponent, MoveTemp(NewChange), LOCTEXT("GroupPaintChange", "Group Stroke"));
+	LongTransactions.Close(GetToolManager());
 
 	// debug groups are invalid now
 	bDrawGroupsDataValid = false;
diff --git a/Engine/Plugins/Experimental/MeshModelingToolsetExp/Source/MeshModelingToolsExp/Private/MeshVertexPaintTool.cpp b/Engine/Plugins/Experimental/MeshModelingToolsetExp/Source/MeshModelingToolsExp/Private/MeshVertexPaintTool.cpp
--- a/Engine/Plugins/Experimental/MeshModelingToolsetExp/Source/MeshModelingToolsExp/Private/MeshVertexPaintTool.cpp
+++ b/Engine/Plugins/Experimental/MeshModelingToolsetExp/Source/MeshModelingToolsExp/Private/MeshVertexPaintTool.cpp
@@ -1842,6 +1842,7 @@ void UMeshVertexPaintTool::BeginChange()
 
 	ActiveChangeBuilder = MakeUnique<TIndexedValuesChangeBuilder<FVector4f, FMeshVertexColorPaintChange>>();
 	ActiveChangeBuilder->BeginNewChange();
+	LongTransactions.Open(LOCTEXT("VertexPaintChange", "Paint Stroke"), GetToolManager());
 }
 
 
@@ -1871,6 +1872,7 @@ void UMeshVertexPaintTool::EndChange()
 	};
 
 	GetToolManager()->EmitObjectChange(this, MoveTemp(NewChange), LOCTEXT("VertexPaintChange", "Paint Stroke"));
+	LongTransactions.Close(GetToolManager());
 }
 
 
diff --git a/Engine/Plugins/Experimental/MeshModelingToolsetExp/Source/MeshModelingToolsExp/Private/MeshVertexSculptTool.cpp b/Engine/Plugins/Experimental/MeshModelingToolsetExp/Source/MeshModelingToolsExp/Private/MeshVertexSculptTool.cpp
--- a/Engine/Plugins/Experimental/MeshModelingToolsetExp/Source/MeshModelingToolsExp/Private/MeshVertexSculptTool.cpp
+++ b/Engine/Plugins/Experimental/MeshModelingToolsetExp/Source/MeshModelingToolsExp/Private/MeshVertexSculptTool.cpp
@@ -1361,6 +1361,7 @@ void UMeshVertexSculptTool::BeginChange()
 {
 	check(ActiveVertexChange == nullptr);
 	ActiveVertexChange = new FMeshVertexChangeBuilder();
+	LongTransactions.Open(LOCTEXT("VertexSculptChange", "Brush Stroke"), GetToolManager());
 }
 
 void UMeshVertexSculptTool::EndChange()
@@ -1374,7 +1375,6 @@ void UMeshVertexSculptTool::EndChange()
 		this->WaitForPendingUndoRedo();
 	};
 
-	GetToolManager()->BeginUndoTransaction(LOCTEXT("VertexSculptChange", "Brush Stroke"));
 	GetToolManager()->EmitObjectChange(DynamicMeshComponent, MoveTemp(NewChange), LOCTEXT("VertexSculptChange", "Brush Stroke"));
 	if (bMeshSymmetryIsValid && bApplySymmetry == false)
 	{
@@ -1383,7 +1383,7 @@ void UMeshVertexSculptTool::EndChange()
 		bMeshSymmetryIsValid = false;
 		SymmetryProperties->bSymmetryCanBeEnabled = bMeshSymmetryIsValid;
 	}
-	GetToolManager()->EndUndoTransaction();
+	LongTransactions.Close(GetToolManager());
 
 	delete ActiveVertexChange;
 	ActiveVertexChange = nullptr;
diff --git a/Engine/Plugins/Experimental/MeshModelingToolsetExp/Source/MeshModelingToolsExp/Private/Sculpting/MeshSculptToolBase.cpp b/Engine/Plugins/Experimental/MeshModelingToolsetExp/Source/MeshModelingToolsExp/Private/Sculpting/MeshSculptToolBase.cpp
--- a/Engine/Plugins/Experimental/MeshModelingToolsetExp/Source/MeshModelingToolsExp/Private/Sculpting/MeshSculptToolBase.cpp
+++ b/Engine/Plugins/Experimental/MeshModelingToolsetExp/Source/MeshModelingToolsExp/Private/Sculpting/MeshSculptToolBase.cpp
@@ -211,6 +211,7 @@ void UMeshSculptToolBase::Shutdown(EToolShutdownType ShutdownType)
 		DynamicMeshComponent = nullptr;
 	}
 
+	LongTransactions.CloseAll(GetToolManager());
 }
 
 
diff --git a/Engine/Plugins/Experimental/MeshModelingToolsetExp/Source/MeshModelingToolsExp/Public/DynamicMeshSculptTool.h b/Engine/Plugins/Experimental/MeshModelingToolsetExp/Source/MeshModelingToolsExp/Public/DynamicMeshSculptTool.h
--- a/Engine/Plugins/Experimental/MeshModelingToolsetExp/Source/MeshModelingToolsExp/Public/DynamicMeshSculptTool.h
+++ b/Engine/Plugins/Experimental/MeshModelingToolsetExp/Source/MeshModelingToolsExp/Public/DynamicMeshSculptTool.h
@@ -18,6 +18,7 @@
 #include "Properties/MeshMaterialProperties.h"
 #include "Properties/RemeshProperties.h"
 #include "TransformTypes.h"
+#include "TransactionUtil.h"
 #include "Sculpting/MeshSculptToolBase.h"
 #include "Async/Async.h"
 #include "Util/UniqueIndexSet.h"
@@ -497,9 +498,8 @@ private:
 
 	FMeshVertexChangeBuilder* ActiveVertexChange = nullptr;
 	UE::Geometry::FDynamicMeshChangeTracker* ActiveMeshChange = nullptr;
-	bool bHasActivePlaceholderTransaction = false;
-	void OpenPlaceholderTransaction(FText Name);
-	void ClosePlaceholderTransaction();
+	UE::TransactionUtil::FLongTransactionTracker LongTransactions;
+	
 	void BeginChange(bool bIsVertexChange);
 	void EndChange();
 	void CancelChange();
diff --git a/Engine/Plugins/Experimental/MeshModelingToolsetExp/Source/MeshModelingToolsExp/Public/Sculpting/MeshSculptToolBase.h b/Engine/Plugins/Experimental/MeshModelingToolsetExp/Source/MeshModelingToolsExp/Public/Sculpting/MeshSculptToolBase.h
--- a/Engine/Plugins/Experimental/MeshModelingToolsetExp/Source/MeshModelingToolsExp/Public/Sculpting/MeshSculptToolBase.h
+++ b/Engine/Plugins/Experimental/MeshModelingToolsetExp/Source/MeshModelingToolsExp/Public/Sculpting/MeshSculptToolBase.h
@@ -10,6 +10,7 @@
 #include "Properties/MeshMaterialProperties.h"
 #include "Changes/ValueWatcher.h"
 #include "Components/BaseDynamicMeshComponent.h"
+#include "TransactionUtil.h"
 #include "MeshSculptToolBase.generated.h"
 
 
@@ -600,4 +601,7 @@ protected:
 	virtual void UpdateFixedSculptPlaneRotation(const FQuat& Rotation);
 	virtual void UpdateFixedPlaneGizmoVisibility(bool bVisible);
 
+protected:
+	UE::TransactionUtil::FLongTransactionTracker LongTransactions;
+
 };
\ No newline at end of file
