Fixed wrong duplicated test

#rnx

diff --git a/Engine/Plugins/Experimental/MeshModelingToolsetExp/Source/SkeletalMeshModifiers/Private/SkeletonModifier.cpp b/Engine/Plugins/Experimental/MeshModelingToolsetExp/Source/SkeletalMeshModifiers/Private/SkeletonModifier.cpp
--- a/Engine/Plugins/Experimental/MeshModelingToolsetExp/Source/SkeletalMeshModifiers/Private/SkeletonModifier.cpp
+++ b/Engine/Plugins/Experimental/MeshModelingToolsetExp/Source/SkeletalMeshModifiers/Private/SkeletonModifier.cpp
@@ -988,7 +988,7 @@ bool USkeletonModifier::RemoveBones(const TArray<FName>& InBoneNames, const bool
 
 bool USkeletonModifier::RenameBone(const FName InOldBoneName, const FName InNewBoneName)
 {
-	if (InOldBoneName == NAME_None || InOldBoneName == NAME_None || InNewBoneName == InOldBoneName)
+	if (InOldBoneName == NAME_None || InNewBoneName == NAME_None || InNewBoneName == InOldBoneName)
 	{
 		UE_LOG(LogAnimation, Error, TEXT("Skeleton Modifier - Rename: cannot rename %s with %s."), *InOldBoneName.ToString(), *InNewBoneName.ToString());
 		return false;
