Add an overload in sculpt tool base to make it easy for subclasses to opt out of sharing brush properties, and make the mesh group paint tool and cloth weight map paint tools both opt out.

Fixes a bug where using the cloth weight map tool, then running the group paint tool after, would make the group paint tool run with some cloth weight tool default brush settings (especially flow = 0), which felt broken especially as those settings are not exposed to the user in either tool.

diff --git a/Engine/Plugins/Experimental/ChaosClothAssetEditor/Source/ChaosClothAssetEditorTools/Public/ChaosClothAsset/ClothWeightMapPaintTool.h b/Engine/Plugins/Experimental/ChaosClothAssetEditor/Source/ChaosClothAssetEditorTools/Public/ChaosClothAsset/ClothWeightMapPaintTool.h
--- a/Engine/Plugins/Experimental/ChaosClothAssetEditor/Source/ChaosClothAssetEditorTools/Public/ChaosClothAsset/ClothWeightMapPaintTool.h
+++ b/Engine/Plugins/Experimental/ChaosClothAssetEditor/Source/ChaosClothAssetEditorTools/Public/ChaosClothAsset/ClothWeightMapPaintTool.h
@@ -330,6 +330,8 @@ protected:
 	virtual void OnEndStroke() override;
 
 	virtual TUniquePtr<FMeshSculptBrushOp>& GetActiveBrushOp();
+
+	virtual bool SharesBrushPropertiesChanges() const override { return false; }
 	// end UMeshSculptToolBase API
 
 
diff --git a/Engine/Plugins/Experimental/MeshModelingToolsetExp/Source/MeshModelingToolsExp/Private/Sculpting/MeshSculptToolBase.cpp b/Engine/Plugins/Experimental/MeshModelingToolsetExp/Source/MeshModelingToolsExp/Private/Sculpting/MeshSculptToolBase.cpp
--- a/Engine/Plugins/Experimental/MeshModelingToolsetExp/Source/MeshModelingToolsExp/Private/Sculpting/MeshSculptToolBase.cpp
+++ b/Engine/Plugins/Experimental/MeshModelingToolsetExp/Source/MeshModelingToolsExp/Private/Sculpting/MeshSculptToolBase.cpp
@@ -97,7 +97,10 @@ void UMeshSculptToolBase::Setup()
 	UMeshSurfacePointTool::Setup();
 
 	BrushProperties = NewObject<USculptBrushProperties>(this);
-	BrushProperties->RestoreProperties(this);
+	if (SharesBrushPropertiesChanges())
+	{
+		BrushProperties->RestoreProperties(this);
+	}
 	// Note that brush properties includes BrushRadius, which, when not used as a constant,
 	// serves as an output property based on target size and brush size, and so it would need
 	// updating after the RestoreProperties() call. But deriving classes will call 
@@ -180,7 +183,10 @@ void UMeshSculptToolBase::Shutdown(EToolShutdownType ShutdownType)
 	BrushIndicator = nullptr;
 	GetToolManager()->GetPairedGizmoManager()->DeregisterGizmoType(VertexSculptIndicatorGizmoType);
 
-	BrushProperties->SaveProperties(this);
+	if (SharesBrushPropertiesChanges())
+	{
+		BrushProperties->SaveProperties(this);
+	}
 	if (GizmoProperties)
 	{
 		GizmoProperties->SaveProperties(this);
diff --git a/Engine/Plugins/Experimental/MeshModelingToolsetExp/Source/MeshModelingToolsExp/Public/MeshGroupPaintTool.h b/Engine/Plugins/Experimental/MeshModelingToolsetExp/Source/MeshModelingToolsExp/Public/MeshGroupPaintTool.h
--- a/Engine/Plugins/Experimental/MeshModelingToolsetExp/Source/MeshModelingToolsExp/Public/MeshGroupPaintTool.h
+++ b/Engine/Plugins/Experimental/MeshModelingToolsetExp/Source/MeshModelingToolsExp/Public/MeshGroupPaintTool.h
@@ -367,6 +367,8 @@ protected:
 	virtual void OnEndStroke() override;
 	virtual void OnCancelStroke() override;
 
+	virtual bool SharesBrushPropertiesChanges() const override { return false; }
+
 	virtual TUniquePtr<FMeshSculptBrushOp>& GetActiveBrushOp();
 	// end UMeshSculptToolBase API
 
diff --git a/Engine/Plugins/Experimental/MeshModelingToolsetExp/Source/MeshModelingToolsExp/Public/Sculpting/MeshSculptToolBase.h b/Engine/Plugins/Experimental/MeshModelingToolsetExp/Source/MeshModelingToolsExp/Public/Sculpting/MeshSculptToolBase.h
--- a/Engine/Plugins/Experimental/MeshModelingToolsetExp/Source/MeshModelingToolsExp/Public/Sculpting/MeshSculptToolBase.h
+++ b/Engine/Plugins/Experimental/MeshModelingToolsetExp/Source/MeshModelingToolsExp/Public/Sculpting/MeshSculptToolBase.h
@@ -296,6 +296,9 @@ protected:
 	virtual FDynamicMesh3* GetBaseMesh() { check(false); return nullptr; }
 	virtual const FDynamicMesh3* GetBaseMesh() const { check(false); return nullptr; }
 
+	// For any subclass where this returns false, BrushProperties will not be automatically saved/restored, so the class won't use BrushProperties changes made in other tools.
+	virtual bool SharesBrushPropertiesChanges() const { return true; }
+
 
 	/**
 	 * Subclass calls this to set up editing component
