add long transactions to dynamic mesh brush tool and protect drag transactions for its subclasses + deform mesh polygons tool

diff --git a/Engine/Plugins/Experimental/MeshModelingToolsetExp/Source/MeshModelingToolsEditorOnlyExp/Private/SkeletalMesh/SkinWeightsPaintTool.cpp b/Engine/Plugins/Experimental/MeshModelingToolsetExp/Source/MeshModelingToolsEditorOnlyExp/Private/SkeletalMesh/SkinWeightsPaintTool.cpp
--- a/Engine/Plugins/Experimental/MeshModelingToolsetExp/Source/MeshModelingToolsEditorOnlyExp/Private/SkeletalMesh/SkinWeightsPaintTool.cpp
+++ b/Engine/Plugins/Experimental/MeshModelingToolsetExp/Source/MeshModelingToolsEditorOnlyExp/Private/SkeletalMesh/SkinWeightsPaintTool.cpp
@@ -1001,6 +1001,7 @@ void USkinWeightsPaintTool::OnBeginDrag(const FRay& WorldRay)
 		StartStamp = UBaseBrushTool::LastBrushStamp;
 		LastStamp = StartStamp;
 		bStampPending = true;
+		LongTransactions.Open(LOCTEXT("PaintWeightChange", "Paint skin weights."), GetToolManager());
 	}
 }
 
@@ -1025,6 +1026,7 @@ void USkinWeightsPaintTool::OnEndDrag(const FRay& Ray)
 	// close change, record transaction
 	const FText TransactionLabel = LOCTEXT("PaintWeightChange", "Paint skin weights.");
 	EndChange(TransactionLabel);
+	LongTransactions.Close(GetToolManager());
 }
 
 bool USkinWeightsPaintTool::OnUpdateHover(const FInputDeviceRay& DevicePos)
diff --git a/Engine/Plugins/Experimental/MeshModelingToolsetExp/Source/MeshModelingToolsExp/Private/DeformMeshPolygonsTool.cpp b/Engine/Plugins/Experimental/MeshModelingToolsetExp/Source/MeshModelingToolsExp/Private/DeformMeshPolygonsTool.cpp
--- a/Engine/Plugins/Experimental/MeshModelingToolsetExp/Source/MeshModelingToolsExp/Private/DeformMeshPolygonsTool.cpp
+++ b/Engine/Plugins/Experimental/MeshModelingToolsetExp/Source/MeshModelingToolsExp/Private/DeformMeshPolygonsTool.cpp
@@ -939,6 +939,8 @@ void UDeformMeshPolygonsTool::Setup()
 
 void UDeformMeshPolygonsTool::Shutdown(EToolShutdownType ShutdownType)
 {
+	LongTransactions.CloseAll(GetToolManager());
+
 	//Tell the background thread to cancel the rest of its jobs before we close;
 	LaplacianDeformer.Reset();
 
@@ -1302,8 +1304,12 @@ void UDeformMeshPolygonsTool::OnCancelDrag()
 	QuickAxisRotator.Reset();
 	QuickAxisTranslater.Reset();
 
-	delete ActiveVertexChange;
-	ActiveVertexChange = nullptr;
+	if (ActiveVertexChange)
+	{
+		LongTransactions.Close(GetToolManager());
+		delete ActiveVertexChange;
+		ActiveVertexChange = nullptr;
+	}
 }
 
 
@@ -1673,6 +1679,7 @@ void UDeformMeshPolygonsTool::BeginChange()
 			ActiveVertexChange = new FMeshVertexChangeBuilder(EMeshVertexChangeComponents::VertexPositions |
 			                                                  EMeshVertexChangeComponents::OverlayNormals);
 			UpdateChangeFromROI(false);
+			LongTransactions.Open(LOCTEXT("PolyMeshDeformationChange", "PolyMesh Edit"), GetToolManager());
 		}
 	}
 }
@@ -1685,6 +1692,7 @@ void UDeformMeshPolygonsTool::EndChange()
 		UpdateChangeFromROI(true);
 		GetToolManager()->EmitObjectChange(DynamicMeshComponent, MoveTemp(ActiveVertexChange->Change),
 		                                   LOCTEXT("PolyMeshDeformationChange", "PolyMesh Edit"));
+		LongTransactions.Close(GetToolManager());
 	}
 
 	delete ActiveVertexChange;
diff --git a/Engine/Plugins/Experimental/MeshModelingToolsetExp/Source/MeshModelingToolsExp/Private/DynamicMeshBrushTool.cpp b/Engine/Plugins/Experimental/MeshModelingToolsetExp/Source/MeshModelingToolsExp/Private/DynamicMeshBrushTool.cpp
--- a/Engine/Plugins/Experimental/MeshModelingToolsetExp/Source/MeshModelingToolsExp/Private/DynamicMeshBrushTool.cpp
+++ b/Engine/Plugins/Experimental/MeshModelingToolsetExp/Source/MeshModelingToolsExp/Private/DynamicMeshBrushTool.cpp
@@ -60,6 +60,8 @@ void UDynamicMeshBrushTool::Shutdown(EToolShutdownType ShutdownType)
 {
 	UBaseBrushTool::Shutdown(ShutdownType);
 
+	LongTransactions.CloseAll(GetToolManager());
+
 	UE::ToolTarget::ShowSourceObject(Target);
 
 	if (PreviewMesh != nullptr)
diff --git a/Engine/Plugins/Experimental/MeshModelingToolsetExp/Source/MeshModelingToolsExp/Private/MeshAttributePaintTool.cpp b/Engine/Plugins/Experimental/MeshModelingToolsetExp/Source/MeshModelingToolsExp/Private/MeshAttributePaintTool.cpp
--- a/Engine/Plugins/Experimental/MeshModelingToolsetExp/Source/MeshModelingToolsExp/Private/MeshAttributePaintTool.cpp
+++ b/Engine/Plugins/Experimental/MeshModelingToolsetExp/Source/MeshModelingToolsExp/Private/MeshAttributePaintTool.cpp
@@ -336,6 +336,7 @@ void UMeshAttributePaintTool::OnBeginDrag(const FRay& WorldRay)
 		StartStamp = UBaseBrushTool::LastBrushStamp;
 		LastStamp = StartStamp;
 		bStampPending = true;
+		LongTransactions.Open(LOCTEXT("AttributeValuesChange", "Paint"), GetToolManager());
 	}
 }
 
@@ -364,9 +365,8 @@ void UMeshAttributePaintTool::OnEndDrag(const FRay& Ray)
 	TUniquePtr<FMeshAttributePaintChange> Change = EndChange();
 	if (Change)
 	{
-		GetToolManager()->BeginUndoTransaction(LOCTEXT("AttributeValuesChange", "Paint"));
 		GetToolManager()->EmitObjectChange(this, MoveTemp(Change), LOCTEXT("AttributeValuesChange", "Paint"));
-		GetToolManager()->EndUndoTransaction();
+		LongTransactions.Close(GetToolManager());
 	}
 }
 
diff --git a/Engine/Plugins/Experimental/MeshModelingToolsetExp/Source/MeshModelingToolsExp/Private/MeshSelectionTool.cpp b/Engine/Plugins/Experimental/MeshModelingToolsetExp/Source/MeshModelingToolsExp/Private/MeshSelectionTool.cpp
--- a/Engine/Plugins/Experimental/MeshModelingToolsetExp/Source/MeshModelingToolsExp/Private/MeshSelectionTool.cpp
+++ b/Engine/Plugins/Experimental/MeshModelingToolsetExp/Source/MeshModelingToolsExp/Private/MeshSelectionTool.cpp
@@ -455,6 +455,7 @@ void UMeshSelectionTool::OnBeginDrag(const FRay& WorldRay)
 		StartStamp = UBaseBrushTool::LastBrushStamp;
 		LastStamp = StartStamp;
 		bStampPending = true;
+		LongTransactions.Open(LOCTEXT("MeshSelectionChange", "Mesh Selection"), GetToolManager());
 	}
 }
 
@@ -759,6 +760,7 @@ void UMeshSelectionTool::OnEndDrag(const FRay& Ray)
 	// close change record
 	TUniquePtr<FToolCommandChange> Change = EndChange();
 	GetToolManager()->EmitObjectChange(Selection, MoveTemp(Change), LOCTEXT("MeshSelectionChange", "Mesh Selection"));
+	LongTransactions.Close(GetToolManager());
 }
 
 
@@ -828,7 +830,12 @@ bool UMeshSelectionTool::ExecuteNestedCancelCommand()
 {
 	if (CanCurrentlyNestedCancel())
 	{
-		ClearSelection();
+		// Only actually clear if there's no in-progress change; this just means escape does nothing in the middle of a mouse drag
+		// (We could make esc undo the active change and cancel the drag in this case; not clear if that is worthwhile. In most similar situations, escape just exits the tool.)
+		if (!ActiveSelectionChange)
+		{
+			ClearSelection();
+		}
 		return true;
 	}
 	return false;
diff --git a/Engine/Plugins/Experimental/MeshModelingToolsetExp/Source/MeshModelingToolsExp/Private/SeamSculptTool.cpp b/Engine/Plugins/Experimental/MeshModelingToolsetExp/Source/MeshModelingToolsExp/Private/SeamSculptTool.cpp
--- a/Engine/Plugins/Experimental/MeshModelingToolsetExp/Source/MeshModelingToolsExp/Private/SeamSculptTool.cpp
+++ b/Engine/Plugins/Experimental/MeshModelingToolsetExp/Source/MeshModelingToolsExp/Private/SeamSculptTool.cpp
@@ -156,6 +156,7 @@ void USeamSculptTool::OnBeginDrag(const FRay& Ray)
 		CaptureState = EActiveCaptureState::DrawNewPath;
 		DrawPathStartVertex = CurrentSnapVertex;
 		DrawPathStartPositionLocal = CurrentSnapPositionLocal;
+		LongTransactions.Open(LOCTEXT("CreateSeamChange", "Add UV Seam"), GetToolManager());
 
 		CurDrawPath.Reset();
 	}
@@ -193,6 +194,7 @@ void USeamSculptTool::OnEndDrag(const FRay& Ray)
 	{
 		CreateSeamAlongPath();
 		CurDrawPath.Reset();
+		LongTransactions.Close(GetToolManager());
 	}
 
 	CaptureState = EActiveCaptureState::NoState;
diff --git a/Engine/Plugins/Experimental/MeshModelingToolsetExp/Source/MeshModelingToolsExp/Public/DeformMeshPolygonsTool.h b/Engine/Plugins/Experimental/MeshModelingToolsetExp/Source/MeshModelingToolsExp/Public/DeformMeshPolygonsTool.h
--- a/Engine/Plugins/Experimental/MeshModelingToolsetExp/Source/MeshModelingToolsExp/Public/DeformMeshPolygonsTool.h
+++ b/Engine/Plugins/Experimental/MeshModelingToolsetExp/Source/MeshModelingToolsExp/Public/DeformMeshPolygonsTool.h
@@ -17,6 +17,8 @@
 #include "Selection/GroupTopologySelector.h"
 #include "Operations/GroupTopologyDeformer.h"
 #include "Solvers/MeshLaplacian.h"
+#include "TransactionUtil.h"
+
 #include "DeformMeshPolygonsTool.generated.h"
 
 class FMeshVertexChangeBuilder;
@@ -267,5 +269,7 @@ protected:
 	void EndChange();
 	void UpdateChangeFromROI(bool bFinal);
 
+private:
+	UE::TransactionUtil::FLongTransactionTracker LongTransactions;
 };
 
diff --git a/Engine/Plugins/Experimental/MeshModelingToolsetExp/Source/MeshModelingToolsExp/Public/DynamicMeshBrushTool.h b/Engine/Plugins/Experimental/MeshModelingToolsetExp/Source/MeshModelingToolsExp/Public/DynamicMeshBrushTool.h
--- a/Engine/Plugins/Experimental/MeshModelingToolsetExp/Source/MeshModelingToolsExp/Public/DynamicMeshBrushTool.h
+++ b/Engine/Plugins/Experimental/MeshModelingToolsetExp/Source/MeshModelingToolsExp/Public/DynamicMeshBrushTool.h
@@ -9,6 +9,7 @@
 #include "DynamicMesh/DynamicMeshAABBTree3.h"
 #include "Components/DynamicMeshComponent.h"
 #include "PreviewMesh.h"
+#include "TransactionUtil.h"
 #include "DynamicMeshBrushTool.generated.h"
 
 PREDECLARE_USE_GEOMETRY_CLASS(FDynamicMesh3);
@@ -54,4 +55,6 @@ protected:
 	// UBaseBrushTool private interface
 	//
 	virtual double EstimateMaximumTargetDimension() override;
+
+	UE::TransactionUtil::FLongTransactionTracker LongTransactions;
 };
\ No newline at end of file
