Skeleton Editing Tool: update gizmo on hierarchy changes

#rnx

diff --git a/Engine/Plugins/Experimental/Animation/SkeletalMeshModelingTools/Source/SkeletalMeshModelingTools/Private/SReferenceSkeletonTree.cpp b/Engine/Plugins/Experimental/Animation/SkeletalMeshModelingTools/Source/SkeletalMeshModelingTools/Private/SReferenceSkeletonTree.cpp
--- a/Engine/Plugins/Experimental/Animation/SkeletalMeshModelingTools/Source/SkeletalMeshModelingTools/Private/SReferenceSkeletonTree.cpp
+++ b/Engine/Plugins/Experimental/Animation/SkeletalMeshModelingTools/Source/SkeletalMeshModelingTools/Private/SReferenceSkeletonTree.cpp
@@ -1022,6 +1022,7 @@ FReply SReferenceSkeletonTree::OnAcceptDrop(const FDragDropEvent& DragDropEvent,
 		if (Notifier.IsValid())
 		{
 			Notifier->Notify( {BoneName}, ESkeletalMeshNotifyType::HierarchyChanged);
+			Notifier->Notify( {BoneName}, ESkeletalMeshNotifyType::BonesSelected);
 		}
 		EndChange();
 	}
diff --git a/Engine/Plugins/Experimental/MeshModelingToolsetExp/Source/MeshModelingToolsEditorOnlyExp/Private/SkeletalMesh/SkeletonEditingTool.cpp b/Engine/Plugins/Experimental/MeshModelingToolsetExp/Source/MeshModelingToolsEditorOnlyExp/Private/SkeletalMesh/SkeletonEditingTool.cpp
--- a/Engine/Plugins/Experimental/MeshModelingToolsetExp/Source/MeshModelingToolsEditorOnlyExp/Private/SkeletalMesh/SkeletonEditingTool.cpp
+++ b/Engine/Plugins/Experimental/MeshModelingToolsetExp/Source/MeshModelingToolsEditorOnlyExp/Private/SkeletalMesh/SkeletonEditingTool.cpp
@@ -436,6 +436,7 @@ void USkeletonEditingTool::RegisterUnParentAction(FInteractiveToolActionSet& InO
 		[this]()
 		{
 			UnParentBones();
+			UpdateGizmo();
 		});
 }
 
@@ -1115,6 +1116,7 @@ void USkeletonEditingTool::HandleSkeletalMeshModified(const TArray<FName>& InBon
 			Selection = BoneNames;
 			break;
 		case ESkeletalMeshNotifyType::HierarchyChanged:
+			UpdateGizmo();
 			break;
 		default:
 			break;
