Skin Paint Tool: fixed crash when ending drag
   - follow-up to 29411935

#rnx

diff --git a/Engine/Plugins/Experimental/MeshModelingToolsetExp/Source/MeshModelingToolsEditorOnlyExp/Private/SkeletalMesh/SkinWeightsPaintTool.cpp b/Engine/Plugins/Experimental/MeshModelingToolsetExp/Source/MeshModelingToolsEditorOnlyExp/Private/SkeletalMesh/SkinWeightsPaintTool.cpp
--- a/Engine/Plugins/Experimental/MeshModelingToolsetExp/Source/MeshModelingToolsEditorOnlyExp/Private/SkeletalMesh/SkinWeightsPaintTool.cpp
+++ b/Engine/Plugins/Experimental/MeshModelingToolsetExp/Source/MeshModelingToolsEditorOnlyExp/Private/SkeletalMesh/SkinWeightsPaintTool.cpp
@@ -1023,10 +1023,13 @@ void USkinWeightsPaintTool::OnEndDrag(const FRay& Ray)
 	bSmoothStroke = false;
 	bStampPending = false;
 
-	// close change, record transaction
-	const FText TransactionLabel = LOCTEXT("PaintWeightChange", "Paint skin weights.");
-	EndChange(TransactionLabel);
-	LongTransactions.Close(GetToolManager());
+	if (ActiveChange)
+	{
+		// close change, record transaction
+		const FText TransactionLabel = LOCTEXT("PaintWeightChange", "Paint skin weights.");
+		EndChange(TransactionLabel);
+		LongTransactions.Close(GetToolManager());
+	}
 }
 
 bool USkinWeightsPaintTool::OnUpdateHover(const FInputDeviceRay& DevicePos)
