Skeleton Editing Tool: fixed selection notifications

#rnx

diff --git a/Engine/Plugins/Experimental/Animation/SkeletalMeshModelingTools/Source/SkeletalMeshModelingTools/Private/SReferenceSkeletonTree.cpp b/Engine/Plugins/Experimental/Animation/SkeletalMeshModelingTools/Source/SkeletalMeshModelingTools/Private/SReferenceSkeletonTree.cpp
--- a/Engine/Plugins/Experimental/Animation/SkeletalMeshModelingTools/Source/SkeletalMeshModelingTools/Private/SReferenceSkeletonTree.cpp
+++ b/Engine/Plugins/Experimental/Animation/SkeletalMeshModelingTools/Source/SkeletalMeshModelingTools/Private/SReferenceSkeletonTree.cpp
@@ -845,11 +845,20 @@ void SReferenceSkeletonTree::HandleGetChildrenForTree(
 
 void SReferenceSkeletonTree::OnSelectionChanged(TSharedPtr<FBoneElement> InItem, ESelectInfo::Type InSelectInfo)
 {
-	if (!InItem || InSelectInfo == ESelectInfo::Direct)
+	if (InSelectInfo == ESelectInfo::Direct)
 	{
 		return;
 	}
 	
+	if (!InItem)
+	{
+		if (Notifier.IsValid())
+		{
+			Notifier->Notify({}, ESkeletalMeshNotifyType::BonesSelected);
+		}
+		return;
+	}
+	
 	TArray<FName> BoneNames;
 	GetSelectedBoneNames(BoneNames);
 	
diff --git a/Engine/Plugins/Experimental/Animation/SkeletalMeshModelingTools/Source/SkeletalMeshModelingTools/Private/SkeletalMeshEditorUtils.cpp b/Engine/Plugins/Experimental/Animation/SkeletalMeshModelingTools/Source/SkeletalMeshModelingTools/Private/SkeletalMeshEditorUtils.cpp
--- a/Engine/Plugins/Experimental/Animation/SkeletalMeshModelingTools/Source/SkeletalMeshModelingTools/Private/SkeletalMeshEditorUtils.cpp
+++ b/Engine/Plugins/Experimental/Animation/SkeletalMeshModelingTools/Source/SkeletalMeshModelingTools/Private/SkeletalMeshEditorUtils.cpp
@@ -264,6 +264,11 @@ void USkeletalMeshEditorContextObject::BindRefSkeletonTree(ISkeletalMeshEditingI
 	FBindData BindData;
 	Tie(BindData.ToToolNotifierHandle, BindData.FromToolNotifierHandle) = BindInterfaceTo(InEditingInterface, RefSkeletonTree->GetNotifier());
 	TreeBindings.Emplace(InEditingInterface, BindData);
+
+	if(const TSharedPtr<ISkeletalMeshEditorBinding> Binding = Editor.Pin()->GetBinding())
+	{
+		RefSkeletonTree->GetNotifier().HandleNotification(Binding->GetSelectedBones(), ESkeletalMeshNotifyType::BonesSelected);
+	}
 }
 
 void USkeletalMeshEditorContextObject::UnbindRefSkeletonTree(ISkeletalMeshEditingInterface* InEditingInterface)
diff --git a/Engine/Plugins/Experimental/MeshModelingToolsetExp/Source/MeshModelingToolsEditorOnlyExp/Private/SkeletalMesh/SkeletonEditingTool.cpp b/Engine/Plugins/Experimental/MeshModelingToolsetExp/Source/MeshModelingToolsEditorOnlyExp/Private/SkeletalMesh/SkeletonEditingTool.cpp
--- a/Engine/Plugins/Experimental/MeshModelingToolsetExp/Source/MeshModelingToolsEditorOnlyExp/Private/SkeletalMesh/SkeletonEditingTool.cpp
+++ b/Engine/Plugins/Experimental/MeshModelingToolsetExp/Source/MeshModelingToolsEditorOnlyExp/Private/SkeletalMesh/SkeletonEditingTool.cpp
@@ -113,16 +113,6 @@ void USkeletonEditingTool::Setup()
 	Modifier = NewObject<USkeletonModifier>(this);
 	Modifier->SetSkeletalMesh(SkeletalMesh);
 
-	// setup current bone
-	const FReferenceSkeleton& RefSkeleton = Modifier->GetReferenceSkeleton();
-	const int32 NumBones = RefSkeleton.GetNum();
-	const FName& RootBoneName = NumBones ? RefSkeleton.GetBoneName(0) : NAME_None;
-
-	if (NumBones)
-	{
-		Selection = {RootBoneName};
-	}
-
 	// setup preview
 	{
 		PreviewMesh = NewObject<UPreviewMesh>(this);
@@ -209,6 +199,7 @@ void USkeletonEditingTool::Setup()
 		if (GizmoWrapper)
 		{
 			GizmoWrapper->Component = Component;
+			GizmoWrapper->Initialize();
 		}
 	}
 
