Some cleanup in the baking tools

diff --git a/Engine/Plugins/Experimental/MeshModelingToolsetExp/Source/MeshModelingToolsEditorOnlyExp/Private/BakeRenderCaptureTool.cpp b/Engine/Plugins/Experimental/MeshModelingToolsetExp/Source/MeshModelingToolsEditorOnlyExp/Private/BakeRenderCaptureTool.cpp
--- a/Engine/Plugins/Experimental/MeshModelingToolsetExp/Source/MeshModelingToolsEditorOnlyExp/Private/BakeRenderCaptureTool.cpp
+++ b/Engine/Plugins/Experimental/MeshModelingToolsetExp/Source/MeshModelingToolsEditorOnlyExp/Private/BakeRenderCaptureTool.cpp
@@ -315,14 +315,14 @@ void UBakeRenderCaptureTool::Setup()
 
 	Settings->MapPreview = BaseColorTexParamName;
 	MapPreviewWatcherIndex = Settings->WatchProperty(Settings->MapPreview, [this](FString) { UpdateVisualization(); });
-	Settings->WatchProperty(Settings->SamplesPerPixel, [this](EBakeTextureSamplesPerPixel) { OpState |= EBakeOpState::Evaluate; });
-	Settings->WatchProperty(Settings->TextureSize, [this](EBakeTextureResolution) { OpState |= EBakeOpState::Evaluate; });
+	Settings->WatchProperty(Settings->SamplesPerPixel, [this](EBakeTextureSamplesPerPixel) { BakeOpState |= EBakeOpState::Evaluate; });
+	Settings->WatchProperty(Settings->TextureSize, [this](EBakeTextureResolution) { BakeOpState |= EBakeOpState::Evaluate; });
 	Settings->WatchProperty(Settings->ValidSampleDepthThreshold, [this](float ValidSampleDepthThreshold)
 	{
 		// The depth capture channel is enabled implicitly when the following parameter is greater than 0.
 		// To disable the depth capture we need to set the value to 0. See :EnableDisableDeviceDepthMap 
 		RenderCaptureProperties->bDeviceDepthMap = (ValidSampleDepthThreshold > 0);
-		OpState |= EBakeOpState::Evaluate;
+		BakeOpState |= EBakeOpState::Evaluate;
 	});
 
 	// Put these properties before the list of preview textures so its easier to find
@@ -335,27 +335,27 @@ void UBakeRenderCaptureTool::Setup()
 	RenderCaptureProperties->RestoreProperties(this);
 	AddToolPropertySource(RenderCaptureProperties);
 
-	RenderCaptureProperties->WatchProperty(RenderCaptureProperties->Resolution, [this](EBakeTextureResolution) { OpState |= EBakeOpState::Evaluate; });
-	RenderCaptureProperties->WatchProperty(RenderCaptureProperties->bBaseColorMap, [this](bool) { OpState |= EBakeOpState::Evaluate; });
-	RenderCaptureProperties->WatchProperty(RenderCaptureProperties->bNormalMap, [this](bool) { OpState |= EBakeOpState::Evaluate; });
-	RenderCaptureProperties->WatchProperty(RenderCaptureProperties->bMetallicMap, [this](bool) { OpState |= EBakeOpState::Evaluate; });
-	RenderCaptureProperties->WatchProperty(RenderCaptureProperties->bRoughnessMap, [this](bool) { OpState |= EBakeOpState::Evaluate; });
-	RenderCaptureProperties->WatchProperty(RenderCaptureProperties->bSpecularMap, [this](bool) { OpState |= EBakeOpState::Evaluate; });
-	RenderCaptureProperties->WatchProperty(RenderCaptureProperties->bPackedMRSMap, [this](bool) { OpState |= EBakeOpState::Evaluate; });
-	RenderCaptureProperties->WatchProperty(RenderCaptureProperties->bEmissiveMap, [this](bool) { OpState |= EBakeOpState::Evaluate; });
-	RenderCaptureProperties->WatchProperty(RenderCaptureProperties->bOpacityMap, [this](bool) { OpState |= EBakeOpState::Evaluate; });
-	RenderCaptureProperties->WatchProperty(RenderCaptureProperties->bSubsurfaceColorMap, [this](bool) { OpState |= EBakeOpState::Evaluate; });
-	RenderCaptureProperties->WatchProperty(RenderCaptureProperties->bAntiAliasing, [this](bool) { OpState |= EBakeOpState::Evaluate; });
+	RenderCaptureProperties->WatchProperty(RenderCaptureProperties->Resolution, [this](EBakeTextureResolution) { BakeOpState |= EBakeOpState::Evaluate; });
+	RenderCaptureProperties->WatchProperty(RenderCaptureProperties->bBaseColorMap, [this](bool) { BakeOpState |= EBakeOpState::Evaluate; });
+	RenderCaptureProperties->WatchProperty(RenderCaptureProperties->bNormalMap, [this](bool) { BakeOpState |= EBakeOpState::Evaluate; });
+	RenderCaptureProperties->WatchProperty(RenderCaptureProperties->bMetallicMap, [this](bool) { BakeOpState |= EBakeOpState::Evaluate; });
+	RenderCaptureProperties->WatchProperty(RenderCaptureProperties->bRoughnessMap, [this](bool) { BakeOpState |= EBakeOpState::Evaluate; });
+	RenderCaptureProperties->WatchProperty(RenderCaptureProperties->bSpecularMap, [this](bool) { BakeOpState |= EBakeOpState::Evaluate; });
+	RenderCaptureProperties->WatchProperty(RenderCaptureProperties->bPackedMRSMap, [this](bool) { BakeOpState |= EBakeOpState::Evaluate; });
+	RenderCaptureProperties->WatchProperty(RenderCaptureProperties->bEmissiveMap, [this](bool) { BakeOpState |= EBakeOpState::Evaluate; });
+	RenderCaptureProperties->WatchProperty(RenderCaptureProperties->bOpacityMap, [this](bool) { BakeOpState |= EBakeOpState::Evaluate; });
+	RenderCaptureProperties->WatchProperty(RenderCaptureProperties->bSubsurfaceColorMap, [this](bool) { BakeOpState |= EBakeOpState::Evaluate; });
+	RenderCaptureProperties->WatchProperty(RenderCaptureProperties->bAntiAliasing, [this](bool) { BakeOpState |= EBakeOpState::Evaluate; });
 	// These are not exposed to the UI, but we watch them anyway because we might change that later
-	RenderCaptureProperties->WatchProperty(RenderCaptureProperties->CaptureFieldOfView, [this](float) { OpState |= EBakeOpState::Evaluate; });
-	RenderCaptureProperties->WatchProperty(RenderCaptureProperties->NearPlaneDist, [this](float) { OpState |= EBakeOpState::Evaluate; });
+	RenderCaptureProperties->WatchProperty(RenderCaptureProperties->CaptureFieldOfView, [this](float) { BakeOpState |= EBakeOpState::Evaluate; });
+	RenderCaptureProperties->WatchProperty(RenderCaptureProperties->NearPlaneDist, [this](float) { BakeOpState |= EBakeOpState::Evaluate; });
 	
 	InputMeshSettings = NewObject<UBakeRenderCaptureInputToolProperties>(this);
 	InputMeshSettings->RestoreProperties(this);
 	AddToolPropertySource(InputMeshSettings);
 	InputMeshSettings->TargetStaticMesh = UE::ToolTarget::GetStaticMeshFromTargetIfAvailable(Target);
 	UpdateUVLayerNames(InputMeshSettings->TargetUVLayer, InputMeshSettings->TargetUVLayerNamesList, *TargetMesh);
-	InputMeshSettings->WatchProperty(InputMeshSettings->TargetUVLayer, [this](FString) { OpState |= EBakeOpState::Evaluate; });
+	InputMeshSettings->WatchProperty(InputMeshSettings->TargetUVLayer, [this](FString) { BakeOpState |= EBakeOpState::Evaluate; });
 	
 	ResultSettings = NewObject<UBakeRenderCaptureResults>(this);
 	ResultSettings->RestoreProperties(this);
@@ -373,8 +373,13 @@ void UBakeRenderCaptureTool::Setup()
 
 	SceneCapture = MakeShared<FSceneCapturePhotoSet>();
 
+	// Initialize baker background compute
+	BakeOp = MakeUnique<TGenericDataBackgroundCompute<FMeshMapBaker>>();
+	BakeOp->Setup(this);
+	BakeOp->OnResultUpdated.AddLambda([this](const TUniquePtr<FMeshMapBaker>& NewResult) { OnMapsUpdated(NewResult); });
+
 	// Make sure we trigger SceneCapture computation in UpdateResult
-	OpState |= EBakeOpState::Evaluate;
+	BakeOpState |= EBakeOpState::Evaluate;
 
 	SetToolDisplayName(LOCTEXT("ToolName", "Bake Render Capture"));
 	GetToolManager()->DisplayMessage(
@@ -421,27 +426,16 @@ void UBakeRenderCaptureTool::Render(IToolsContextRenderAPI* RenderAPI)
 
 void UBakeRenderCaptureTool::OnTick(float DeltaTime)
 {
-	if (Compute)
-	{
-		Compute->Tick(DeltaTime);
+	BakeOp->Tick(DeltaTime);
 
-		if (static_cast<bool>(OpState & EBakeOpState::Invalid))
-		{
-			PreviewMesh->SetOverrideRenderMaterial(ErrorPreviewMaterial);
-		}
-		else
-		{
-			const float ElapsedComputeTime = Compute->GetElapsedComputeTime();
-			if (!CanAccept() && ElapsedComputeTime > SecondsBeforeWorkingMaterial)
-			{
-				PreviewMesh->SetOverrideRenderMaterial(WorkingPreviewMaterial);
-			}
-		}
-	}
-	else if (static_cast<bool>(OpState & EBakeOpState::Invalid))
+	if (static_cast<bool>(BakeOpState & EBakeOpState::Invalid))
 	{
 		PreviewMesh->SetOverrideRenderMaterial(ErrorPreviewMaterial);
-	} 
+	}
+	else if (!CanAccept() && BakeOp->GetElapsedComputeTime() > SecondsBeforeWorkingMaterial)
+	{
+		PreviewMesh->SetOverrideRenderMaterial(WorkingPreviewMaterial);
+	}
 }
 
 
@@ -454,17 +448,11 @@ void UBakeRenderCaptureTool::OnShutdown(EToolShutdownType ShutdownType)
 	InputMeshSettings->SaveProperties(this);
 	VisualizationProps->SaveProperties(this);
 
-	if (PreviewMesh != nullptr)
-	{
-		PreviewMesh->SetVisible(false);
-		PreviewMesh->Disconnect();
-		PreviewMesh = nullptr;
-	}
+	PreviewMesh->SetVisible(false);
+	PreviewMesh->Disconnect();
+	PreviewMesh = nullptr;
 
-	if (Compute)
-	{
-		Compute->Shutdown();
-	}
+	BakeOp->Shutdown();
 
 	if (ShutdownType == EToolShutdownType::Accept && ResultSettings->IsEmpty() == false)
 	{
@@ -714,7 +702,7 @@ void UBakeRenderCaptureTool::CreateAssets(UWorld* SourceWorld)
 // Return false if the user requested a texture but it is not yet baked, or if the tool is in an invalid state
 bool UBakeRenderCaptureTool::CanAccept() const
 {
-	if ((OpState & EBakeOpState::Invalid) == EBakeOpState::Invalid)
+	if ((BakeOpState & EBakeOpState::Invalid) == EBakeOpState::Invalid)
 	{
 		return false;
 	}
@@ -773,7 +761,6 @@ bool UBakeRenderCaptureTool::CanAccept() const
 
 
 
-// This function gets called via UBakeRenderCaptureTool::InvalidateCompute
 TUniquePtr<TGenericDataOperator<FMeshMapBaker>> UBakeRenderCaptureTool::MakeNewOperator()
 {
 	// We should not have requested a bake if we don't have a SceneCapture
@@ -1022,24 +1009,11 @@ void UBakeRenderCaptureTool::InitializePreviewMaterials()
 }
 
 
-void UBakeRenderCaptureTool::InvalidateCompute()
-{
-	if (!Compute)
-	{
-		// Initialize background compute
-		Compute = MakeUnique<TGenericDataBackgroundCompute<FMeshMapBaker>>();
-		Compute->Setup(this);
-		Compute->OnResultUpdated.AddLambda([this](const TUniquePtr<FMeshMapBaker>& NewResult) { OnMapsUpdated(NewResult); });
-	}
-	Compute->InvalidateResult();
-	OpState = EBakeOpState::Clean;
-}
-
 // Process dirty props and update background compute. Called by UBakeRenderCaptureTool::Render
 void UBakeRenderCaptureTool::UpdateResult()
 {
 	// Return if the bake is already launched/complete.
-	if (OpState == EBakeOpState::Clean)
+	if (BakeOpState == EBakeOpState::Clean)
 	{
 		return;
 	}
@@ -1106,9 +1080,9 @@ void UBakeRenderCaptureTool::UpdateResult()
 		NotifyOfPropertyChangeByTool(Settings);
 	}
 
-	// The bake operation, Compute, stores a pointer to the SceneCapture so that must not be modified while baking
-	const bool bComputeInProgress = (Compute && (Compute->GetElapsedComputeTime() > 0.f));
-	if (bComputeInProgress)
+	// The bake operation stores a pointer to the SceneCapture so that must not be modified while baking
+	const bool bBakeOpInProgress = BakeOp->GetElapsedComputeTime() > 0.f;
+	if (bBakeOpInProgress)
 	{
 		return;
 	}
@@ -1168,14 +1142,14 @@ void UBakeRenderCaptureTool::UpdateResult()
 		InvalidateResults(FRenderCaptureTypeFlags::All(true));
 
 		// Only call UpdateVisualization when we first detect the invalid inputs
-		const bool bWasValid = static_cast<bool>(OpState & EBakeOpState::Invalid) == false;
+		const bool bWasValid = static_cast<bool>(BakeOpState & EBakeOpState::Invalid) == false;
 		if (bWasValid)
 		{
 			UpdateVisualization();
 		}
 
 		// Set an invalid op state so we re-enter this function until the inputs are valid
-		OpState = EBakeOpState::Invalid;
+		BakeOpState = EBakeOpState::Invalid;
 	}
 	else
 	{
@@ -1234,7 +1208,8 @@ void UBakeRenderCaptureTool::UpdateResult()
 			UpdateVisualization();
 
 			// Start another bake operation, this will bake the computed captures with invalidated/null texture results
-			InvalidateCompute();
+			BakeOp->InvalidateResult();
+			BakeOpState = EBakeOpState::Clean;
 
 			// Cache computed parameters which are used to determine if results need re-baking
 			ComputedTargetUVLayer = InputMeshSettings->TargetUVLayer;
@@ -1249,11 +1224,6 @@ void UBakeRenderCaptureTool::UpdateResult()
 
 void UBakeRenderCaptureTool::UpdateVisualization()
 {
-	if (Settings->MapPreview.IsEmpty())
-	{
-		return;
-	}
-
 	const bool bSubsurfaceMaterial = ResultSettings->SubsurfaceColorMap || ResultSettings->OpacityMap;
 	const bool bPackedMRS = ResultSettings->PackedMRSMap != nullptr;
 
@@ -1265,67 +1235,39 @@ void UBakeRenderCaptureTool::UpdateVisualization()
 		ensure(Material->GetShadingModels().HasShadingModel(EMaterialShadingModel::MSM_Subsurface));
 		ensure(Material->GetBlendMode() == EBlendMode::BLEND_Masked);
 	}
-
-	if (VisualizationProps->bPreviewAsMaterial)
+	
+	const auto TrySetTexture =
+		[Material, this](const FString& TextureName, TObjectPtr<UTexture2D> Texture, TObjectPtr<UTexture2D> Fallback, bool bMaterialHasTexture)
 	{
-		const auto TrySetTexture =
-			[Material](const FString& TextureName, TObjectPtr<UTexture2D> Texture, TObjectPtr<UTexture2D> Fallback, bool bMaterialHasTexture)
+		if (bMaterialHasTexture)
 		{
-			if (bMaterialHasTexture)
+			if (VisualizationProps->bPreviewAsMaterial)
 			{
 				Material->SetTextureParameterValue(FName(TextureName), Texture ? Texture : Fallback);
 			}
-		};
-
-		// Set all computed textures or fallback to the empty texture map
-		TrySetTexture(BaseColorTexParamName, ResultSettings->BaseColorMap, EmptyColorMapWhite, true);
-		TrySetTexture(EmissiveTexParamName,  ResultSettings->EmissiveMap,  EmptyEmissiveMap,   true);
-		TrySetTexture(NormalTexParamName,    ResultSettings->NormalMap,    EmptyNormalMap,     true);
-		TrySetTexture(PackedMRSTexParamName, ResultSettings->PackedMRSMap, EmptyPackedMRSMap,  bPackedMRS);
-		TrySetTexture(RoughnessTexParamName, ResultSettings->RoughnessMap, EmptyRoughnessMap, !bPackedMRS);
-		TrySetTexture(MetallicTexParamName,  ResultSettings->MetallicMap,  EmptyMetallicMap,  !bPackedMRS);
-		TrySetTexture(SpecularTexParamName,  ResultSettings->SpecularMap,  EmptySpecularMap,  !bPackedMRS);
-		TrySetTexture(OpacityTexParamName,   ResultSettings->OpacityMap,   EmptyOpacityMap,    bSubsurfaceMaterial);
-		TrySetTexture(SubsurfaceColorTexParamName, ResultSettings->SubsurfaceColorMap, EmptySubsurfaceColorMap, bSubsurfaceMaterial);
-	}
-	else
-	{
-		const auto TrySetTexture =
-			[Material, this](const FString& TextureName, TObjectPtr<UTexture2D> Texture, TObjectPtr<UTexture2D> Fallback, bool bMaterialHasTexture)
-		{
-			// Set the BaseColor texture to the MapPreview texture if it exists and use white otherwise
-			if (TextureName == Settings->MapPreview)
+			else
 			{
-				if (bMaterialHasTexture && Texture)
+				if (TextureName == Settings->MapPreview)
 				{
-					Material->SetTextureParameterValue(FName(BaseColorTexParamName), Texture);
+					Material->SetTextureParameterValue(FName(BaseColorTexParamName), Texture ? Texture : Fallback);
 				}
 				else
-				{
-					Material->SetTextureParameterValue(FName(BaseColorTexParamName), EmptyColorMapWhite);
-				}
-			}
-
-			// Set the non-BaseColor texture parameters to empty fallback textures
-			if (TextureName != BaseColorTexParamName)
-			{
-				if (bMaterialHasTexture)
 				{
 					Material->SetTextureParameterValue(FName(TextureName), Fallback);
 				}
 			}
-		};
+		}
+	};
 
-		TrySetTexture(BaseColorTexParamName, ResultSettings->BaseColorMap, EmptyColorMapWhite, true);
-		TrySetTexture(EmissiveTexParamName,  ResultSettings->EmissiveMap,  EmptyEmissiveMap,   true);
-		TrySetTexture(NormalTexParamName,    ResultSettings->NormalMap,    EmptyNormalMap,     true);
-		TrySetTexture(PackedMRSTexParamName, ResultSettings->PackedMRSMap, EmptyPackedMRSMap,  bPackedMRS);
-		TrySetTexture(RoughnessTexParamName, ResultSettings->RoughnessMap, EmptyRoughnessMap, !bPackedMRS);
-		TrySetTexture(MetallicTexParamName,  ResultSettings->MetallicMap,  EmptyMetallicMap,  !bPackedMRS);
-		TrySetTexture(SpecularTexParamName,  ResultSettings->SpecularMap,  EmptySpecularMap,  !bPackedMRS);
-		TrySetTexture(OpacityTexParamName,   ResultSettings->OpacityMap,   EmptyOpacityMap,    bSubsurfaceMaterial);
-		TrySetTexture(SubsurfaceColorTexParamName, ResultSettings->SubsurfaceColorMap, EmptySubsurfaceColorMap, bSubsurfaceMaterial);
-	}
+	TrySetTexture(BaseColorTexParamName, ResultSettings->BaseColorMap, EmptyColorMapWhite, true);
+	TrySetTexture(EmissiveTexParamName,  ResultSettings->EmissiveMap,  EmptyEmissiveMap,   true);
+	TrySetTexture(NormalTexParamName,    ResultSettings->NormalMap,    EmptyNormalMap,     true);
+	TrySetTexture(PackedMRSTexParamName, ResultSettings->PackedMRSMap, EmptyPackedMRSMap,  bPackedMRS);
+	TrySetTexture(RoughnessTexParamName, ResultSettings->RoughnessMap, EmptyRoughnessMap, !bPackedMRS);
+	TrySetTexture(MetallicTexParamName,  ResultSettings->MetallicMap,  EmptyMetallicMap,  !bPackedMRS);
+	TrySetTexture(SpecularTexParamName,  ResultSettings->SpecularMap,  EmptySpecularMap,  !bPackedMRS);
+	TrySetTexture(OpacityTexParamName,   ResultSettings->OpacityMap,   EmptyOpacityMap,    bSubsurfaceMaterial);
+	TrySetTexture(SubsurfaceColorTexParamName, ResultSettings->SubsurfaceColorMap, EmptySubsurfaceColorMap, bSubsurfaceMaterial);
 
 	Material->SetScalarParameterValue(TEXT("UVChannel"), InputMeshSettings->GetTargetUVLayerIndex());
 	PreviewMesh->SetOverrideRenderMaterial(Material);
@@ -1336,8 +1278,7 @@ void UBakeRenderCaptureTool::UpdateVisualization()
 
 void UBakeRenderCaptureTool::InvalidateResults(FRenderCaptureTypeFlags Invalidate)
 {
-	// Note that the bake operation, Compute, updates ResultSettings when results are available via the
-	// Compute->OnResultUpdated delegate.
+	// Note that the bake operation updates ResultSettings when results are available via the OnResultUpdated delegate
 
 	if (Invalidate.bBaseColor)
 	{
diff --git a/Engine/Plugins/Experimental/MeshModelingToolsetExp/Source/MeshModelingToolsEditorOnlyExp/Public/BakeRenderCaptureTool.h b/Engine/Plugins/Experimental/MeshModelingToolsetExp/Source/MeshModelingToolsEditorOnlyExp/Public/BakeRenderCaptureTool.h
--- a/Engine/Plugins/Experimental/MeshModelingToolsetExp/Source/MeshModelingToolsEditorOnlyExp/Public/BakeRenderCaptureTool.h
+++ b/Engine/Plugins/Experimental/MeshModelingToolsetExp/Source/MeshModelingToolsEditorOnlyExp/Public/BakeRenderCaptureTool.h
@@ -337,7 +337,6 @@ protected:
 	void UpdateResult();
 	void UpdateVisualization();
 	void InvalidateResults(UE::Geometry::FRenderCaptureTypeFlags ToInvalidate);
-	void InvalidateCompute();
 	void OnMapsUpdated(const TUniquePtr<UE::Geometry::FMeshMapBaker>& NewResult);
 
 	/**
@@ -356,9 +355,8 @@ protected:
 	 */
 	void CreateAssets(UWorld* SourceWorld);
 
-	// The baking background compute operation
-	TUniquePtr<TGenericDataBackgroundCompute<UE::Geometry::FMeshMapBaker>> Compute = nullptr;
-	EBakeOpState OpState = EBakeOpState::Evaluate;
+	TUniquePtr<TGenericDataBackgroundCompute<UE::Geometry::FMeshMapBaker>> BakeOp = nullptr;
+	EBakeOpState BakeOpState = EBakeOpState::Evaluate;
 
 	TSharedPtr<UE::Geometry::FDynamicMesh3, ESPMode::ThreadSafe> TargetMesh;
 	TSharedPtr<UE::Geometry::FDynamicMeshAABBTree3, ESPMode::ThreadSafe> TargetMeshSpatial;
diff --git a/Engine/Plugins/Experimental/MeshModelingToolsetExp/Source/MeshModelingToolsExp/Private/BakeMeshAttributeMapsToolBase.cpp b/Engine/Plugins/Experimental/MeshModelingToolsetExp/Source/MeshModelingToolsExp/Private/BakeMeshAttributeMapsToolBase.cpp
--- a/Engine/Plugins/Experimental/MeshModelingToolsetExp/Source/MeshModelingToolsExp/Private/BakeMeshAttributeMapsToolBase.cpp
+++ b/Engine/Plugins/Experimental/MeshModelingToolsetExp/Source/MeshModelingToolsExp/Private/BakeMeshAttributeMapsToolBase.cpp
@@ -96,24 +96,19 @@ void UBakeMeshAttributeMapsToolBase::OnTick(float DeltaTime)
 	if (Compute)
 	{
 		Compute->Tick(DeltaTime);
+	}
 
-		if (static_cast<bool>(OpState & EBakeOpState::Invalid))
-		{
-			PreviewMesh->SetOverrideRenderMaterial(ErrorPreviewMaterial);
-		}
-		else
+	if (static_cast<bool>(OpState & EBakeOpState::Invalid))
+	{
+		PreviewMesh->SetOverrideRenderMaterial(ErrorPreviewMaterial);
+	}
+	else
+	{
+		if (!CanAccept() && Compute && Compute->GetElapsedComputeTime() > SecondsBeforeWorkingMaterial)
 		{
-			const float ElapsedComputeTime = Compute->GetElapsedComputeTime();
-			if (!CanAccept() && ElapsedComputeTime > SecondsBeforeWorkingMaterial)
-			{
-				PreviewMesh->SetOverrideRenderMaterial(WorkingPreviewMaterial);
-			}
+			PreviewMesh->SetOverrideRenderMaterial(WorkingPreviewMaterial);
 		}
 	}
-	else if (static_cast<bool>(OpState & EBakeOpState::Invalid))
-	{
-		PreviewMesh->SetOverrideRenderMaterial(ErrorPreviewMaterial);
-	} 
 }
 
 
diff --git a/Engine/Plugins/Experimental/MeshModelingToolsetExp/Source/MeshModelingToolsExp/Private/BakeMeshAttributeVertexTool.cpp b/Engine/Plugins/Experimental/MeshModelingToolsetExp/Source/MeshModelingToolsExp/Private/BakeMeshAttributeVertexTool.cpp
--- a/Engine/Plugins/Experimental/MeshModelingToolsetExp/Source/MeshModelingToolsExp/Private/BakeMeshAttributeVertexTool.cpp
+++ b/Engine/Plugins/Experimental/MeshModelingToolsetExp/Source/MeshModelingToolsExp/Private/BakeMeshAttributeVertexTool.cpp
@@ -390,59 +390,42 @@ void UBakeMeshAttributeVertexTool::OnShutdown(EToolShutdownType ShutdownType)
 	UE::ToolTarget::ShowSourceObject(Targets[0]);
 	SetSourceObjectVisible(true);
 
-	if (Compute)
-	{
-		Compute->Shutdown();
+	Compute->Shutdown();
+
+	if (ShutdownType == EToolShutdownType::Accept)
+	{
+		GetToolManager()->BeginUndoTransaction(LOCTEXT("BakeMeshAttributeVertexToolTransactionName",
+													   "Bake Mesh Attribute Vertex"));
+		FConversionToMeshDescriptionOptions ConvertOptions;
+		ConvertOptions.SetToVertexColorsOnly();
+		ConvertOptions.bTransformVtxColorsSRGBToLinear = true;
+		UE::ToolTarget::CommitDynamicMeshUpdate(
+			Targets[0],
+			*PreviewMesh->GetMesh(),
+			false, // bHaveModifiedTopology
+			ConvertOptions);
+		GetToolManager()->EndUndoTransaction();
 	}
 
-	if (PreviewMesh)
-	{
-		if (ShutdownType == EToolShutdownType::Accept)
-		{
-			GetToolManager()->BeginUndoTransaction(LOCTEXT("BakeMeshAttributeVertexToolTransactionName",
-			                                               "Bake Mesh Attribute Vertex"));
-			FConversionToMeshDescriptionOptions ConvertOptions;
-			ConvertOptions.SetToVertexColorsOnly();
-			ConvertOptions.bTransformVtxColorsSRGBToLinear = true;
-			UE::ToolTarget::CommitDynamicMeshUpdate(
-				Targets[0],
-				*PreviewMesh->GetMesh(),
-				false, // bHaveModifiedTopology
-				ConvertOptions);
-			GetToolManager()->EndUndoTransaction();
-		}
-
-		PreviewMesh->SetVisible(false);
-		PreviewMesh->Disconnect();
-		PreviewMesh = nullptr;
-	}
+	PreviewMesh->SetVisible(false);
+	PreviewMesh->Disconnect();
+	PreviewMesh = nullptr;
 
 	RecordAnalytics(BakeAnalytics, TEXT("BakeVertex"));
 }
 
 void UBakeMeshAttributeVertexTool::OnTick(float DeltaTime)
 {
-	if (Compute)
-	{
-		Compute->Tick(DeltaTime);
+	Compute->Tick(DeltaTime);
 
-		if (static_cast<bool>(OpState & EBakeOpState::Invalid))
-		{
-			PreviewMesh->SetOverrideRenderMaterial(ErrorPreviewMaterial);
-		}
-		else
-		{
-			const float ElapsedComputeTime = Compute->GetElapsedComputeTime();
-			if (!CanAccept() && ElapsedComputeTime > SecondsBeforeWorkingMaterial)
-			{
-				PreviewMesh->SetOverrideRenderMaterial(WorkingPreviewMaterial);
-			}
-		}
-	}
-	else if (static_cast<bool>(OpState & EBakeOpState::Invalid))
+	if (static_cast<bool>(OpState & EBakeOpState::Invalid))
 	{
 		PreviewMesh->SetOverrideRenderMaterial(ErrorPreviewMaterial);
 	}
+	else if (!CanAccept() && Compute->GetElapsedComputeTime() > SecondsBeforeWorkingMaterial)
+	{
+		PreviewMesh->SetOverrideRenderMaterial(WorkingPreviewMaterial);
+	}
 }
 
 void UBakeMeshAttributeVertexTool::Render(IToolsContextRenderAPI* RenderAPI)
@@ -453,7 +436,7 @@ void UBakeMeshAttributeVertexTool::Render(IToolsContextRenderAPI* RenderAPI)
 bool UBakeMeshAttributeVertexTool::CanAccept() const
 {
 	const bool bValidOp = (OpState & EBakeOpState::Invalid) != EBakeOpState::Invalid;
-	return Compute ? bValidOp && Compute->HaveValidResult() : false;
+	return bValidOp && Compute->HaveValidResult();
 }
 
 TUniquePtr<UE::Geometry::TGenericDataOperator<FMeshVertexBaker>> UBakeMeshAttributeVertexTool::MakeNewOperator()
