[Atom - GeometryProcessing] Adding a FilterLODLevel option to FBaseMeshInstance for completely culling meshes past a certain LOD level.

The FOptions::FilterDecorativePartsLODLevel is too broad and doesn't allow to only remove certain instances of meshes on a specific LOD i.e:. Some meshes can be optimized by replacing multiple mesh instances on lower LOD with a single custom mesh (Atom fake underside optimization).

diff --git a/Engine/Plugins/Experimental/MeshModelingToolsetExp/Source/GeometryProcessingAdapters/Private/CombineMeshInstancesImpl.cpp b/Engine/Plugins/Experimental/MeshModelingToolsetExp/Source/GeometryProcessingAdapters/Private/CombineMeshInstancesImpl.cpp
--- a/Engine/Plugins/Experimental/MeshModelingToolsetExp/Source/GeometryProcessingAdapters/Private/CombineMeshInstancesImpl.cpp
+++ b/Engine/Plugins/Experimental/MeshModelingToolsetExp/Source/GeometryProcessingAdapters/Private/CombineMeshInstancesImpl.cpp
@@ -107,6 +107,8 @@ struct FMeshPartInstance
 
 	EMeshDetailLevel DetailLevel = EMeshDetailLevel::Standard;
 
+	int32 FilterLODLevel = -1;
+
 	bool bAllowApproximation = true;
 
 	// allow FMeshPartInstance to maintain link to external representation of instance
@@ -239,6 +241,7 @@ void InitializeMeshPartAssembly(
 		NewInstance.SourceComponent = SourceMeshInstance.SourceComponent;
 		NewInstance.SourceInstanceIndex = SourceMeshInstance.SourceInstanceIndex;
 		NewInstance.DetailLevel = static_cast<EMeshDetailLevel>( static_cast<int32>(SourceMeshInstance.DetailLevel) );
+		NewInstance.FilterLODLevel = SourceMeshInstance.FilterLODLevel;
 		NewInstance.bAllowApproximation = SourceMeshInstance.bAllowApproximation;
 		for ( FTransform3d Transform : SourceMeshInstance.TransformSequence )
 		{
@@ -297,6 +300,7 @@ void InitializeMeshPartAssembly(
 		NewInstance.SourceComponent = nullptr;
 		NewInstance.SourceInstanceIndex = 0;
 		NewInstance.DetailLevel = static_cast<EMeshDetailLevel>(static_cast<int32>(SourceMeshInstance.DetailLevel));
+		NewInstance.FilterLODLevel = SourceMeshInstance.FilterLODLevel;
 		NewInstance.bAllowApproximation = SourceMeshInstance.bAllowApproximation;
 		for (FTransform3d Transform : SourceMeshInstance.TransformSequence)
 		{
@@ -3498,6 +3502,11 @@ void BuildCombinedMesh(
 			{
 				const FDynamicMesh3* InstanceAppendMesh = UseAppendMesh;
 
+				if (Instance.FilterLODLevel >= 0 && LODLevel >= Instance.FilterLODLevel)
+				{
+					continue;
+				}
+
 				bool bIsDecorativePart = (Instance.DetailLevel == EMeshDetailLevel::Decorative);
 				if (bIsDecorativePart)
 				{
diff --git a/Engine/Source/Developer/GeometryProcessingInterfaces/Public/GeometryProcessingInterfaces/CombineMeshInstances.h b/Engine/Source/Developer/GeometryProcessingInterfaces/Public/GeometryProcessingInterfaces/CombineMeshInstances.h
--- a/Engine/Source/Developer/GeometryProcessingInterfaces/Public/GeometryProcessingInterfaces/CombineMeshInstances.h
+++ b/Engine/Source/Developer/GeometryProcessingInterfaces/Public/GeometryProcessingInterfaces/CombineMeshInstances.h
@@ -79,6 +79,8 @@ public:
 
 		bool bAllowApproximation = true;			// if false, only Copied or Simplified LODs will be used for this part Instance. Will be combined w/ the GroupData flag.
 
+		int32 FilterLODLevel = -1;					// LOD level to filter out this mesh, value -1 disable this option. 
+
 		// in some cases it may be desirable to have "groups" of instances which should be output as separate meshes, but
 		// be jointly processed in terms of (eg) the part LODs. If any InstanceSubsetID is non-zero, then instance subsets
 		// are grouped/extracted by integer ID and will be returned as separate FOutputMesh's in the FResults. 
