Connect up weld edges setting in Mesh to Collision tool

diff --git a/Engine/Plugins/Experimental/MeshModelingToolsetExp/Source/MeshModelingToolsExp/Private/Physics/ExtractCollisionGeometryTool.cpp b/Engine/Plugins/Experimental/MeshModelingToolsetExp/Source/MeshModelingToolsExp/Private/Physics/ExtractCollisionGeometryTool.cpp
--- a/Engine/Plugins/Experimental/MeshModelingToolsetExp/Source/MeshModelingToolsExp/Private/Physics/ExtractCollisionGeometryTool.cpp
+++ b/Engine/Plugins/Experimental/MeshModelingToolsetExp/Source/MeshModelingToolsExp/Private/Physics/ExtractCollisionGeometryTool.cpp
@@ -263,7 +263,7 @@ void UExtractCollisionGeometryTool::RecalculateMesh_Complex()
 	if (CollisionProvider)
 	{
 		FTransformSequence3d Transform;
-		UE::Geometry::ConvertComplexCollisionToMeshes(CollisionProvider, CurrentMesh, FTransformSequence3d(), bMeshErrors, true, true);
+		UE::Geometry::ConvertComplexCollisionToMeshes(CollisionProvider, CurrentMesh, FTransformSequence3d(), bMeshErrors, Settings->bWeldEdges, true);
 	}
 
 	PreviewMesh->UpdatePreview(&CurrentMesh);
diff --git a/Engine/Plugins/Experimental/MeshModelingToolsetExp/Source/MeshModelingToolsExp/Public/Physics/ExtractCollisionGeometryTool.h b/Engine/Plugins/Experimental/MeshModelingToolsetExp/Source/MeshModelingToolsExp/Public/Physics/ExtractCollisionGeometryTool.h
--- a/Engine/Plugins/Experimental/MeshModelingToolsetExp/Source/MeshModelingToolsExp/Public/Physics/ExtractCollisionGeometryTool.h
+++ b/Engine/Plugins/Experimental/MeshModelingToolsetExp/Source/MeshModelingToolsExp/Public/Physics/ExtractCollisionGeometryTool.h
@@ -44,10 +44,6 @@ public:
 	UPROPERTY(EditAnywhere, Category = Options)
 	EExtractCollisionOutputType CollisionType = EExtractCollisionOutputType::Simple;
 
-	/** Whether or not to weld coincident border edges of the Complex Collision Mesh (if possible) */
-	UPROPERTY(EditAnywhere, Category = Options, Meta = (EditCondition = "CollisionType == EExtractCollisionOutputType::Complex"))
-	bool bWeldEdges = true;
-
 	/** Whether or not to generate a seperate Mesh Object for each Simple Collision Shape  */
 	UPROPERTY(EditAnywhere, Category = Options, Meta = (EditCondition = "CollisionType == EExtractCollisionOutputType::Simple"))
 	bool bOutputSeparateMeshes = true;
@@ -59,6 +55,10 @@ public:
 	/** Show/Hide input mesh */
 	UPROPERTY(EditAnywhere, Category = Options)
 	bool bShowInputMesh = true;
+
+	/** Whether or not to weld coincident border edges of the Complex Collision Mesh (if possible) */
+	UPROPERTY(EditAnywhere, AdvancedDisplay, Category = Options, Meta = (EditCondition = "CollisionType == EExtractCollisionOutputType::Complex"))
+	bool bWeldEdges = true;
 };
 
 
