In the collision to mesh tool, make the simple collision visualization and related settings only active when tool is converting simple collisions

diff --git a/Engine/Plugins/Experimental/MeshModelingToolsetExp/Source/MeshModelingToolsExp/Private/Physics/ExtractCollisionGeometryTool.cpp b/Engine/Plugins/Experimental/MeshModelingToolsetExp/Source/MeshModelingToolsExp/Private/Physics/ExtractCollisionGeometryTool.cpp
--- a/Engine/Plugins/Experimental/MeshModelingToolsetExp/Source/MeshModelingToolsExp/Private/Physics/ExtractCollisionGeometryTool.cpp
+++ b/Engine/Plugins/Experimental/MeshModelingToolsetExp/Source/MeshModelingToolsExp/Private/Physics/ExtractCollisionGeometryTool.cpp
@@ -82,7 +82,6 @@ void UExtractCollisionGeometryTool::Setup()
 	Settings = NewObject<UExtractCollisionToolProperties>(this);
 	Settings->RestoreProperties(this);
 	AddToolPropertySource(Settings);
-	Settings->WatchProperty(Settings->CollisionType, [this](EExtractCollisionOutputType NewValue) { bResultValid = false; });
 	Settings->WatchProperty(Settings->bWeldEdges, [this](bool bNewValue) { bResultValid = false; });
 	Settings->WatchProperty(Settings->bShowPreview, [this](bool bNewValue) { PreviewMesh->SetVisible(bNewValue); });
 	PreviewMesh->SetVisible(Settings->bShowPreview);
@@ -93,7 +92,20 @@ void UExtractCollisionGeometryTool::Setup()
 	VizSettings->RestoreProperties(this);
 	AddToolPropertySource(VizSettings);
 	VizSettings->Initialize(this);
-	VizSettings->bEnableShowCollision = false; // This tool always shows collision geometry
+
+	// Enable simple collision visualization and related settings only when extracting simple collision
+	VizSettings->bEnableShowCollision = false;
+	Settings->WatchProperty(Settings->CollisionType, [this](EExtractCollisionOutputType NewValue)
+	{
+		bResultValid = false;
+		SetToolPropertySourceEnabled(VizSettings, NewValue == EExtractCollisionOutputType::Simple);
+		VizSettings->bShowCollision = NewValue == EExtractCollisionOutputType::Simple;
+		VizSettings->bVisualizationDirty = true;
+		NotifyOfPropertyChangeByTool(VizSettings);
+	});
+	SetToolPropertySourceEnabled(VizSettings, Settings->CollisionType == EExtractCollisionOutputType::Simple);
+	VizSettings->bShowCollision = Settings->CollisionType == EExtractCollisionOutputType::Simple;
+	NotifyOfPropertyChangeByTool(VizSettings);
 	
 
 	UBodySetup* BodySetup = UE::ToolTarget::GetPhysicsBodySetup(Target);
