Add option for Mesh to Collision tool to hide target mesh, so user has an easier way to see close-fitting solid collision meshes

diff --git a/Engine/Plugins/Experimental/MeshModelingToolsetExp/Source/MeshModelingToolsExp/Private/Physics/SetCollisionGeometryTool.cpp b/Engine/Plugins/Experimental/MeshModelingToolsetExp/Source/MeshModelingToolsExp/Private/Physics/SetCollisionGeometryTool.cpp
--- a/Engine/Plugins/Experimental/MeshModelingToolsetExp/Source/MeshModelingToolsExp/Private/Physics/SetCollisionGeometryTool.cpp
+++ b/Engine/Plugins/Experimental/MeshModelingToolsetExp/Source/MeshModelingToolsExp/Private/Physics/SetCollisionGeometryTool.cpp
@@ -314,6 +314,11 @@ void USetCollisionGeometryTool::Setup()
 	Settings->WatchProperty(Settings->HullTolerance, [this](float) { InvalidateCompute(); });
 	Settings->WatchProperty(Settings->SweepAxis, [this](EProjectedHullAxis) { InvalidateCompute(); });
 	Settings->WatchProperty(Settings->LevelSetResolution, [this](int32) { InvalidateCompute(); });
+	Settings->WatchProperty(Settings->bShowTargetMesh, [this](bool bNewValue) 
+	{
+		UE::ToolTarget::SetSourceObjectVisible(Targets.Last(), bNewValue);
+	});
+	UE::ToolTarget::SetSourceObjectVisible(Targets.Last(), Settings->bShowTargetMesh);
 
 	if (InitialSourceMeshes.Num() == 1)
 	{
@@ -476,6 +481,10 @@ void USetCollisionGeometryTool::OnShutdown(EToolShutdownType ShutdownType)
 			UE::ToolTarget::ShowSourceObject(Targets[k]);
 		}
 	}
+	if (!Settings->bShowTargetMesh)
+	{
+		UE::ToolTarget::ShowSourceObject(Targets.Last());
+	}
 
 	if (Compute)
 	{
diff --git a/Engine/Plugins/Experimental/MeshModelingToolsetExp/Source/MeshModelingToolsExp/Public/Physics/SetCollisionGeometryTool.h b/Engine/Plugins/Experimental/MeshModelingToolsetExp/Source/MeshModelingToolsExp/Public/Physics/SetCollisionGeometryTool.h
--- a/Engine/Plugins/Experimental/MeshModelingToolsetExp/Source/MeshModelingToolsExp/Public/Physics/SetCollisionGeometryTool.h
+++ b/Engine/Plugins/Experimental/MeshModelingToolsetExp/Source/MeshModelingToolsExp/Public/Physics/SetCollisionGeometryTool.h
@@ -206,6 +206,10 @@ public:
 	UPROPERTY(EditAnywhere, Category = OutputOptions)
 	ECollisionGeometryMode SetCollisionType = ECollisionGeometryMode::SimpleAndComplex;
 
+	/** Show/Hide target mesh */
+	UPROPERTY(EditAnywhere, Category = TargetVisualization)
+	bool bShowTargetMesh = true;
+
 	// Set by the tool to tell the settings object whether the tool is using multiple inputs.
 	UPROPERTY()
 	bool bUsingMultipleInputs = false;
