Skin Paint Tool: fixed crash when destroying the vertex weight editor

diff --git a/Engine/Plugins/Experimental/MeshModelingToolsetExp/Source/ModelingEditorUI/Private/SkinWeightDetailCustomization.cpp b/Engine/Plugins/Experimental/MeshModelingToolsetExp/Source/ModelingEditorUI/Private/SkinWeightDetailCustomization.cpp
--- a/Engine/Plugins/Experimental/MeshModelingToolsetExp/Source/ModelingEditorUI/Private/SkinWeightDetailCustomization.cpp
+++ b/Engine/Plugins/Experimental/MeshModelingToolsetExp/Source/ModelingEditorUI/Private/SkinWeightDetailCustomization.cpp
@@ -755,12 +755,20 @@ TSharedRef<SWidget> SVertexWeightItem::GenerateWidgetForColumn(const FName& Colu
 	return SNullWidget::NullWidget;
 }
 
+SVertexWeightEditor::~SVertexWeightEditor()
+{
+	if (Tool.IsValid())
+	{
+		Tool->OnSelectionChanged.RemoveAll(this);
+		Tool->OnWeightsChanged.RemoveAll(this);
+		Tool.Reset();
+	}
+}
+
 void SVertexWeightEditor::Construct(const FArguments& InArgs, USkinWeightsPaintTool* InSkinTool)
 {
 	Tool = InSkinTool;
-	Tool->OnSelectionChanged.AddLambda([this](){ RefreshView(); });
-	Tool->OnWeightsChanged.AddLambda([this](){ RefreshView(); });
-	
+
 	ChildSlot
 	[
 		SNew(SBox)
@@ -782,6 +790,9 @@ void SVertexWeightEditor::Construct(const FArguments& InArgs, USkinWeightsPaintT
 	];
 
 	RefreshView();
+	
+	Tool->OnSelectionChanged.AddSP(this, &SVertexWeightEditor::RefreshView);
+	Tool->OnWeightsChanged.AddSP(this, &SVertexWeightEditor::RefreshView);
 }
 
 void SVertexWeightEditor::RefreshView()
diff --git a/Engine/Plugins/Experimental/MeshModelingToolsetExp/Source/ModelingEditorUI/Public/SkinWeightDetailCustomization.h b/Engine/Plugins/Experimental/MeshModelingToolsetExp/Source/ModelingEditorUI/Public/SkinWeightDetailCustomization.h
--- a/Engine/Plugins/Experimental/MeshModelingToolsetExp/Source/ModelingEditorUI/Public/SkinWeightDetailCustomization.h
+++ b/Engine/Plugins/Experimental/MeshModelingToolsetExp/Source/ModelingEditorUI/Public/SkinWeightDetailCustomization.h
@@ -75,6 +75,8 @@ public:
 	SLATE_END_ARGS()
 
 	void Construct(const FArguments& InArgs, USkinWeightsPaintTool* InSkinTool);
+
+	virtual ~SVertexWeightEditor() override;
 	
 	void RefreshView();
 
