fix missing scale in line set component visualization of convex hulls

diff --git a/Engine/Plugins/Experimental/MeshModelingToolsetExp/Source/MeshModelingToolsExp/Private/Physics/CollisionGeometryVisualization.cpp b/Engine/Plugins/Experimental/MeshModelingToolsetExp/Source/MeshModelingToolsExp/Private/Physics/CollisionGeometryVisualization.cpp
--- a/Engine/Plugins/Experimental/MeshModelingToolsetExp/Source/MeshModelingToolsExp/Private/Physics/CollisionGeometryVisualization.cpp
+++ b/Engine/Plugins/Experimental/MeshModelingToolsetExp/Source/MeshModelingToolsExp/Private/Physics/CollisionGeometryVisualization.cpp
@@ -211,7 +211,7 @@ void InitializePreviewGeometryLines(
 			const FKConvexElem& Convex = AggGeom.ConvexElems[Index];
 			FTransform ElemTransform = Convex.GetTransform();
 			ElemTransform.ScaleTranslation(PhysicsData.ExternalScale3D);
-			ElemTransform.SetScale3D(PhysicsData.ExternalScale3D);
+			ElemTransform.MultiplyScale3D(PhysicsData.ExternalScale3D);
 			int32 NumTriangles = Convex.IndexData.Num() / 3;
 			for (int32 k = 0; k < NumTriangles; ++k)
 			{
