minor modeling mode tooltip cleanup -- fix a typo in edit polygroups and add requested tooltip in vox offset

diff --git a/Engine/Plugins/Experimental/MeshModelingToolsetExp/Source/MeshModelingToolsExp/Public/VoxelMorphologyMeshesTool.h b/Engine/Plugins/Experimental/MeshModelingToolsetExp/Source/MeshModelingToolsExp/Public/VoxelMorphologyMeshesTool.h
--- a/Engine/Plugins/Experimental/MeshModelingToolsetExp/Source/MeshModelingToolsExp/Public/VoxelMorphologyMeshesTool.h
+++ b/Engine/Plugins/Experimental/MeshModelingToolsetExp/Source/MeshModelingToolsExp/Public/VoxelMorphologyMeshesTool.h
@@ -21,10 +21,12 @@ class MESHMODELINGTOOLSEXP_API UVoxelMorphologyMeshesToolProperties : public UIn
 
 public:
 
-	UPROPERTY(EditAnywhere, Category = Morphology)
+	/** Which offset (or morphology) operation to apply */
+	UPROPERTY(EditAnywhere, Category = Offset)
 	EMorphologyOperation Operation = EMorphologyOperation::Dilate;
 
-	UPROPERTY(EditAnywhere, Category = Morphology, meta = (UIMin = ".1", UIMax = "100", ClampMin = ".001", ClampMax = "1000"))
+	/** Distance to offset the mesh */
+	UPROPERTY(EditAnywhere, Category = Offset, meta = (UIMin = ".1", UIMax = "100", ClampMin = ".001", ClampMax = "1000"))
 	double Distance = 5;
 
 	/** Apply a "VoxWrap" operation to input mesh(es) before computing the morphology.  Fixes results for inputs with holes and/or self-intersections. */
