make the modeling mode physics inspector trigger updates when collisions are updated by external actions

diff --git a/Engine/Plugins/Experimental/MeshModelingToolsetExp/Source/MeshModelingToolsExp/Private/Physics/PhysicsInspectorTool.cpp b/Engine/Plugins/Experimental/MeshModelingToolsetExp/Source/MeshModelingToolsExp/Private/Physics/PhysicsInspectorTool.cpp
--- a/Engine/Plugins/Experimental/MeshModelingToolsetExp/Source/MeshModelingToolsExp/Private/Physics/PhysicsInspectorTool.cpp
+++ b/Engine/Plugins/Experimental/MeshModelingToolsetExp/Source/MeshModelingToolsExp/Private/Physics/PhysicsInspectorTool.cpp
@@ -51,49 +51,88 @@ void UPhysicsInspectorTool::Setup()
 	AddToolPropertySource(VizSettings);
 	VizSettings->bEnableShowCollision = false; // This tool always shows collision geometry
 
-	// Intitialize the collision geometry visualization
+	// Initialize the collision geometry visualization
 	{
 		VizSettings->Initialize(this);
-		int32 FirstLineSetIndex = 0;
-		for (int32 ComponentIdx = 0; ComponentIdx < Targets.Num(); ComponentIdx++)
-		{
-			UBodySetup* BodySetup = UE::ToolTarget::GetPhysicsBodySetup(Targets[ComponentIdx]);
-			if (BodySetup)
-			{
-				TSharedPtr<FPhysicsDataCollection> PhysicsData = MakeShared<FPhysicsDataCollection>();
-				PhysicsData->InitializeFromComponent( UE::ToolTarget::GetTargetComponent(Targets[ComponentIdx]), true);
-
-				PhysicsInfos.Add(PhysicsData);
-
-				UPreviewGeometry* PreviewGeom = NewObject<UPreviewGeometry>(this);
-				FTransform TargetTransform = (FTransform) UE::ToolTarget::GetLocalToWorldTransform(Targets[ComponentIdx]);
-				PhysicsData->ExternalScale3D = TargetTransform.GetScale3D();
-				TargetTransform.SetScale3D(FVector::OneVector);
-				PreviewGeom->CreateInWorld(UE::ToolTarget::GetTargetActor(Targets[ComponentIdx])->GetWorld(), TargetTransform);
-				PreviewElements.Add(PreviewGeom);
-
-				UE::PhysicsTools::PartiallyInitializeCollisionGeometryVisualization(PreviewGeom, VizSettings, *PhysicsData, FirstLineSetIndex);
-				FirstLineSetIndex += PreviewGeom->LineSets.Num();
-
-				UPhysicsObjectToolPropertySet* ObjectProps = NewObject<UPhysicsObjectToolPropertySet>(this);
-				UE::PhysicsTools::InitializePhysicsToolObjectPropertySet(PhysicsData.Get(), ObjectProps);
-				AddToolPropertySource(ObjectProps);
-			}
-		}
-		VizSettings->bVisualizationDirty = false;
+		InitializePreviewGeometry(false /*bClearExisting*/);
 	}
 
+	OnCreatePhysicsDelegateHandle = UActorComponent::GlobalCreatePhysicsDelegate.AddUObject(this, &UPhysicsInspectorTool::OnCreatePhysics);
+
 	SetToolDisplayName(LOCTEXT("ToolName", "Physics Inspector"));
 	GetToolManager()->DisplayMessage(
 		LOCTEXT("OnStartTool", "Inspect Physics data for the selected Objects"),
 		EToolMessageLevel::UserNotification);
 }
 
+void UPhysicsInspectorTool::InitializePreviewGeometry(bool bClearExisting)
+{
+	if (bClearExisting)
+	{
+		for (UPreviewGeometry* Preview : PreviewElements)
+		{
+			Preview->Disconnect();
+		}
+		PreviewElements.Reset();
+		for (int32 ObjectDataIdx = 0; ObjectDataIdx < ObjectData.Num(); ++ObjectDataIdx)
+		{
+			RemoveToolPropertySource(ObjectData[ObjectDataIdx]);
+		}
+		ObjectData.Reset();
+	}
+
+	int32 FirstLineSetIndex = 0;
+	for (int32 ComponentIdx = 0; ComponentIdx < Targets.Num(); ComponentIdx++)
+	{
+		const UBodySetup* BodySetup = UE::ToolTarget::GetPhysicsBodySetup(Targets[ComponentIdx]);
+		if (BodySetup)
+		{
+			FPhysicsDataCollection PhysicsData;
+			PhysicsData.InitializeFromComponent(UE::ToolTarget::GetTargetComponent(Targets[ComponentIdx]), true);
+
+			UPreviewGeometry* PreviewGeom = NewObject<UPreviewGeometry>(this);
+			FTransform TargetTransform = (FTransform)UE::ToolTarget::GetLocalToWorldTransform(Targets[ComponentIdx]);
+			PhysicsData.ExternalScale3D = TargetTransform.GetScale3D();
+			TargetTransform.SetScale3D(FVector::OneVector);
+			PreviewGeom->CreateInWorld(UE::ToolTarget::GetTargetActor(Targets[ComponentIdx])->GetWorld(), TargetTransform);
+			PreviewElements.Add(PreviewGeom);
+
+			UE::PhysicsTools::PartiallyInitializeCollisionGeometryVisualization(PreviewGeom, VizSettings, PhysicsData, FirstLineSetIndex);
+			FirstLineSetIndex += PreviewGeom->LineSets.Num();
+
+			UPhysicsObjectToolPropertySet* ObjectProps = NewObject<UPhysicsObjectToolPropertySet>(this);
+			UE::PhysicsTools::InitializePhysicsToolObjectPropertySet(&PhysicsData, ObjectProps);
+			AddToolPropertySource(ObjectProps);
+			ObjectData.Add(ObjectProps);
+		}
+	}
+	VizSettings->bVisualizationDirty = false;
+	bUnderlyingPhysicsObjectsUpdated = false;
+}
+
+void UPhysicsInspectorTool::OnCreatePhysics(UActorComponent* Component)
+{
+	if (bUnderlyingPhysicsObjectsUpdated) // if flag already set, do not need to re-check
+	{
+		return;
+	}
+	for (int32 ComponentIdx = 0; ComponentIdx < Targets.Num(); ++ComponentIdx)
+	{
+		if (UE::ToolTarget::GetTargetComponent(Targets[ComponentIdx]) == Component)
+		{
+			bUnderlyingPhysicsObjectsUpdated = true;
+			return;
+		}
+	}
+}
+
 
 void UPhysicsInspectorTool::OnShutdown(EToolShutdownType ShutdownType)
 {
 	VizSettings->SaveProperties(this);
 
+	UActorComponent::GlobalCreatePhysicsDelegate.Remove(OnCreatePhysicsDelegateHandle);
+
 	for (UPreviewGeometry* Preview : PreviewElements)
 	{
 		Preview->Disconnect();
@@ -106,6 +145,11 @@ void UPhysicsInspectorTool::OnShutdown(EToolShutdownType ShutdownType)
 
 void UPhysicsInspectorTool::OnTick(float DeltaTime)
 {
+	if (bUnderlyingPhysicsObjectsUpdated)
+	{
+		InitializePreviewGeometry(true /*bClearExisting*/);
+	}
+
 	if (VizSettings->bVisualizationDirty)
 	{
 		int32 FirstLineSetIndex= 0;
diff --git a/Engine/Plugins/Experimental/MeshModelingToolsetExp/Source/MeshModelingToolsExp/Public/Physics/PhysicsInspectorTool.h b/Engine/Plugins/Experimental/MeshModelingToolsetExp/Source/MeshModelingToolsExp/Public/Physics/PhysicsInspectorTool.h
--- a/Engine/Plugins/Experimental/MeshModelingToolsetExp/Source/MeshModelingToolsExp/Public/Physics/PhysicsInspectorTool.h
+++ b/Engine/Plugins/Experimental/MeshModelingToolsetExp/Source/MeshModelingToolsExp/Public/Physics/PhysicsInspectorTool.h
@@ -48,6 +48,8 @@ public:
 		return true;
 	}
 
+	void OnCreatePhysics(UActorComponent* Component);
+
 protected:
 
 	UPROPERTY()
@@ -60,6 +62,11 @@ protected:
 	UPROPERTY()
 	TArray<TObjectPtr<UPreviewGeometry>> PreviewElements;
 
-	// these are TSharedPtr because TPimplPtr cannot currently be added to a TArray?
-	TArray<TSharedPtr<FPhysicsDataCollection>> PhysicsInfos;
+private:
+	// Helper to create or re-create preview geometry
+	void InitializePreviewGeometry(bool bClearExisting);
+	// Delegate to track when physics data may have been updated
+	FDelegateHandle OnCreatePhysicsDelegateHandle;
+	// A flag to track when the preview geometry needs to be re-initialized
+	bool bUnderlyingPhysicsObjectsUpdated = false;
 };
