expose convex decomposition's shape protection to the set collision geometry tool

diff --git a/Engine/Plugins/Experimental/MeshModelingToolsetExp/Source/MeshModelingToolsExp/Private/Physics/SetCollisionGeometryTool.cpp b/Engine/Plugins/Experimental/MeshModelingToolsetExp/Source/MeshModelingToolsExp/Private/Physics/SetCollisionGeometryTool.cpp
--- a/Engine/Plugins/Experimental/MeshModelingToolsetExp/Source/MeshModelingToolsExp/Private/Physics/SetCollisionGeometryTool.cpp
+++ b/Engine/Plugins/Experimental/MeshModelingToolsetExp/Source/MeshModelingToolsExp/Private/Physics/SetCollisionGeometryTool.cpp
@@ -311,6 +311,10 @@ void USetCollisionGeometryTool::Setup()
 	Settings->WatchProperty(Settings->ConvexDecompositionSearchFactor, [this](int32) { InvalidateCompute(); });
 	Settings->WatchProperty(Settings->AddHullsErrorTolerance, [this](int32) { InvalidateCompute(); });
 	Settings->WatchProperty(Settings->MinPartThickness, [this](int32) { InvalidateCompute(); });
+	Settings->WatchProperty(Settings->bUseNegativeSpaceInDecomposition, [this](int32) { InvalidateCompute(); });
+	Settings->WatchProperty(Settings->NegativeSpaceMinRadius, [this](int32) { InvalidateCompute(); });
+	Settings->WatchProperty(Settings->NegativeSpaceTolerance, [this](int32) { InvalidateCompute(); });
+	Settings->WatchProperty(Settings->bIgnoreInternalNegativeSpace, [this](int32) { InvalidateCompute(); });
 	Settings->WatchProperty(Settings->HullTolerance, [this](float) { InvalidateCompute(); });
 	Settings->WatchProperty(Settings->SweepAxis, [this](EProjectedHullAxis) { InvalidateCompute(); });
 	Settings->WatchProperty(Settings->LevelSetResolution, [this](int32) { InvalidateCompute(); });
@@ -427,6 +431,10 @@ TUniquePtr<UE::Geometry::TGenericDataOperator<FPhysicsDataCollection>> USetColli
 	Op->UseShapeGenerator->ConvexDecompositionSearchFactor = Settings->ConvexDecompositionSearchFactor;
 	Op->UseShapeGenerator->ConvexDecompositionErrorTolerance = Settings->AddHullsErrorTolerance;
 	Op->UseShapeGenerator->ConvexDecompositionMinPartThickness = Settings->MinPartThickness;
+	Op->UseShapeGenerator->bConvexDecompositionProtectNegativeSpace = Settings->bUseNegativeSpaceInDecomposition;
+	Op->UseShapeGenerator->NegativeSpaceMinRadius = Settings->NegativeSpaceMinRadius;
+	Op->UseShapeGenerator->NegativeSpaceTolerance = Settings->NegativeSpaceTolerance;
+	Op->UseShapeGenerator->bIgnoreInternalNegativeSpace = Settings->bIgnoreInternalNegativeSpace;
 	Op->UseShapeGenerator->HullSimplifyTolerance = Settings->HullTolerance;
 	Op->UseShapeGenerator->LevelSetGridResolution = Settings->LevelSetResolution;
 
diff --git a/Engine/Plugins/Experimental/MeshModelingToolsetExp/Source/MeshModelingToolsExp/Public/Physics/SetCollisionGeometryTool.h b/Engine/Plugins/Experimental/MeshModelingToolsetExp/Source/MeshModelingToolsExp/Public/Physics/SetCollisionGeometryTool.h
--- a/Engine/Plugins/Experimental/MeshModelingToolsetExp/Source/MeshModelingToolsExp/Public/Physics/SetCollisionGeometryTool.h
+++ b/Engine/Plugins/Experimental/MeshModelingToolsetExp/Source/MeshModelingToolsExp/Public/Physics/SetCollisionGeometryTool.h
@@ -187,6 +187,24 @@ public:
 		EditConditionHides, EditCondition = "GeometryType == ECollisionGeometryType::ConvexHulls && MaxHullsPerMesh > 1"))
 	float MinPartThickness = 0.1;
 
+	/** Whether to guide the convex decomposition to prioritize not filling negative space of the input shape */
+	UPROPERTY(EditAnywhere, Category = ConvexHulls, meta = (
+		EditConditionHides, EditCondition = "GeometryType == ECollisionGeometryType::ConvexHulls && MaxHullsPerMesh > 1"))
+	bool bUseNegativeSpaceInDecomposition = false;
+
+	/** Negative space closer to the input than this tolerance distance can be filled in */
+	UPROPERTY(EditAnywhere, Category = NegativeSpace, meta = (UIMin = ".001", UIMax = "100", ClampMin = "0",
+		EditConditionHides, EditCondition = "GeometryType == ECollisionGeometryType::ConvexHulls && MaxHullsPerMesh > 1 && bUseNegativeSpaceInDecomposition"))
+	double NegativeSpaceTolerance = 3;
+	/** Minimum radius of negative space to protect; tunnels with radius smaller than this could be filled in */
+	UPROPERTY(EditAnywhere, Category = NegativeSpace, meta = (UIMin = ".001", UIMax = "100", ClampMin = "0",
+		EditConditionHides, EditCondition = "GeometryType == ECollisionGeometryType::ConvexHulls && MaxHullsPerMesh > 1 && bUseNegativeSpaceInDecomposition"))
+	double NegativeSpaceMinRadius = 10;
+	/** Whether to ignore negative space that is not accessible by traversing from the convex hull (via paths w/ radius of at least Negative Space Tolerance) */
+	UPROPERTY(EditAnywhere, Category = NegativeSpace, meta = (
+		EditConditionHides, EditCondition = "GeometryType == ECollisionGeometryType::ConvexHulls && MaxHullsPerMesh > 1 && bUseNegativeSpaceInDecomposition"))
+	bool bIgnoreInternalNegativeSpace = true;
+
 	/** If > 0, the polygon used to generate the swept hull will be simplified up to this distance tolerance, in cm */
 	UPROPERTY(EditAnywhere, Category = SweptHulls, meta = (UIMin = "0", UIMax = "10", ClampMin = "0", ClampMax = "100000",
 		EditConditionHides, EditCondition = "GeometryType == ECollisionGeometryType::SweptHulls"))
diff --git a/Engine/Plugins/Runtime/GeometryProcessing/Source/DynamicMesh/Private/ShapeApproximation/MeshSimpleShapeApproximation.cpp b/Engine/Plugins/Runtime/GeometryProcessing/Source/DynamicMesh/Private/ShapeApproximation/MeshSimpleShapeApproximation.cpp
--- a/Engine/Plugins/Runtime/GeometryProcessing/Source/DynamicMesh/Private/ShapeApproximation/MeshSimpleShapeApproximation.cpp
+++ b/Engine/Plugins/Runtime/GeometryProcessing/Source/DynamicMesh/Private/ShapeApproximation/MeshSimpleShapeApproximation.cpp
@@ -389,7 +389,24 @@ void FMeshSimpleShapeApproximation::Generate_ConvexHullDecompositions(FSimpleSha
 		const FDynamicMesh3& SourceMesh = *SourceMeshes[idx];
 		// TODO: if (bSimplifyHulls), also consider simplifying the input?
 		FConvexDecomposition3 Decomposition(SourceMesh);
-		const int32 NumAdditionalSplits = FMath::FloorToInt32(float(ConvexDecompositionMaxPieces) * ConvexDecompositionSearchFactor);
+		int32 NumAdditionalSplits = FMath::FloorToInt32(float(ConvexDecompositionMaxPieces) * ConvexDecompositionSearchFactor);
+		if (bConvexDecompositionProtectNegativeSpace)
+		{
+			FNegativeSpaceSampleSettings Settings;
+			Settings.bOnlyConnectedToHull = bIgnoreInternalNegativeSpace;
+			Settings.MinRadius = NegativeSpaceMinRadius;
+			Settings.ReduceRadiusMargin = NegativeSpaceTolerance;
+			Settings.MinRadius = FMath::Max(1, (NegativeSpaceMinRadius + NegativeSpaceTolerance) * .5);
+			Settings.SampleMethod = FNegativeSpaceSampleSettings::ESampleMethod::VoxelSearch;
+			Settings.bRequireSearchSampleCoverage = true;
+			Settings.TargetNumSamples = 1; // let the sample coverage determine the number of spheres to place
+
+			Decomposition.InitializeNegativeSpace(Settings);
+
+			// Let negative space decide when to stop merging; target only 1 piece if negative space allows
+			NumAdditionalSplits += ConvexDecompositionMaxPieces;
+			ConvexDecompositionMaxPieces = 1;
+		}
 		Decomposition.Compute(ConvexDecompositionMaxPieces, NumAdditionalSplits, ConvexDecompositionErrorTolerance, ConvexDecompositionMinPartThickness);
 
 		for (int32 HullIdx = 0; HullIdx < Decomposition.NumHulls(); HullIdx++)
diff --git a/Engine/Plugins/Runtime/GeometryProcessing/Source/DynamicMesh/Public/ShapeApproximation/MeshSimpleShapeApproximation.h b/Engine/Plugins/Runtime/GeometryProcessing/Source/DynamicMesh/Public/ShapeApproximation/MeshSimpleShapeApproximation.h
--- a/Engine/Plugins/Runtime/GeometryProcessing/Source/DynamicMesh/Public/ShapeApproximation/MeshSimpleShapeApproximation.h
+++ b/Engine/Plugins/Runtime/GeometryProcessing/Source/DynamicMesh/Public/ShapeApproximation/MeshSimpleShapeApproximation.h
@@ -76,6 +76,15 @@ public:
 	double ConvexDecompositionErrorTolerance = 0;
 	/** Minimum part thickness for convex decomposition (in cm); hulls thinner than this will be merged into adjacent hulls, if possible. */
 	double ConvexDecompositionMinPartThickness = .1;
+	/** Whether to guide the convex decomposition to prioritize not filling negative space of the input shape */
+	bool bConvexDecompositionProtectNegativeSpace = false;
+
+	/** Negative space closer to the input than this tolerance distance can be filled in */
+	double NegativeSpaceTolerance = 3;
+	/** Minimum radius of negative space to protect; tunnels with radius smaller than this could be filled in */
+	double NegativeSpaceMinRadius = 10;
+	/** Whether to ignore negative space that is not accessible by traversing from the convex hull (via paths w/ radius of at least Negative Space Tolerance) */
+	bool bIgnoreInternalNegativeSpace = true;
 
 	bool bUseExactComputationForBox = false;
 
