Modified the tool requirements for VertexSculpt, DynamicSculpt, and Remesh so that they require DynamicMeshProvider/Commiter rather than MeshDescriptionProvider/Commiter. These tools operate on dynamic meshes internally and this currently blocks any tool targets that cannot provide also MeshDescriptions.

diff --git a/Engine/Plugins/Experimental/MeshModelingToolsetExp/Source/MeshModelingToolsExp/Private/DynamicMeshSculptTool.cpp b/Engine/Plugins/Experimental/MeshModelingToolsetExp/Source/MeshModelingToolsExp/Private/DynamicMeshSculptTool.cpp
--- a/Engine/Plugins/Experimental/MeshModelingToolsetExp/Source/MeshModelingToolsExp/Private/DynamicMeshSculptTool.cpp
+++ b/Engine/Plugins/Experimental/MeshModelingToolsetExp/Source/MeshModelingToolsExp/Private/DynamicMeshSculptTool.cpp
@@ -37,6 +37,9 @@
 #include "UObject/ObjectMacros.h"
 #include "Materials/Material.h"
 #include "Materials/MaterialInstanceDynamic.h"
+#include "TargetInterfaces/DynamicMeshCommitter.h"
+#include "TargetInterfaces/DynamicMeshProvider.h"
+#include "TargetInterfaces/PrimitiveComponentBackedTarget.h"
 
 #include UE_INLINE_GENERATED_CPP_BY_NAME(DynamicMeshSculptTool)
 
@@ -67,6 +70,16 @@ UMeshSurfacePointTool* UDynamicMeshSculptToolBuilder::CreateNewTool(const FToolB
 	return SculptTool;
 }
 
+const FToolTargetTypeRequirements& UDynamicMeshSculptToolBuilder::GetTargetRequirements() const
+{
+	static FToolTargetTypeRequirements TypeRequirements({
+		UMaterialProvider::StaticClass(),
+		UDynamicMeshCommitter::StaticClass(),
+		UDynamicMeshProvider::StaticClass(),
+		UPrimitiveComponentBackedTarget::StaticClass()
+	});
+	return TypeRequirements;
+}
 
 
 void UDynamicSculptToolActions::DiscardAttributes()
diff --git a/Engine/Plugins/Experimental/MeshModelingToolsetExp/Source/MeshModelingToolsExp/Private/MeshVertexSculptTool.cpp b/Engine/Plugins/Experimental/MeshModelingToolsetExp/Source/MeshModelingToolsExp/Private/MeshVertexSculptTool.cpp
--- a/Engine/Plugins/Experimental/MeshModelingToolsetExp/Source/MeshModelingToolsExp/Private/MeshVertexSculptTool.cpp
+++ b/Engine/Plugins/Experimental/MeshModelingToolsetExp/Source/MeshModelingToolsExp/Private/MeshVertexSculptTool.cpp
@@ -32,6 +32,8 @@
 #include "Sculpting/MeshSculptBrushOps.h"
 #include "Sculpting/StampFalloffs.h"
 #include "Sculpting/MeshSculptUtil.h"
+#include "TargetInterfaces/DynamicMeshCommitter.h"
+#include "TargetInterfaces/DynamicMeshProvider.h"
 
 #include "TargetInterfaces/PrimitiveComponentBackedTarget.h"
 #include "TargetInterfaces/MaterialProvider.h"
@@ -63,6 +65,17 @@ UMeshSurfacePointTool* UMeshVertexSculptToolBuilder::CreateNewTool(const FToolBu
 	return SculptTool;
 }
 
+const FToolTargetTypeRequirements& UMeshVertexSculptToolBuilder::GetTargetRequirements() const
+{
+	static FToolTargetTypeRequirements TypeRequirements({
+		UMaterialProvider::StaticClass(),
+		UDynamicMeshProvider::StaticClass(),
+		UDynamicMeshCommitter::StaticClass(),
+		UPrimitiveComponentBackedTarget::StaticClass()
+	});
+	return TypeRequirements;
+}
+
 
 /*
  * internal Change classes
diff --git a/Engine/Plugins/Experimental/MeshModelingToolsetExp/Source/MeshModelingToolsExp/Private/RemeshMeshTool.cpp b/Engine/Plugins/Experimental/MeshModelingToolsetExp/Source/MeshModelingToolsExp/Private/RemeshMeshTool.cpp
--- a/Engine/Plugins/Experimental/MeshModelingToolsetExp/Source/MeshModelingToolsExp/Private/RemeshMeshTool.cpp
+++ b/Engine/Plugins/Experimental/MeshModelingToolsetExp/Source/MeshModelingToolsExp/Private/RemeshMeshTool.cpp
@@ -23,6 +23,8 @@
 #include "TargetInterfaces/MeshDescriptionProvider.h"
 #include "TargetInterfaces/PrimitiveComponentBackedTarget.h"
 #include "ToolTargetManager.h"
+#include "TargetInterfaces/DynamicMeshCommitter.h"
+#include "TargetInterfaces/DynamicMeshProvider.h"
 
 #include UE_INLINE_GENERATED_CPP_BY_NAME(RemeshMeshTool)
 
@@ -44,6 +46,17 @@ UMultiSelectionMeshEditingTool* URemeshMeshToolBuilder::CreateNewTool(const FToo
 	return NewObject<URemeshMeshTool>(SceneState.ToolManager);
 }
 
+const FToolTargetTypeRequirements& URemeshMeshToolBuilder::GetTargetRequirements() const
+{
+	static FToolTargetTypeRequirements TypeRequirements({
+		UMaterialProvider::StaticClass(),
+		UDynamicMeshCommitter::StaticClass(),
+		UDynamicMeshProvider::StaticClass(),
+		UPrimitiveComponentBackedTarget::StaticClass()
+	});
+	return TypeRequirements;
+}
+
 /*
  * Tool
  */
diff --git a/Engine/Plugins/Experimental/MeshModelingToolsetExp/Source/MeshModelingToolsExp/Public/DynamicMeshSculptTool.h b/Engine/Plugins/Experimental/MeshModelingToolsetExp/Source/MeshModelingToolsExp/Public/DynamicMeshSculptTool.h
--- a/Engine/Plugins/Experimental/MeshModelingToolsetExp/Source/MeshModelingToolsExp/Public/DynamicMeshSculptTool.h
+++ b/Engine/Plugins/Experimental/MeshModelingToolsetExp/Source/MeshModelingToolsExp/Public/DynamicMeshSculptTool.h
@@ -109,6 +109,8 @@ public:
 	}
 
 	virtual UMeshSurfacePointTool* CreateNewTool(const FToolBuilderState& SceneState) const override;
+protected:
+	virtual const FToolTargetTypeRequirements& GetTargetRequirements() const override;
 };
 
 
diff --git a/Engine/Plugins/Experimental/MeshModelingToolsetExp/Source/MeshModelingToolsExp/Public/MeshVertexSculptTool.h b/Engine/Plugins/Experimental/MeshModelingToolsetExp/Source/MeshModelingToolsExp/Public/MeshVertexSculptTool.h
--- a/Engine/Plugins/Experimental/MeshModelingToolsetExp/Source/MeshModelingToolsExp/Public/MeshVertexSculptTool.h
+++ b/Engine/Plugins/Experimental/MeshModelingToolsetExp/Source/MeshModelingToolsExp/Public/MeshVertexSculptTool.h
@@ -37,6 +37,8 @@ class MESHMODELINGTOOLSEXP_API UMeshVertexSculptToolBuilder : public UMeshSurfac
 
 public:
 	virtual UMeshSurfacePointTool* CreateNewTool(const FToolBuilderState& SceneState) const override;
+protected:
+	virtual const FToolTargetTypeRequirements& GetTargetRequirements() const override;
 };
 
 
diff --git a/Engine/Plugins/Experimental/MeshModelingToolsetExp/Source/MeshModelingToolsExp/Public/RemeshMeshTool.h b/Engine/Plugins/Experimental/MeshModelingToolsetExp/Source/MeshModelingToolsExp/Public/RemeshMeshTool.h
--- a/Engine/Plugins/Experimental/MeshModelingToolsetExp/Source/MeshModelingToolsExp/Public/RemeshMeshTool.h
+++ b/Engine/Plugins/Experimental/MeshModelingToolsetExp/Source/MeshModelingToolsExp/Public/RemeshMeshTool.h
@@ -31,6 +31,8 @@ public:
 	virtual bool CanBuildTool(const FToolBuilderState& SceneState) const override;
 
 	virtual UMultiSelectionMeshEditingTool* CreateNewTool(const FToolBuilderState& SceneState) const override;
+protected:
+	virtual const FToolTargetTypeRequirements& GetTargetRequirements() const override;
 };
 
 /**
