Hide the 'show solid' simple collision visualization option in the Collision to Mesh tool, since it's redundant with another option in the main tool

diff --git a/Engine/Plugins/Experimental/MeshModelingToolsetExp/Source/MeshModelingToolsExp/Private/Physics/CollisionGeometryVisualization.cpp b/Engine/Plugins/Experimental/MeshModelingToolsetExp/Source/MeshModelingToolsExp/Private/Physics/CollisionGeometryVisualization.cpp
--- a/Engine/Plugins/Experimental/MeshModelingToolsetExp/Source/MeshModelingToolsExp/Private/Physics/CollisionGeometryVisualization.cpp
+++ b/Engine/Plugins/Experimental/MeshModelingToolsetExp/Source/MeshModelingToolsExp/Private/Physics/CollisionGeometryVisualization.cpp
@@ -298,7 +298,7 @@ void UpdatePreviewGeometryForCollision(
 		FColor TriangleColor = Settings->GetTriangleSetColor(TriColorIndex++);
 		TriangleSet->SetAllTrianglesColor(TriangleColor);
 		TriangleSet->SetAllTrianglesMaterial(Settings->GetSolidMaterial());
-		TriangleSet->SetVisibility(Settings->bShowCollision && Settings->bShowSolid);
+		TriangleSet->SetVisibility(Settings->bEnableShowSolid && Settings->bShowCollision && Settings->bShowSolid);
 	});
 }
 
@@ -337,7 +337,7 @@ void UE::PhysicsTools::InitializeCollisionGeometryVisualization(
 		0);
 	InitializePreviewGeometrySolid(
 		PhysicsData, PreviewGeom, Settings->GetSolidMaterial(), [&Settings](int SolidSetIndex) { return Settings->GetTriangleSetColor(SolidSetIndex); },
-		Settings->bShowCollision && Settings->bShowSolid,
+		Settings->bEnableShowSolid && Settings->bShowCollision && Settings->bShowSolid,
 		CircleStepResolution,
 		0);
 
@@ -387,7 +387,7 @@ void UE::PhysicsTools::PartiallyInitializeCollisionGeometryVisualization(
 
 	InitializePreviewGeometrySolid(
 		PhysicsData, PreviewGeom, Settings->GetSolidMaterial(), [&Settings](int SolidSetIndex) { return Settings->GetTriangleSetColor(SolidSetIndex); },
-		Settings->bShowCollision && Settings->bShowSolid,
+		Settings->bEnableShowSolid && Settings->bShowCollision && Settings->bShowSolid,
 		CircleStepResolution,
 		ColorIndex);
 }
diff --git a/Engine/Plugins/Experimental/MeshModelingToolsetExp/Source/MeshModelingToolsExp/Private/Physics/ExtractCollisionGeometryTool.cpp b/Engine/Plugins/Experimental/MeshModelingToolsetExp/Source/MeshModelingToolsExp/Private/Physics/ExtractCollisionGeometryTool.cpp
--- a/Engine/Plugins/Experimental/MeshModelingToolsetExp/Source/MeshModelingToolsExp/Private/Physics/ExtractCollisionGeometryTool.cpp
+++ b/Engine/Plugins/Experimental/MeshModelingToolsetExp/Source/MeshModelingToolsExp/Private/Physics/ExtractCollisionGeometryTool.cpp
@@ -89,6 +89,7 @@ void UExtractCollisionGeometryTool::Setup()
 	UE::ToolTarget::SetSourceObjectVisible(Target, Settings->bShowInputMesh);
 
 	VizSettings = NewObject<UCollisionGeometryVisualizationProperties>(this);
+	VizSettings->bEnableShowSolid = false; // This solid visualization is redundant to the 'show preview' option in the general settings section of this tool
 	VizSettings->RestoreProperties(this);
 	AddToolPropertySource(VizSettings);
 	VizSettings->Initialize(this);
diff --git a/Engine/Plugins/Experimental/MeshModelingToolsetExp/Source/MeshModelingToolsExp/Public/Physics/CollisionPropertySets.h b/Engine/Plugins/Experimental/MeshModelingToolsetExp/Source/MeshModelingToolsExp/Public/Physics/CollisionPropertySets.h
--- a/Engine/Plugins/Experimental/MeshModelingToolsetExp/Source/MeshModelingToolsExp/Public/Physics/CollisionPropertySets.h
+++ b/Engine/Plugins/Experimental/MeshModelingToolsetExp/Source/MeshModelingToolsExp/Public/Physics/CollisionPropertySets.h
@@ -180,7 +180,8 @@ public:
 	bool bShowCollision = true;
 
 	/** Whether to show solid shapes in addition to wireframes */
-	UPROPERTY(EditAnywhere, Category = "Collision Visualization")
+	UPROPERTY(EditAnywhere, Category = "Collision Visualization", 
+		meta = (EditCondition = "bEnableShowSolid", EditConditionHides, HideEditConditionToggle))
 	bool bShowSolid = false;
 
 	/** Thickness of lines used to visualize collision shapes */
@@ -216,6 +217,10 @@ public:
 	UPROPERTY(Transient, meta=(TransientToolProperty))
 	bool bEnableShowCollision = true;
 
+	//~Some tools will not want the 'show solid' option
+	UPROPERTY(Transient, meta = (TransientToolProperty))
+	bool bEnableShowSolid = true;
+
 	bool bVisualizationDirty = false;
 };
 
