Approximate Actors - Properly handle custom textures sizes options

diff --git a/Engine/Plugins/Experimental/MeshModelingToolsetExp/Source/GeometryProcessingAdapters/Private/ApproximateActorsImpl.cpp b/Engine/Plugins/Experimental/MeshModelingToolsetExp/Source/GeometryProcessingAdapters/Private/ApproximateActorsImpl.cpp
--- a/Engine/Plugins/Experimental/MeshModelingToolsetExp/Source/GeometryProcessingAdapters/Private/ApproximateActorsImpl.cpp
+++ b/Engine/Plugins/Experimental/MeshModelingToolsetExp/Source/GeometryProcessingAdapters/Private/ApproximateActorsImpl.cpp
@@ -189,6 +189,11 @@ static void BakeTexturesFromPhotoCapture(
 		return (HitTID == IndexConstants::InvalidID);
 	};
 
+	auto GetNumDownsampleSteps = [&](int32 TargetImageSize)
+	{
+		return (Options.TextureImageSize / TargetImageSize) * Supersample;
+	};
+
 	FSceneCapturePhotoSet::FSceneSample DefaultSample;
 	FVector4f InvalidColor(0, -1, 0, 1);
 	DefaultSample.BaseColor = FVector3f(InvalidColor.X, InvalidColor.Y, InvalidColor.Z);
@@ -227,15 +232,16 @@ static void BakeTexturesFromPhotoCapture(
 	}
 
 	// downsample the image if necessary
-	if (Supersample > 1)
+	int32 BaseColorDownsampleSteps = GetNumDownsampleSteps(Options.CustomTextureSizeBaseColor);
+	if (BaseColorDownsampleSteps > 1)
 	{
 		TRACE_CPUPROFILER_EVENT_SCOPE(ApproximateActorsImpl_Textures_Downsample);
-		TImageBuilder<FVector4f> Downsampled = ColorImage->FastDownsample(Supersample, FVector4f::Zero(), [](FVector4f V, int N) { return V / (float)N; });
+		TImageBuilder<FVector4f> Downsampled = ColorImage->FastDownsample(BaseColorDownsampleSteps, FVector4f::Zero(), [](FVector4f V, int N) { return V / (float)N; });
 		*ColorImage = MoveTemp(Downsampled);
 	}
 
 	// this lambda is used to process the per-channel images. It does the bake, applies infill, and downsamples if necessary
-	auto ProcessChannelFunc = [&](ERenderCaptureType CaptureType)
+	auto ProcessChannelFunc = [&](ERenderCaptureType CaptureType, int32 CaptureTextureSize)
 	{
 		FVector4f DefaultValue(0, 0, 0, 0);
 		FMeshGenericWorldPositionColorBaker ChannelBaker;
@@ -254,9 +260,10 @@ static void BakeTexturesFromPhotoCapture(
 			return FVector4f(SumValue.X * InvSum, SumValue.Y * InvSum, SumValue.Z * InvSum, 1.0f);
 		});
 
-		if (Supersample > 1)
+		int32 DownsampleSteps = GetNumDownsampleSteps(CaptureTextureSize);
+		if (DownsampleSteps > 1)
 		{
-			TImageBuilder<FVector4f> Downsampled = Image->FastDownsample(Supersample, FVector4f::Zero(), [](FVector4f V, int N) { return V / (float)N; });
+			TImageBuilder<FVector4f> Downsampled = Image->FastDownsample(DownsampleSteps, FVector4f::Zero(), [](FVector4f V, int N) { return V / (float)N; });
 			*Image = MoveTemp(Downsampled);
 		}
 
@@ -272,27 +279,28 @@ static void BakeTexturesFromPhotoCapture(
 
 		if (Options.bUsePackedMRS && (bMetallic || bRoughness || bSpecular))
 		{
-			PackedMRSImage = ProcessChannelFunc(ERenderCaptureType::CombinedMRS);
+			int32 MRSTextureSize = FMath::Max(Options.CustomTextureSizeRoughness, FMath::Max(Options.CustomTextureSizeMetallic, Options.CustomTextureSizeSpecular));
+			PackedMRSImage = ProcessChannelFunc(ERenderCaptureType::CombinedMRS, MRSTextureSize);
 		}
 		else
 		{
 			if (bRoughness)
 			{
-				RoughnessImage = ProcessChannelFunc(ERenderCaptureType::Roughness);
+				RoughnessImage = ProcessChannelFunc(ERenderCaptureType::Roughness, Options.CustomTextureSizeRoughness);
 			}
 			if (bMetallic)
 			{
-				MetallicImage = ProcessChannelFunc(ERenderCaptureType::Metallic);
+				MetallicImage = ProcessChannelFunc(ERenderCaptureType::Metallic, Options.CustomTextureSizeMetallic);
 			}
 			if (bSpecular)
 			{
-				SpecularImage = ProcessChannelFunc(ERenderCaptureType::Specular);
+				SpecularImage = ProcessChannelFunc(ERenderCaptureType::Specular, Options.CustomTextureSizeSpecular);
 			}
 		}
 
 		if (Options.bBakeEmissive)
 		{
-			EmissiveImage = ProcessChannelFunc(ERenderCaptureType::Emissive);
+			EmissiveImage = ProcessChannelFunc(ERenderCaptureType::Emissive, Options.CustomTextureSizeEmissive);
 		}
 	}
 
@@ -325,9 +333,10 @@ static void BakeTexturesFromPhotoCapture(
 		NormalMapBaker.Bake();
 		NormalImage = NormalMapBaker.TakeResult();
 
-		if (Supersample > 1)
+		int32 NormalMapDownsampleSteps = GetNumDownsampleSteps(Options.CustomTextureSizeNormalMap);
+		if (NormalMapDownsampleSteps > 1)
 		{
-			TImageBuilder<FVector3f> Downsampled = NormalImage->FastDownsample(Supersample, FVector3f::Zero(), [](FVector3f V, int N) { return V / (float)N; });
+			TImageBuilder<FVector3f> Downsampled = NormalImage->FastDownsample(NormalMapDownsampleSteps, FVector3f::Zero(), [](FVector3f V, int N) { return V / (float)N; });
 			*NormalImage = MoveTemp(Downsampled);
 		}
 	}
@@ -960,6 +969,38 @@ IGeometryProcessing_ApproximateActors::FOptions FApproximateActorsImpl::Construc
 	Options.HardNormalsAngleDeg = FMath::Clamp(UseSettings.HardNormalAngle, 0.001f, 89.99f);
 
 	Options.TextureImageSize = UseSettings.MaterialSettings.TextureSize.X;
+	
+	if (UseSettings.MaterialSettings.TextureSizingType == TextureSizingType_UseSingleTextureSize)
+	{
+		Options.TextureSizePolicy = ETextureSizePolicy::TextureSize;
+	}
+	else if (UseSettings.MaterialSettings.TextureSizingType == TextureSizingType_UseAutomaticBiasedSizes)
+	{
+		Options.TextureSizePolicy = ETextureSizePolicy::CustomTextureSize;
+
+		int32 NormalSize = Options.TextureImageSize;
+		int32 BaseColorSize = FMath::Max(Options.TextureImageSize >> 1, 32);
+		int32 PropertiesSize = FMath::Max(Options.TextureImageSize >> 2, 16);
+
+		Options.CustomTextureSizeBaseColor = BaseColorSize;
+		Options.CustomTextureSizeRoughness = PropertiesSize;
+		Options.CustomTextureSizeMetallic = PropertiesSize;
+		Options.CustomTextureSizeSpecular = PropertiesSize;
+		Options.CustomTextureSizeEmissive = PropertiesSize;
+		Options.CustomTextureSizeNormalMap = NormalSize;
+	}
+	else if (UseSettings.MaterialSettings.TextureSizingType == TextureSizingType_UseManualOverrideTextureSize)
+	{
+		Options.TextureSizePolicy = ETextureSizePolicy::CustomTextureSize;
+
+		Options.CustomTextureSizeBaseColor = UseSettings.MaterialSettings.DiffuseTextureSize.X;
+		Options.CustomTextureSizeRoughness = UseSettings.MaterialSettings.RoughnessTextureSize.X;
+		Options.CustomTextureSizeMetallic = UseSettings.MaterialSettings.MetallicTextureSize.X;
+		Options.CustomTextureSizeSpecular = UseSettings.MaterialSettings.SpecularTextureSize.X;
+		Options.CustomTextureSizeEmissive = UseSettings.MaterialSettings.EmissiveTextureSize.X;
+		Options.CustomTextureSizeNormalMap = UseSettings.MaterialSettings.NormalTextureSize.X;
+	}
+
 	Options.AntiAliasMultiSampling = FMath::Max(1, UseSettings.MultiSamplingAA);
 
 	Options.RenderCaptureImageSize = (UseSettings.RenderCaptureResolution == 0) ?
@@ -1182,9 +1223,40 @@ void FApproximateActorsImpl::GenerateApproximationForActorSet(const FInput& Inpu
 	FOptions OverridenOptions = Options;
 
 	// evaluate required texture size if needed
-	if (Options.TextureSizePolicy == ETextureSizePolicy::TexelDensity)
+	if (OverridenOptions.TextureSizePolicy == ETextureSizePolicy::TexelDensity)
+	{
+		OverridenOptions.TextureSizePolicy = ETextureSizePolicy::TextureSize;
+		OverridenOptions.TextureImageSize = FMaterialUtilities::GetTextureSizeFromTargetTexelDensity(FinalMesh, OverridenOptions.MeshTexelDensity);
+	}
+
+	if (OverridenOptions.TextureSizePolicy == ETextureSizePolicy::TextureSize)
+	{
+		// Assign the same texture size to all properties
+		OverridenOptions.TextureImageSize = FMath::RoundUpToPowerOfTwo(OverridenOptions.TextureImageSize);
+		OverridenOptions.CustomTextureSizeBaseColor = OverridenOptions.TextureImageSize;
+		OverridenOptions.CustomTextureSizeRoughness = OverridenOptions.TextureImageSize;
+		OverridenOptions.CustomTextureSizeMetallic = OverridenOptions.TextureImageSize;
+		OverridenOptions.CustomTextureSizeSpecular = OverridenOptions.TextureImageSize;
+		OverridenOptions.CustomTextureSizeEmissive = OverridenOptions.TextureImageSize;
+		OverridenOptions.CustomTextureSizeNormalMap = OverridenOptions.TextureImageSize;
+	}
+	else
 	{
-		OverridenOptions.TextureImageSize = FMaterialUtilities::GetTextureSizeFromTargetTexelDensity(FinalMesh, Options.MeshTexelDensity);
+		OverridenOptions.CustomTextureSizeBaseColor = FMath::RoundUpToPowerOfTwo(OverridenOptions.CustomTextureSizeBaseColor);
+		OverridenOptions.CustomTextureSizeRoughness = FMath::RoundUpToPowerOfTwo(OverridenOptions.CustomTextureSizeRoughness);
+		OverridenOptions.CustomTextureSizeMetallic = FMath::RoundUpToPowerOfTwo(OverridenOptions.CustomTextureSizeMetallic);
+		OverridenOptions.CustomTextureSizeSpecular = FMath::RoundUpToPowerOfTwo(OverridenOptions.CustomTextureSizeSpecular);
+		OverridenOptions.CustomTextureSizeEmissive = FMath::RoundUpToPowerOfTwo(OverridenOptions.CustomTextureSizeEmissive);
+		OverridenOptions.CustomTextureSizeNormalMap = FMath::RoundUpToPowerOfTwo(OverridenOptions.CustomTextureSizeNormalMap);
+
+		// Store the maximum texture size in TextureImageSize. This will serve as the capture size for all properties in the baking cache.
+		// Images will then be downscaled as needed for each property.
+		OverridenOptions.TextureImageSize = OverridenOptions.CustomTextureSizeBaseColor;
+		OverridenOptions.TextureImageSize = FMath::Max(OverridenOptions.TextureImageSize, OverridenOptions.CustomTextureSizeRoughness);
+		OverridenOptions.TextureImageSize = FMath::Max(OverridenOptions.TextureImageSize, OverridenOptions.CustomTextureSizeMetallic);
+		OverridenOptions.TextureImageSize = FMath::Max(OverridenOptions.TextureImageSize, OverridenOptions.CustomTextureSizeSpecular);
+		OverridenOptions.TextureImageSize = FMath::Max(OverridenOptions.TextureImageSize, OverridenOptions.CustomTextureSizeEmissive);
+		OverridenOptions.TextureImageSize = FMath::Max(OverridenOptions.TextureImageSize, OverridenOptions.CustomTextureSizeNormalMap);
 	}
 
 	// bake textures for Actor
diff --git a/Engine/Source/Developer/GeometryProcessingInterfaces/Public/GeometryProcessingInterfaces/ApproximateActors.h b/Engine/Source/Developer/GeometryProcessingInterfaces/Public/GeometryProcessingInterfaces/ApproximateActors.h
--- a/Engine/Source/Developer/GeometryProcessingInterfaces/Public/GeometryProcessingInterfaces/ApproximateActors.h
+++ b/Engine/Source/Developer/GeometryProcessingInterfaces/Public/GeometryProcessingInterfaces/ApproximateActors.h
@@ -89,7 +89,8 @@ public:
 	enum class ETextureSizePolicy : uint8
 	{
 		TextureSize = 0,
-		TexelDensity = 1
+		TexelDensity = 1,
+		CustomTextureSize = 2
 	};
 
 	enum class ENaniteFallbackTarget : uint8
@@ -195,16 +196,22 @@ public:
 		UMaterialInterface* BakeMaterial = nullptr;		// if null, will use /MeshModelingToolsetExp/Materials/FullMaterialBakePreviewMaterial_PackedMRS instead
 		FName BaseColorTexParamName = FName("BaseColor");
 		bool bBakeBaseColor = true;
+		int32 CustomTextureSizeBaseColor = 1024;
 		FName RoughnessTexParamName = FName("Roughness");
 		bool bBakeRoughness = true;
+		int32 CustomTextureSizeRoughness = 1024;
 		FName MetallicTexParamName = FName("Metallic");
 		bool bBakeMetallic = true;
+		int32 CustomTextureSizeMetallic = 1024;
 		FName SpecularTexParamName = FName("Specular");
 		bool bBakeSpecular = true;
+		int32 CustomTextureSizeSpecular = 1024;
 		FName EmissiveTexParamName = FName("Emissive");
 		bool bBakeEmissive = true;
+		int32 CustomTextureSizeEmissive = 1024;
 		FName NormalTexParamName = FName("NormalMap");
 		bool bBakeNormalMap = true;
+		int32 CustomTextureSizeNormalMap = 1024;
 		
 		bool bUsePackedMRS = true;
 		FName PackedMRSTexParamName = FName("PackedMRS");
