For modeling Collision to Mesh tool, fix transforms of simple collision parts when Output Separate Meshes is enabled

diff --git a/Engine/Plugins/Experimental/MeshModelingToolsetExp/Source/MeshModelingToolsExp/Private/Physics/ExtractCollisionGeometryTool.cpp b/Engine/Plugins/Experimental/MeshModelingToolsetExp/Source/MeshModelingToolsExp/Private/Physics/ExtractCollisionGeometryTool.cpp
--- a/Engine/Plugins/Experimental/MeshModelingToolsetExp/Source/MeshModelingToolsExp/Private/Physics/ExtractCollisionGeometryTool.cpp
+++ b/Engine/Plugins/Experimental/MeshModelingToolsetExp/Source/MeshModelingToolsExp/Private/Physics/ExtractCollisionGeometryTool.cpp
@@ -185,7 +185,7 @@ void UExtractCollisionGeometryTool::OnShutdown(EToolShutdownType ShutdownType)
 				FAxisAlignedBox3d Bounds = MeshPart.GetBounds();
 				MeshTransforms::Translate(MeshPart, -Bounds.Center());
 				FTransform3d CenterTransform = ActorTransform;
-				CenterTransform.SetTranslation(CenterTransform.GetTranslation() + Bounds.Center());
+				CenterTransform.SetTranslation(CenterTransform.GetTranslation() + ActorTransform.TransformVector(Bounds.Center()));
 				FString NewName = FString::Printf(TEXT("%s_Collision%d"), *TargetName, k);
 				EmitNewMesh(MoveTemp(MeshPart), CenterTransform, NewName);
 			}
