Change display name for Draw Spline's 'World' mode to 'World Objects' to try to make it clearer that it's raycasting to objects in the world

diff --git a/Engine/Plugins/Experimental/MeshModelingToolsetExp/Source/MeshModelingToolsEditorOnlyExp/Public/DrawSplineTool.h b/Engine/Plugins/Experimental/MeshModelingToolsetExp/Source/MeshModelingToolsEditorOnlyExp/Public/DrawSplineTool.h
--- a/Engine/Plugins/Experimental/MeshModelingToolsetExp/Source/MeshModelingToolsEditorOnlyExp/Public/DrawSplineTool.h
+++ b/Engine/Plugins/Experimental/MeshModelingToolsetExp/Source/MeshModelingToolsEditorOnlyExp/Public/DrawSplineTool.h
@@ -160,7 +160,7 @@ public:
 	bool bPreviewUsingActorCopy = true;
 
 	/** Whether to place spline points on the surface of objects in the world */
-	UPROPERTY(EditAnywhere, Category = RaycastTargets, meta = (DisplayName = "World"))
+	UPROPERTY(EditAnywhere, Category = RaycastTargets, meta = (DisplayName = "World Objects"))
 	bool bHitWorld = true;
 
 	/** Whether to place spline points on a custom, user-adjustable plane */
