Make Simple Collision Editor tool only run on targets that support writing collision

(This stops it from running on volumes, which always overwrite the simple collision with bsp-based autogenerated shapes)

diff --git a/Engine/Plugins/Experimental/MeshModelingToolsetExp/Source/MeshModelingToolsExp/Private/Physics/SimpleCollisionEditorTool.cpp b/Engine/Plugins/Experimental/MeshModelingToolsetExp/Source/MeshModelingToolsExp/Private/Physics/SimpleCollisionEditorTool.cpp
--- a/Engine/Plugins/Experimental/MeshModelingToolsetExp/Source/MeshModelingToolsExp/Private/Physics/SimpleCollisionEditorTool.cpp
+++ b/Engine/Plugins/Experimental/MeshModelingToolsetExp/Source/MeshModelingToolsExp/Private/Physics/SimpleCollisionEditorTool.cpp
@@ -11,6 +11,7 @@
 
 #include "Physics/PhysicsDataCollection.h"
 #include "Physics/CollisionGeometryVisualization.h"
+#include "Physics/ComponentCollisionUtil.h"
 
 // physics data
 #include "PhysicsEngine/AggregateGeom.h"
@@ -35,6 +36,22 @@ using namespace UE::Geometry;
 #define LOCTEXT_NAMESPACE "USimpleCollisionEditorTool"
 
 
+bool USimpleCollisionEditorToolBuilder::CanBuildTool(const FToolBuilderState& SceneState) const
+{
+	if (Super::CanBuildTool(SceneState)) // let super verify that there is only one targetable component
+	{
+		// make sure we can also write collision on the target
+		bool bCanReadWrite = false;
+		SceneState.TargetManager->EnumerateSelectedAndTargetableComponents(SceneState, GetTargetRequirements(), [&](UActorComponent* TargetComponent)
+		{
+			bCanReadWrite = Cast<UPrimitiveComponent>(TargetComponent) &&
+				UE::Geometry::ComponentTypeSupportsCollision(Cast<UPrimitiveComponent>(TargetComponent), UE::Geometry::EComponentCollisionSupportLevel::ReadWrite);
+		});
+		return bCanReadWrite;
+	}
+	return false;
+}
+
 const FToolTargetTypeRequirements& USimpleCollisionEditorToolBuilder::GetTargetRequirements() const
 {
 	static FToolTargetTypeRequirements TypeRequirements({
diff --git a/Engine/Plugins/Experimental/MeshModelingToolsetExp/Source/MeshModelingToolsExp/Public/Physics/SimpleCollisionEditorTool.h b/Engine/Plugins/Experimental/MeshModelingToolsetExp/Source/MeshModelingToolsExp/Public/Physics/SimpleCollisionEditorTool.h
--- a/Engine/Plugins/Experimental/MeshModelingToolsetExp/Source/MeshModelingToolsExp/Public/Physics/SimpleCollisionEditorTool.h
+++ b/Engine/Plugins/Experimental/MeshModelingToolsetExp/Source/MeshModelingToolsExp/Public/Physics/SimpleCollisionEditorTool.h
@@ -18,6 +18,8 @@ class MESHMODELINGTOOLSEXP_API USimpleCollisionEditorToolBuilder : public USingl
 	GENERATED_BODY()
 
 public:
+	virtual bool CanBuildTool(const FToolBuilderState& SceneState) const override;
+
 	virtual USingleSelectionMeshEditingTool* CreateNewTool(const FToolBuilderState& SceneState) const override;
 
 protected:
