Skeleton Editing Tool: Snap Bone to Vertex Selection
    - component selection is done on release only (not while dragging for a better UX)
    - support for vertices, edges and faces selection

diff --git a/Engine/Plugins/Experimental/Animation/SkeletalMeshModelingTools/Source/SkeletalMeshModelingTools/Private/DetailCustomization/SkeletonEditingToolPropertyCustomizations.cpp b/Engine/Plugins/Experimental/Animation/SkeletalMeshModelingTools/Source/SkeletalMeshModelingTools/Private/DetailCustomization/SkeletonEditingToolPropertyCustomizations.cpp
--- a/Engine/Plugins/Experimental/Animation/SkeletalMeshModelingTools/Source/SkeletalMeshModelingTools/Private/DetailCustomization/SkeletonEditingToolPropertyCustomizations.cpp
+++ b/Engine/Plugins/Experimental/Animation/SkeletalMeshModelingTools/Source/SkeletalMeshModelingTools/Private/DetailCustomization/SkeletonEditingToolPropertyCustomizations.cpp
@@ -10,6 +10,8 @@
 #include "Widgets/Input/SSegmentedControl.h"
 #include "HAL/PlatformApplicationMisc.h"
 #include "InteractiveToolManager.h"
+#include "ModelingToolsEditorModeStyle.h"
+#include "Selection/PolygonSelectionMechanic.h"
 
 
 #define LOCTEXT_NAMESPACE "SkeletonEditingToolCustomization"
@@ -61,10 +63,6 @@ void ISkeletonEditingPropertiesDetailCustomization::UpdateProperties(
 	const TAttribute<bool>& InEnabled) const
 {
 	for (const FName& PropName: GetProperties())
- 
- 
- 
- 
 	{
 		const TSharedRef<IPropertyHandle> Handle = DetailBuilder.GetProperty(PropName);
 		if (IDetailPropertyRow* Row = DetailBuilder.EditDefaultProperty(Handle))
@@ -786,7 +784,8 @@ namespace SkeletonEditingCustomizationLocals
 	enum EOperationType
 	{
 		Add,
-		Edit
+		Edit,
+		Component
 	};
 }
 
@@ -800,15 +799,17 @@ void FSkeletonEditingToolDetailCustomization::CustomizeDetails(IDetailLayoutBuil
 	TArray<TWeakObjectPtr<UObject>> ObjectsBeingCustomized;
 	DetailBuilder.GetObjectsBeingCustomized(ObjectsBeingCustomized);
 	
-	TWeakObjectPtr<USkeletonEditingTool> Tool = ObjectsBeingCustomized.Num() == 1 ?
+	Tool = ObjectsBeingCustomized.Num() == 1 ?
 		CastChecked<USkeletonEditingTool>(ObjectsBeingCustomized[0]) : nullptr;
 	
 	if (!Tool.IsValid())
 	{
 		return;
 	}
+
+	const FSegmentedControlStyle& SegmentedControlStyle = FAppStyle::Get().GetWidgetStyle<FSegmentedControlStyle>("SegmentedControl");
 	
-	auto IsCreateEnabled = [Tool]()
+	auto IsCreateEnabled = [this]()
 	{
 		return Tool->GetOperation() == EEditingOperation::Create;
 	};
@@ -817,30 +818,170 @@ void FSkeletonEditingToolDetailCustomization::CustomizeDetails(IDetailLayoutBuil
 	IDetailCategoryBuilder& ActionCategory = DetailBuilder.EditCategory("Action", FText::GetEmpty(), ECategoryPriority::Important);
 	ActionCategory.AddCustomRow(LOCTEXT("ActionCategory", "Action"), false)
 	[
-		SNew(SBox)
-		.Padding(2.0f)
+		SNew(SHorizontalBox)
+		+SHorizontalBox::Slot()
+		.Padding(2.f, 4.f)
 		[
-			SNew(SSegmentedControl<EOperationType>)
-			
-			.Value_Lambda([IsCreateEnabled]()
+			SNew(SCheckBox)
+			.Style(&SegmentedControlStyle.FirstControlStyle)
+			.ToolTipText(LOCTEXT("AddModeTooltip", "Create new bones. (N)"))
+			.OnCheckStateChanged_Lambda([this](ECheckBoxState InCheckState)
 			{
-				return IsCreateEnabled() ? Add : Edit;
+				if (InCheckState == ECheckBoxState::Checked)
+				{
+					Tool->SetOperation(EEditingOperation::Create);
+				}
 			})
-			.OnValueChanged_Lambda([Tool](EOperationType Mode)
+			.IsChecked_Lambda([IsCreateEnabled]()
 			{
-				const EEditingOperation Operation = Mode == Add ? EEditingOperation::Create : EEditingOperation::Select;
-				Tool->SetOperation(Operation);
+				return IsCreateEnabled() ? ECheckBoxState::Checked : ECheckBoxState::Unchecked;
 			})
-			+ SSegmentedControl<EOperationType>::Slot(Add)
-			.Text(LOCTEXT("AddMode", "Add"))
-			.ToolTip(LOCTEXT("AddModeTooltip", "Create new bones. (N)"))
-			
-			+ SSegmentedControl<EOperationType>::Slot(Edit)
-			.Text(LOCTEXT("EditMode", "Edit"))
-			.ToolTip(LOCTEXT("EditModeTooltip", "Edit current bone(s) selection. (Esc)"))
+			[
+				SNew(SOverlay)
+				+ SOverlay::Slot()
+				.VAlign(VAlign_Center)
+				.HAlign(HAlign_Center)
+				[
+					SNew(STextBlock)
+					.Text(LOCTEXT("AddMode", "Add"))
+					.TextStyle(&FAppStyle::Get().GetWidgetStyle<FTextBlockStyle>("ButtonText"))
+					.ColorAndOpacity(FLinearColor::White)
+				]
+			]
+		]
+		+SHorizontalBox::Slot()
+		.Padding(2.f, 4.f)
+		[
+			SNew(SCheckBox)
+			.Style(&SegmentedControlStyle.LastControlStyle)
+			.ToolTipText(LOCTEXT("EditModeTooltip", "Edit current bone(s) selection. (Esc)"))
+			.OnCheckStateChanged_Lambda([this](ECheckBoxState InCheckState)
+			{
+				if (InCheckState == ECheckBoxState::Checked)
+				{
+					Tool->SetOperation(EEditingOperation::Select);
+				}
+			})
+			.IsChecked_Lambda([IsCreateEnabled]()
+			{
+				return IsCreateEnabled() ? ECheckBoxState::Unchecked : ECheckBoxState::Checked;
+			})
+			[
+				SNew(SOverlay)
+					
+				+ SOverlay::Slot()
+				.VAlign(VAlign_Center)
+				.HAlign(HAlign_Center)
+				[
+					SNew(STextBlock)
+					.Text(LOCTEXT("EditMode", "Edit"))
+					.TextStyle(&FAppStyle::Get().GetWidgetStyle<FTextBlockStyle>("ButtonText"))
+					.ColorAndOpacity(FLinearColor::White)
+				]
+			]
+		]
+
+		+ SHorizontalBox::Slot()
+		.Padding(8.f, 4.f)
+		.AutoWidth()
+		.HAlign(HAlign_Center)
+		.VAlign(VAlign_Center)
+		[
+			SNew(SCheckBox)
+			.Style(FAppStyle::Get(), "ToggleButtonCheckbox")
+			.ToolTipText(LOCTEXT("ComponentModeTooltip", "Select vertices, edges, and triangles to place bone. (T)"))
+			.OnCheckStateChanged_Lambda([this](ECheckBoxState)
+			{
+				Tool->Properties->bEnableComponentSelection = !Tool->Properties->bEnableComponentSelection;
+			})
+			.IsChecked_Lambda([this]()
+			{
+				return Tool->Properties->bEnableComponentSelection ? ECheckBoxState::Checked : ECheckBoxState::Unchecked;
+			})
+			[
+				SNew(SImage)
+					.Image(FModelingToolsEditorModeStyle::Get()->GetBrush("ModelingToolsManagerCommands.BeginMeshSelectionTool"))
+			]
 		]
 	];
-	const TAttribute<EVisibility> EditionVisibility = TAttribute<EVisibility>::CreateLambda([Tool]()
+
+	CustomizeEditAction(ActionCategory);
+	CustomizeComponentSelection(DetailBuilder);
+}
+
+void FSkeletonEditingToolDetailCustomization::CustomizeComponentSelection(IDetailLayoutBuilder& DetailBuilder) const
+{
+	const TAttribute<bool> IsEnabledBySelection = TAttribute<bool>::CreateLambda([this]()
+	{
+		return Tool.IsValid() ? Tool->HasSelectedComponent() && !Tool->GetSelection().IsEmpty() : false;
+	});
+
+	const TAttribute<EVisibility> VisibilityAttribute = TAttribute<EVisibility>::CreateLambda([this]()
+	{
+		return (Tool.IsValid() && Tool->Properties->bEnableComponentSelection) ? EVisibility::Visible : EVisibility::Hidden;
+	});
+
+	static const FName SnapCategoryName("Component Snapping");
+	IDetailCategoryBuilder& SnapCategory = DetailBuilder.EditCategory(SnapCategoryName, FText::GetEmpty(), ECategoryPriority::Important);
+
+	SnapCategory.AddCustomRow(LOCTEXT("SnapButtonRow", "Snap"), false)
+	[
+		SNew(SHorizontalBox)
+		+SHorizontalBox::Slot()
+		.Padding(2.f, 4.f)
+		[
+			SNew(SButton)
+			.HAlign(HAlign_Center)
+			.VAlign(VAlign_Center)
+			.Text_Lambda([this]
+			{
+				const bool bCreate = Tool.IsValid() && Tool->GetOperation() == EEditingOperation::Create;
+				return bCreate ? LOCTEXT("SnapCreateButtonLabel", "Create") : LOCTEXT("SnapButtonLabel", "Snap");
+			})
+			.IsEnabled_Lambda([this]()
+			{
+				return Tool->HasSelectedComponent(); 
+			})
+			.OnClicked_Lambda([this]()
+			{
+				const bool bCreate = Tool.IsValid() && Tool->GetOperation() == EEditingOperation::Create;
+				Tool->SnapBoneToComponentSelection(bCreate);
+				return FReply::Handled();
+			})
+		]
+	]
+	.Visibility(VisibilityAttribute);
+	
+	// add selection properties
+	TObjectPtr<UMeshTopologySelectionMechanicProperties> SelectionMechanicProperties = Tool->SelectionMechanic->Properties;
+	if (IDetailPropertyRow* SelectionMechanicRow = SnapCategory.AddExternalObjects({SelectionMechanicProperties}))
+	{
+		TSharedPtr<SWidget> NameWidget, ValueWidget;
+		SelectionMechanicRow->GetDefaultWidgets(NameWidget, ValueWidget);
+		SelectionMechanicRow->CustomWidget()
+		.NameContent()
+		[
+			SNew(STextBlock)
+				.Text(FText::FromString(TEXT("Options")))
+				.Font(IDetailLayoutBuilder::GetDetailFont())
+		]
+		.ValueContent()
+		[
+			ValueWidget->AsShared()
+		];
+		SelectionMechanicRow->ShouldAutoExpand(false);
+		SelectionMechanicRow->Visibility(VisibilityAttribute);
+	}
+}
+
+void FSkeletonEditingToolDetailCustomization::CustomizeEditAction(IDetailCategoryBuilder& InActionCategory) const
+{
+	if (!Tool.IsValid())
+	{
+		return;
+	}
+	
+	const TAttribute<EVisibility> EditionVisibility = TAttribute<EVisibility>::CreateLambda([this]()
 	{
 		if (!Tool.IsValid())
 		{
@@ -849,13 +990,13 @@ void FSkeletonEditingToolDetailCustomization::CustomizeDetails(IDetailLayoutBuil
 		return Tool->GetOperation() == EEditingOperation::Create ? EVisibility::Hidden : EVisibility::Visible;
 	});
 	
-	const TAttribute<bool> IsEnabledBySelection = TAttribute<bool>::CreateLambda([Tool]()
+	const TAttribute<bool> IsEnabledBySelection = TAttribute<bool>::CreateLambda([this]()
 	{
 		return Tool.IsValid() ? !Tool->GetSelection().IsEmpty() : false;
 	});
 	FSegmentedControlStyle SegmentedControlStyle = FAppStyle::Get().GetWidgetStyle<FSegmentedControlStyle>("SegmentedControl");
 	
-	ActionCategory.AddCustomRow(LOCTEXT("EditCategory", "Edit"), false)
+	InActionCategory.AddCustomRow(LOCTEXT("EditCategory", "Edit"), false)
 	[
 		SNew(SHorizontalBox)
 		+SHorizontalBox::Slot()
@@ -863,13 +1004,13 @@ void FSkeletonEditingToolDetailCustomization::CustomizeDetails(IDetailLayoutBuil
 		[
 			SNew(SCheckBox)
 			.Style(FAppStyle::Get(), "SegmentedCombo.ButtonOnly")
-			.OnCheckStateChanged_Lambda([Tool](ECheckBoxState InCheckState)
+			.OnCheckStateChanged_Lambda([this](ECheckBoxState InCheckState)
 			{
 				const EEditingOperation Operation = InCheckState == ECheckBoxState::Checked ?
 					EEditingOperation::Parent : EEditingOperation::Select;
 				Tool->SetOperation(Operation);
 			})
-			.IsChecked_Lambda([Tool]()
+			.IsChecked_Lambda([this]()
 			{
 				return Tool->GetOperation() == EEditingOperation::Parent ? ECheckBoxState::Checked : ECheckBoxState::Unchecked;
 			})
@@ -897,12 +1038,12 @@ void FSkeletonEditingToolDetailCustomization::CustomizeDetails(IDetailLayoutBuil
 			.VAlign(VAlign_Center)
 			.Text(LOCTEXT("EditDisconnect", "Disconnect"))
 			.ToolTipText(LOCTEXT("EditDisconnectTooltip", "Unparent current selection. (SHIFT+P)"))
-			.OnClicked_Lambda([Tool]()
+			.OnClicked_Lambda([this]()
 			{
 				Tool->UnParentBones();
 				return FReply::Handled();
 			})
-			.IsEnabled_Lambda([Tool]()
+			.IsEnabled_Lambda([this]()
 			{
 				if (!Tool.IsValid())
 				{
@@ -920,12 +1061,12 @@ void FSkeletonEditingToolDetailCustomization::CustomizeDetails(IDetailLayoutBuil
 			.VAlign(VAlign_Center)
 			.Text(LOCTEXT("EditRemove", "Remove"))
 			.ToolTipText(LOCTEXT("EditRemoveTooltip", "Remove current selection. (Delete)"))
-			.OnClicked_Lambda([Tool]()
+			.OnClicked_Lambda([this]()
 			{
 				Tool->RemoveBones();
 				return FReply::Handled();
 			})
-			.IsEnabled_Lambda([Tool]()
+			.IsEnabled_Lambda([this]()
 			{
 				if (!Tool.IsValid())
 				{
diff --git a/Engine/Plugins/Experimental/Animation/SkeletalMeshModelingTools/Source/SkeletalMeshModelingTools/Private/DetailCustomization/SkeletonEditingToolPropertyCustomizations.h b/Engine/Plugins/Experimental/Animation/SkeletalMeshModelingTools/Source/SkeletalMeshModelingTools/Private/DetailCustomization/SkeletonEditingToolPropertyCustomizations.h
--- a/Engine/Plugins/Experimental/Animation/SkeletalMeshModelingTools/Source/SkeletalMeshModelingTools/Private/DetailCustomization/SkeletonEditingToolPropertyCustomizations.h
+++ b/Engine/Plugins/Experimental/Animation/SkeletalMeshModelingTools/Source/SkeletalMeshModelingTools/Private/DetailCustomization/SkeletonEditingToolPropertyCustomizations.h
@@ -18,6 +18,11 @@ class FSkeletonEditingToolDetailCustomization : public IDetailCustomization
 public:
 	static TSharedRef<IDetailCustomization> MakeInstance();
 	virtual void CustomizeDetails(IDetailLayoutBuilder& DetailBuilder) override;
+private:
+	void CustomizeEditAction(IDetailCategoryBuilder& InActionCategory) const;
+	void CustomizeComponentSelection(IDetailLayoutBuilder& DetailBuilder) const;
+	
+	TWeakObjectPtr<USkeletonEditingTool> Tool = nullptr;
 };
 
 /**
diff --git a/Engine/Plugins/Experimental/MeshModelingToolsetExp/Source/MeshModelingToolsEditorOnlyExp/Private/SkeletalMesh/SkeletonEditingTool.cpp b/Engine/Plugins/Experimental/MeshModelingToolsetExp/Source/MeshModelingToolsEditorOnlyExp/Private/SkeletalMesh/SkeletonEditingTool.cpp
--- a/Engine/Plugins/Experimental/MeshModelingToolsetExp/Source/MeshModelingToolsEditorOnlyExp/Private/SkeletalMesh/SkeletonEditingTool.cpp
+++ b/Engine/Plugins/Experimental/MeshModelingToolsetExp/Source/MeshModelingToolsEditorOnlyExp/Private/SkeletalMesh/SkeletonEditingTool.cpp
@@ -17,6 +17,8 @@
 #include "ToolSetupUtil.h"
 #include "BaseBehaviors/ClickDragBehavior.h"
 #include "BaseGizmos/GizmoViewContext.h"
+#include "DynamicMesh/MeshNormals.h"
+#include "Selection/PolygonSelectionMechanic.h"
 
 #include "SkeletalMesh/SkeletonTransformProxy.h"
 #include "TargetInterfaces/PrimitiveComponentBackedTarget.h"
@@ -109,134 +111,234 @@ void USkeletonEditingTool::Setup()
 		return;
 	}
 
-	// setup modifier
-	Modifier = NewObject<USkeletonModifier>(this);
-	Modifier->SetSkeletalMesh(SkeletalMesh);
-
-	// setup preview
+	SetupModifier(SkeletalMesh);
+	SetupPreviewMesh();
+	SetupProperties();
+	SetupBehaviors();
+	SetupGizmo(Component);
+	SetupWatchers();
+	SetupComponentsSelection();
+	
+	if (EditorContext.IsValid())
 	{
-		PreviewMesh = NewObject<UPreviewMesh>(this);
-		PreviewMesh->bBuildSpatialDataStructure = true;
-		PreviewMesh->CreateInWorld(TargetWorld.Get(), FTransform::Identity);
+		EditorContext->HideSkeleton();
+		EditorContext->BindTo(this);
+	}
+}
 
-		PreviewMesh->SetTransform(UE::ToolTarget::GetLocalToWorldTransform(Target));
+void USkeletonEditingTool::SetupModifier(USkeletalMesh* InSkeletalMesh)
+{
+	if (InSkeletalMesh)
+	{
+		Modifier = NewObject<USkeletonModifier>(this);
+		Modifier->SetSkeletalMesh(InSkeletalMesh);
+	}
+	else
+	{
+		Modifier = nullptr;
+	}
+}
 
-		ToolSetupUtil::ApplyRenderingConfigurationToPreview(PreviewMesh, Target);
-		PreviewMesh->SetTangentsMode(EDynamicMeshComponentTangentsMode::AutoCalculated);
-		PreviewMesh->ReplaceMesh(UE::ToolTarget::GetDynamicMeshCopy(Target));
+void USkeletonEditingTool::SetupPreviewMesh()
+{
+	PreviewMesh = NewObject<UPreviewMesh>(this);
+	PreviewMesh->bBuildSpatialDataStructure = true;
+	PreviewMesh->CreateInWorld(TargetWorld.Get(), FTransform::Identity);
 
-		const FComponentMaterialSet MaterialSet = UE::ToolTarget::GetMaterialSet(Target);
-		PreviewMesh->SetMaterials(MaterialSet.Materials);
+	PreviewMesh->SetTransform(UE::ToolTarget::GetLocalToWorldTransform(Target));
 
-		// hide the skeletal mesh component
-		UE::ToolTarget::HideSourceObject(Target);
-	}
+	ToolSetupUtil::ApplyRenderingConfigurationToPreview(PreviewMesh, Target);
+	PreviewMesh->SetTangentsMode(EDynamicMeshComponentTangentsMode::AutoCalculated);
+	PreviewMesh->ReplaceMesh(UE::ToolTarget::GetDynamicMeshCopy(Target));
 
-	ToolPropertyObjects.Add(this);
+	const FComponentMaterialSet MaterialSet = UE::ToolTarget::GetMaterialSet(Target);
+	PreviewMesh->SetMaterials(MaterialSet.Materials);
 	
-	// setup properties
+	// configure secondary render material for selected triangles
+	UMaterialInterface* SelectionMaterial = ToolSetupUtil::GetSelectionMaterial(FLinearColor::Yellow, GetToolManager());
+	if (SelectionMaterial != nullptr)
 	{
-		ProjectionProperties = NewObject<UProjectionProperties>();
-		ProjectionProperties->Initialize(this, PreviewMesh);
-		ProjectionProperties->RestoreProperties(this);
-		AddToolPropertySource(ProjectionProperties);
+		PreviewMesh->SetSecondaryRenderMaterial(SelectionMaterial);
+	}
 
-		MirroringProperties = NewObject<UMirroringProperties>();
-		MirroringProperties->Initialize(this);
-		MirroringProperties->RestoreProperties(this);
-		AddToolPropertySource(MirroringProperties);
+	// hide the skeletal mesh component
+	UE::ToolTarget::HideSourceObject(Target);
+}
 
-		OrientingProperties = NewObject<UOrientingProperties>();
-		OrientingProperties->Initialize(this);
-		OrientingProperties->RestoreProperties(this);
-		AddToolPropertySource(OrientingProperties);
+void USkeletonEditingTool::SetupProperties()
+{
+	ToolPropertyObjects.Add(this);
+		
+	ProjectionProperties = NewObject<UProjectionProperties>();
+	ProjectionProperties->Initialize(this, PreviewMesh);
+	ProjectionProperties->RestoreProperties(this);
+	AddToolPropertySource(ProjectionProperties);
+
+	MirroringProperties = NewObject<UMirroringProperties>();
+	MirroringProperties->Initialize(this);
+	MirroringProperties->RestoreProperties(this);
+	AddToolPropertySource(MirroringProperties);
+
+	OrientingProperties = NewObject<UOrientingProperties>();
+	OrientingProperties->Initialize(this);
+	OrientingProperties->RestoreProperties(this);
+	AddToolPropertySource(OrientingProperties);
+
+	Properties = NewObject<USkeletonEditingProperties>();
+	Properties->Initialize(this);
+	Properties->Name = GetCurrentBone();
+	Properties->RestoreProperties(this);
+	AddToolPropertySource(Properties);
+}
 
-		Properties = NewObject<USkeletonEditingProperties>();
-		Properties->Initialize(this);
-		Properties->Name = GetCurrentBone();
-		Properties->RestoreProperties(this);
-		AddToolPropertySource(Properties);
-	}
+void USkeletonEditingTool::SetupBehaviors()
+{
+	UClickDragInputBehavior* ClickDragBehavior = NewObject<UClickDragInputBehavior>(this);
+	ClickDragBehavior->Initialize(this);
+	ClickDragBehavior->Modifiers.RegisterModifier(AddToSelectionModifier, FInputDeviceState::IsShiftKeyDown);
+	ClickDragBehavior->Modifiers.RegisterModifier(ToggleSelectionModifier, FInputDeviceState::IsCtrlKeyDown);
+	AddInputBehavior(ClickDragBehavior);
+}
 
-	// setup drag & drop behaviour
+void USkeletonEditingTool::SetupGizmo(USkeletalMeshComponent* InComponent)
+{
+	if (!GizmoContext.IsValid() || !InComponent)
 	{
-		UClickDragInputBehavior* ClickDragBehavior = NewObject<UClickDragInputBehavior>(this);
-		ClickDragBehavior->Initialize(this);
-		ClickDragBehavior->Modifiers.RegisterModifier(AddToSelectionModifier, FInputDeviceState::IsShiftKeyDown);
-		ClickDragBehavior->Modifiers.RegisterModifier(ToggleSelectionModifier, FInputDeviceState::IsCtrlKeyDown);
-		AddInputBehavior(ClickDragBehavior);
+		return;
 	}
 
-	// setup gizmo
-	if (GizmoContext.IsValid())
+	// create state target
+	UGizmoLambdaStateTarget* NewTarget = NewObject<UGizmoLambdaStateTarget>(this);
+	NewTarget->BeginUpdateFunction = [this]()
 	{
-		UGizmoLambdaStateTarget* NewTarget = NewObject<UGizmoLambdaStateTarget>(this);
-		NewTarget->BeginUpdateFunction = [this]()
+		if (GizmoWrapper && GizmoWrapper->CanInteract())
 		{
-			if (GizmoWrapper && GizmoWrapper->CanInteract())
-			{
-				TGuardValue OperationGuard(Operation, EEditingOperation::Transform);
-				BeginChange();
-			}
-		};
-		NewTarget->EndUpdateFunction = [this]()
+			TGuardValue OperationGuard(Operation, EEditingOperation::Transform);
+			BeginChange();
+		}
+	};
+	NewTarget->EndUpdateFunction = [this]()
+	{
+		if (GizmoWrapper && GizmoWrapper->CanInteract())
 		{
-			if (GizmoWrapper && GizmoWrapper->CanInteract())
-			{
-				TGuardValue OperationGuard(Operation, EEditingOperation::Transform);
-				EndChange();
+			TGuardValue OperationGuard(Operation, EEditingOperation::Transform);
+			EndChange();
 
-				const IToolsContextQueriesAPI* ToolsContextQueries = GetToolManager()->GetPairedGizmoManager()->GetContextQueriesAPI();
-				check(ToolsContextQueries);
-				if (ToolsContextQueries->GetCurrentCoordinateSystem() == EToolContextCoordinateSystem::World)
-				{
-					UpdateGizmo();
-				}
+			const IToolsContextQueriesAPI* ToolsContextQueries = GetToolManager()->GetPairedGizmoManager()->GetContextQueriesAPI();
+			check(ToolsContextQueries);
+			if (ToolsContextQueries->GetCurrentCoordinateSystem() == EToolContextCoordinateSystem::World)
+			{
+				UpdateGizmo();
 			}
-		};
-		
-		GizmoWrapper = GizmoContext->GetNewWrapper(GetToolManager(), this, NewTarget);
-		if (GizmoWrapper)
-		{
-			GizmoWrapper->Component = Component;
-			GizmoWrapper->Initialize();
 		}
-	}
+	};
 
-	// setup watchers
-	SelectionWatcher.Initialize(
-	[this]()
-	{
-		return this->Selection;
-	},
-	[this](TArray<FName> InBoneNames)
+	// create gizmo
+	GizmoWrapper = GizmoContext->GetNewWrapper(GetToolManager(), this, NewTarget);
+	if (GizmoWrapper)
 	{
-		UpdateGizmo();
-		if (NeedsNotification())
-		{
-			GetNotifier().Notify(InBoneNames, ESkeletalMeshNotifyType::BonesSelected);
-		}
-	},
-	this->Selection);
+		GizmoWrapper->Component = InComponent;
+		GizmoWrapper->Initialize();
+	}
 
+	// bind coordinate system change
 	IToolsContextQueriesAPI* ToolsContextQueries = GetToolManager()->GetPairedGizmoManager()->GetContextQueriesAPI();
 	CoordinateSystemWatcher.Initialize(
-	[ToolsContextQueries]()
-	{
-		check(ToolsContextQueries);
-		return ToolsContextQueries->GetCurrentCoordinateSystem();
-	},
-	[this](EToolContextCoordinateSystem)
-	{
-		UpdateGizmo();
-	},
-	ToolsContextQueries->GetCurrentCoordinateSystem());
+		[ToolsContextQueries]()
+		{
+			check(ToolsContextQueries);
+			return ToolsContextQueries->GetCurrentCoordinateSystem();
+		},
+		[this](EToolContextCoordinateSystem)
+		{
+			UpdateGizmo();
+		},
+		ToolsContextQueries->GetCurrentCoordinateSystem()
+	);
+}
 
-	if (EditorContext.IsValid())
+void USkeletonEditingTool::SetupWatchers()
+{
+	SelectionWatcher.Initialize(
+		[this]()
+		{
+			return this->Selection;
+		},
+		[this](TArray<FName> InBoneNames)
+		{
+			UpdateGizmo();
+			if (NeedsNotification())
+			{
+				GetNotifier().Notify(InBoneNames, ESkeletalMeshNotifyType::BonesSelected);
+			}
+		},
+		this->Selection
+	);
+}
+
+void USkeletonEditingTool::SetupComponentsSelection()
+{
+	using namespace UE::Geometry;
+	
+	if (!PreviewMesh)
 	{
-		EditorContext->HideSkeleton();
-		EditorContext->BindTo(this);
+		return;
 	}
+	
+	// set up vertex selection mechanic
+	SelectionMechanic = NewObject<UPolygonSelectionMechanic>(this);
+	SelectionMechanic->Setup(this);
+	SelectionMechanic->SetIsEnabled(false);
+	// the property is disabled and managed thru customization
+	SetToolPropertySourceEnabled(SelectionMechanic->Properties, false);
+	SelectionMechanic->SetMarqueeSelectionUpdateType(EMarqueeSelectionUpdateType::OnRelease);
+	SelectionMechanic->Properties->bMarqueeIgnoreOcclusion = false;
+	SelectionMechanic->Properties->bHitBackFaces = false;
+	SelectionMechanic->Properties->bDisplayPolygroupReliantControls = false;
+
+	SelectionMechanic->Properties->bSelectVertices = true;
+	SelectionMechanic->Properties->bSelectEdges = false;
+	SelectionMechanic->Properties->bSelectFaces = false;
+	SelectionMechanic->SetShowSelectableCorners(true);
+
+	// adjust selection rendering for this context
+	SelectionMechanic->PolyEdgesRenderer.PointColor = FLinearColor(0.78f, 0.f, 0.78f);
+	SelectionMechanic->PolyEdgesRenderer.PointSize = 5.0f;
+	SelectionMechanic->PolyEdgesRenderer.DepthBias = 0.01f;
+	
+	SelectionMechanic->HilightRenderer.PointColor = FLinearColor::Red;
+	SelectionMechanic->HilightRenderer.PointSize = 10.0f;
+	
+	SelectionMechanic->SelectionRenderer.LineThickness = 0.0f;
+	SelectionMechanic->SelectionRenderer.DepthBias = 0.01f;
+	SelectionMechanic->SelectionRenderer.PointColor = FLinearColor::Yellow;
+	SelectionMechanic->SelectionRenderer.PointSize = 5.0f;
+	SelectionMechanic->SetShowEdges(false);
+	
+	// initialize the polygon selection mechanic
+	PreviewMesh->ProcessMesh([this](const FDynamicMesh3& Mesh)
+	{
+		Topology = MakeUnique<FTriangleGroupTopology>(&Mesh, true);
+		SelectionMechanic->Initialize(&Mesh,
+			FTransform::Identity,
+			TargetWorld.Get(),
+			Topology.Get(),
+			[this]() { return PreviewMesh->GetSpatial(); });
+	});
+
+	// triangles rendering
+	PreviewMesh->EnableSecondaryTriangleBuffers([this](const FDynamicMesh3* Mesh, int32 TriangleID)
+	{
+		return SelectionMechanic->GetActiveSelection().IsSelectedTriangle(Mesh, Topology.Get(), TriangleID);
+	});
+	SelectionMechanic->OnSelectionChanged.AddWeakLambda(this, [this]()
+	{
+		PreviewMesh->FastNotifySecondaryTrianglesChanged();
+	});
+	SelectionMechanic->OnFaceSelectionPreviewChanged.AddWeakLambda(this, [this]()
+	{
+		PreviewMesh->FastNotifySecondaryTrianglesChanged();
+	});
 }
 
 void USkeletonEditingTool::UpdateGizmo() const
@@ -248,7 +350,7 @@ void USkeletonEditingTool::UpdateGizmo() const
 
 	GizmoWrapper->Initialize();
 
-	const bool bUseGizmo = Operation == EEditingOperation::Select || Operation == EEditingOperation::Transform; 
+	const bool bUseGizmo = !Properties->bEnableComponentSelection && (Operation == EEditingOperation::Select || Operation == EEditingOperation::Transform); 
 	if (Selection.IsEmpty() || !bUseGizmo)
 	{
 		return;
@@ -311,6 +413,13 @@ void USkeletonEditingTool::Shutdown(EToolShutdownType ShutdownType)
 	
 	Super::Shutdown(ShutdownType);
 
+	if (SelectionMechanic)
+	{
+		SelectionMechanic->Properties->SaveProperties(this);
+		SelectionMechanic->Shutdown();
+		SelectionMechanic = nullptr;
+	}
+	
 	// remove preview mesh
 	if (PreviewMesh != nullptr)
 	{
@@ -359,6 +468,7 @@ void USkeletonEditingTool::RegisterActions(FInteractiveToolActionSet& InOutActio
 	RegisterCopyAction(InOutActionSet, GetActionId());
 	RegisterPasteAction(InOutActionSet, GetActionId());
 	RegisterDuplicateAction(InOutActionSet, GetActionId());
+	RegisterSelectComponentsAction(InOutActionSet, GetActionId());
 }
 
 void USkeletonEditingTool::RegisterCreateAction(FInteractiveToolActionSet& InOutActionSet, const int32 InActionId)
@@ -525,6 +635,53 @@ void USkeletonEditingTool::RegisterDuplicateAction(FInteractiveToolActionSet& In
 		});
 }
 
+void USkeletonEditingTool::RegisterSelectComponentsAction(FInteractiveToolActionSet& InOutActionSet, const int32 InActionId)
+{
+	InOutActionSet.RegisterAction(this, InActionId,
+		TEXT("SelectComponents"),
+		LOCTEXT("SelectComponents", "Select Components"),
+		LOCTEXT("SelectComponentsDesc", "Select Components"),
+		EModifierKey::None, EKeys::T,
+		[this]()
+		{
+			Properties->bEnableComponentSelection = !Properties->bEnableComponentSelection; 
+			UpdateGizmo();
+		});
+}
+
+TArray<int32> USkeletonEditingTool::GetSelectedComponents() const
+{
+	TArray<int32> IDs;
+
+	PreviewMesh->ProcessMesh([this, &IDs](const FDynamicMesh3& Mesh)
+	{
+		const FGroupTopologySelection& TopologySelection = SelectionMechanic->GetActiveSelection();
+		if (SelectionMechanic->Properties->bSelectVertices)
+		{
+			Algo::CopyIf(TopologySelection.SelectedCornerIDs, IDs, [&Mesh](int ID)
+			{
+				return Mesh.IsVertex(ID);	
+			});
+		}
+		else if (SelectionMechanic->Properties->bSelectEdges)
+		{
+			Algo::CopyIf(TopologySelection.SelectedEdgeIDs, IDs, [&Mesh](int ID)
+			{
+				return Mesh.IsEdge(ID);	
+			});
+		}
+		else if (SelectionMechanic->Properties->bSelectFaces)
+		{
+			Algo::CopyIf(TopologySelection.SelectedGroupIDs, IDs, [&Mesh](int ID)
+			{
+				return Mesh.IsTriangle(ID);	
+			});
+		}
+	});
+	
+	return MoveTemp(IDs);
+}
+
 void USkeletonEditingTool::CreateNewBone()
 {
 	if (Operation != EEditingOperation::Create)
@@ -625,6 +782,79 @@ void USkeletonEditingTool::UnParentBones()
 	CancelChange();
 }
 
+void USkeletonEditingTool::SnapBoneToComponentSelection(const bool bCreate)
+{
+	const TArray<int32> IDs = GetSelectedComponents();
+	if (IDs.IsEmpty())
+	{
+		return;
+	}
+
+	const FReferenceSkeleton& RefSkeleton = Modifier->GetReferenceSkeleton();
+	const TArray<FName> Bones = GetSelectedBones();
+	const bool bHasValidBone = Bones.ContainsByPredicate([&](const FName& InBoneName)
+	{
+		return RefSkeleton.FindRawBoneIndex(InBoneName) > INDEX_NONE;
+	});
+
+	if (!bHasValidBone)
+	{
+		return;
+	}
+
+	const int32 BoneIndex = RefSkeleton.FindRawBoneIndex(GetCurrentBone());
+	const int32 ParentBoneIndex = bCreate ? BoneIndex : RefSkeleton.GetRawParentIndex(BoneIndex);
+	
+	FTransform NewTransform = ComputeTransformFromComponents(IDs);
+	const FTransform& ParentGlobal = Modifier->GetTransform(ParentBoneIndex, true);
+	NewTransform = NewTransform.GetRelativeTransform(ParentGlobal);
+
+	bool bSuccess = false;
+	if (bCreate)
+	{
+		TGuardValue OperationGuard(Operation, EEditingOperation::Create);
+		BeginChange();
+
+		static const FName DefaultName("joint");
+		const FName BoneName = Modifier->GetUniqueName(DefaultName);
+		const FName ParentName = Selection.IsEmpty() ? NAME_None : Selection.Last(); 
+		bSuccess = Modifier->AddBone(BoneName, ParentName, NewTransform);
+		if (bSuccess)
+		{
+			if (NeedsNotification())
+			{
+				GetNotifier().Notify({BoneName}, ESkeletalMeshNotifyType::BonesAdded);
+			}
+	
+			Selection = {BoneName};
+			Properties->Name = BoneName;
+			EndChange();
+		}
+	}
+	else
+	{
+		TGuardValue OperationGuard(Operation, EEditingOperation::Transform);
+		BeginChange();
+
+		bSuccess = Modifier->SetBoneTransform(GetCurrentBone(), NewTransform, Properties->bUpdateChildren);
+		if (bSuccess)
+		{
+			EndChange();
+		}
+	}
+
+	if (bSuccess)
+	{
+		Properties->Transform = NewTransform;
+		SelectionMechanic->ClearSelection();
+		UpdateGizmo();
+	}
+	else
+	{
+		CancelChange();	
+	}
+}
+
 void USkeletonEditingTool::ParentBones(const FName& InParentName)
 {
 	if (Operation != EEditingOperation::Parent)
@@ -734,6 +964,252 @@ void USkeletonEditingTool::RenameBones()
 	CancelChange();
 }
 
+namespace SkeletonEditingTool
+{
+	using namespace UE::Geometry;
+	
+	// this is a helper class to compute a transform from a Components selection.
+	// NOTE: the code bellow assumes that the ids are safe to use with Mesh
+	// so make sure they are tested before calling any function in here.
+	
+	struct FTransformFromMesh
+	{
+		FTransformFromMesh() = delete;
+		FTransformFromMesh(const FDynamicMesh3& InMesh)
+			: Mesh(InMesh)
+		{}
+		
+		FTransform GetVertexTransform(const int32 InId) const
+		{
+			static constexpr bool bFrameNormalY = false;
+			FVector3d Normal = FMeshNormals::ComputeVertexNormal(Mesh, InId);
+			const FFrame3d VertexFrame = Mesh.GetVertexFrame(InId, bFrameNormalY, &Normal);
+			return VertexFrame.ToFTransform();
+		}
+
+		FTransform GetEdgeTransform(const int32 InId) const
+		{
+			const FDynamicMesh3::FEdge& EdgeRef = Mesh.GetEdgeRef(InId);
+			
+			const FVector3d& V0 = Mesh.GetVertexRef(EdgeRef.Vert[0]);
+			const FVector3d& V1 = Mesh.GetVertexRef(EdgeRef.Vert[1]);
+			const FVector3d Centroid = (V0 + V1) * 0.5;
+			
+			const FVector3d Normal = Mesh.GetEdgeNormal(InId);
+			const FVector3d X = (V1 - V0).GetSafeNormal();
+			const FVector3d Y = Normal.Cross(X);
+			const FQuat Rotation = FRotationMatrix::MakeFromXY(X, Y).ToQuat();
+
+			return FTransform(Rotation, Centroid);
+		}
+
+		FTransform GetTriangleTransform(const int32 InId) const
+		{
+			return Mesh.GetTriFrame(InId, 0).ToFTransform();	
+		}
+		
+		FTransform GetTwoVerticesTransform(const int32 InId0, const int32 InId1) const
+		{
+			// check if this is an edge
+			auto GetEdgeId = [this, InId0, InId1]() -> int32
+			{	
+				for (const int32 EdgeId : Mesh.EdgeIndicesItr())
+				{
+					const FDynamicMesh3::FEdge& EdgeRef = Mesh.GetEdgeRef(EdgeId);
+					if (EdgeRef.Vert.Contains(InId0) && EdgeRef.Vert.Contains(InId1))
+					{
+						return EdgeId;
+					}
+				}
+				return INDEX_NONE;
+			};
+
+			const int32 EdgeId = GetEdgeId();
+			if (EdgeId != INDEX_NONE)
+			{
+				return GetEdgeTransform(EdgeId);
+			}
+
+			const FVector3d& V0 = Mesh.GetVertexRef(InId0);
+			const FVector3d& V1 = Mesh.GetVertexRef(InId1);
+			const FVector3d Centroid = (V0 + V1) * 0.5;
+			
+			const FVector3d N0 = FMeshNormals::ComputeVertexNormal(Mesh, InId0);
+			const FVector3d N1 = FMeshNormals::ComputeVertexNormal(Mesh, InId1);
+			const FVector3d Normal = ((N0 + N1) * 0.5).GetSafeNormal();
+			
+			const FVector3d X = (V1 - V0).GetSafeNormal();
+			const FVector3d Y = Normal.Cross(X);
+			const FQuat Rotation = FRotationMatrix::MakeFromXY(X, Y).ToQuat();
+
+			return FTransform(Rotation, Centroid);;
+		}
+		
+		FTransform GetThreeVerticesTransform(const int32 InId0, const int32 InId1, const int32 InId2) const
+		{
+			// check if this is a triangle
+			auto GetTriangleId = [this, InId0, InId1, InId2]() -> int32
+			{
+				for (const int32 TriangleID : Mesh.TriangleIndicesItr())
+				{
+					const FIndex3i& TriRef = Mesh.GetTriangleRef(TriangleID);
+					if (TriRef.Contains(InId0) && TriRef.Contains(InId1) && TriRef.Contains(InId2))
+					{
+						return TriangleID;
+					}
+				}
+				return INDEX_NONE;
+			};
+
+			const int32 TriangleID = GetTriangleId();
+			if (TriangleID != INDEX_NONE)
+			{
+				return GetTriangleTransform(TriangleID);
+			}
+
+			const FVector3d& V0 = Mesh.GetVertexRef(InId0);
+			const FVector3d& V1 = Mesh.GetVertexRef(InId1);
+			const FVector3d& V2 = Mesh.GetVertexRef(InId2);
+			const FVector3d Centroid = (V0 + V1 + V2) / 3.0;
+
+			const FVector3d Normal = VectorUtil::Normal(V0, V1, V2);
+			const FVector3d X = (V1 - V0).GetSafeNormal();
+			const FVector3d Y = Normal.Cross(X);
+			const FQuat Rotation = FRotationMatrix::MakeFromXY(X, Y).ToQuat();
+
+			return FTransform(Rotation, Centroid);
+		}
+
+		FTransform GetVerticesTransform(const TArray<int32>& InVertexIDs) const
+		{
+			switch(InVertexIDs.Num())
+			{
+			case 0:
+				return FTransform::Identity;
+			case 1:
+				return GetVertexTransform(InVertexIDs[0]);
+			case 2:
+				return GetTwoVerticesTransform(InVertexIDs[0], InVertexIDs[1]);
+			case 3:
+				return GetThreeVerticesTransform(InVertexIDs[0], InVertexIDs[1], InVertexIDs[2]);
+			default:
+				break;
+			}
+
+			// compute centroid only
+			FVector Centroid(0.0);
+			for (const int32 VertexID: InVertexIDs)
+			{
+				Centroid += Mesh.GetVertexRef(VertexID);
+			}
+			Centroid /= static_cast<double>(InVertexIDs.Num());
+
+			FTransform Transform(Centroid);
+			return MoveTemp(Transform);
+		}
+		
+		FTransform GetEdgesTransform(const TArray<int32>& InEdgeIDs) const
+		{
+			switch (InEdgeIDs.Num())
+			{
+			case 0:
+				return FTransform::Identity;
+			case 1:
+				return GetEdgeTransform(InEdgeIDs[0]);
+			default:
+				break;
+			}
+
+			FVector Centroid(0.0), Normal(0.0);
+			for (const int32 EdgeID: InEdgeIDs)
+			{
+				const FDynamicMesh3::FEdge& EdgeRef = Mesh.GetEdgeRef(EdgeID);
+				const FVector3d& V0 = Mesh.GetVertexRef(EdgeRef.Vert[0]);
+				const FVector3d& V1 = Mesh.GetVertexRef(EdgeRef.Vert[1]);
+				Centroid += ((V0 + V1) * 0.5);
+				Normal += Mesh.GetEdgeNormal(EdgeID);
+			}
+			
+			Centroid /= static_cast<double>(InEdgeIDs.Num());
+			Normalize(Normal);
+
+			const FDynamicMesh3::FEdge& EdgeRef = Mesh.GetEdgeRef(InEdgeIDs[0]);
+			const FVector3d& V0 = Mesh.GetVertexRef(EdgeRef.Vert[0]);
+			const FVector3d& V1 = Mesh.GetVertexRef(EdgeRef.Vert[1]);
+			const FVector3d X = (((V0 + V1) * 0.5) - Centroid).GetSafeNormal();
+			const FVector3d Y = Normal.Cross(X);
+			const FQuat Rotation = FRotationMatrix::MakeFromXY(X, Y).ToQuat();
+
+			return FTransform(Rotation, Centroid);
+		}
+
+		FTransform GetFacesTransform(const TArray<int32>& InTriIDs) const
+		{
+			switch (InTriIDs.Num())
+			{
+			case 0:
+				return FTransform::Identity;
+			case 1:
+				return GetTriangleTransform(InTriIDs[0]);
+			default:
+				break;
+			}
+
+			FVector Centroid(0.0), Normal(0.0);
+			for (const int32 TriID: InTriIDs)
+			{
+				Centroid += Mesh.GetTriCentroid(TriID);
+				Normal += Mesh.GetTriNormal(TriID);
+			}
+			
+			Centroid /= static_cast<double>(InTriIDs.Num());
+			Normalize(Normal);
+
+			const FVector3d X = (Mesh.GetTriCentroid(InTriIDs[0]) - Centroid).GetSafeNormal();
+			const FVector3d Y = Normal.Cross(X);
+			const FQuat Rotation = FRotationMatrix::MakeFromXY(X, Y).ToQuat();
+
+			return FTransform(Rotation, Centroid);
+		}
+
+	private:
+		const FDynamicMesh3& Mesh;
+	};
+}
+
+FTransform USkeletonEditingTool::ComputeTransformFromComponents(const TArray<int32>& InIDs) const
+{
+	const int32 NumComponents = InIDs.Num();
+	if (NumComponents == 0)
+	{
+		return FTransform::Identity;
+	}
+
+	const FDynamicMesh3* Mesh = PreviewMesh->GetMesh();
+	if (!ensure(Mesh))
+	{
+		return FTransform::Identity;
+	}
+
+	const SkeletonEditingTool::FTransformFromMesh TransformFromMesh(*Mesh);
+	if (SelectionMechanic->Properties->bSelectVertices)
+	{
+		return TransformFromMesh.GetVerticesTransform(InIDs);
+	}
+	
+	if (SelectionMechanic->Properties->bSelectEdges)
+	{
+		return TransformFromMesh.GetEdgesTransform(InIDs);
+	}
+	
+	if (SelectionMechanic->Properties->bSelectFaces)
+	{
+		return TransformFromMesh.GetFacesTransform(InIDs);
+	}
+
+	return FTransform::Identity;
+}
+
 void USkeletonEditingTool::OnClickPress(const FInputDeviceRay& InPressPos)
 {
 	if (PendingFunction)
@@ -856,6 +1332,11 @@ void USkeletonEditingTool::SetOperation(const EEditingOperation InOperation, con
 	}
 }
 
+bool USkeletonEditingTool::HasSelectedComponent() const
+{
+	return !SelectionMechanic->GetActiveSelection().IsEmpty();
+}
+
 void USkeletonEditingTool::SelectBone(const FName& InBoneName)
 {
 	TArray<FName> NewSelection = Selection;
@@ -970,10 +1451,19 @@ void USkeletonEditingTool::OnTick(float DeltaTime)
 	
 	SelectionWatcher.CheckAndUpdate();
 	CoordinateSystemWatcher.CheckAndUpdate();
+
+	PreviewMesh->EnableWireframe(Properties->bEnableComponentSelection);
+	SelectionMechanic->SetIsEnabled(Properties->bEnableComponentSelection);
+	SelectionMechanic->Tick(DeltaTime);
 }
 
 FInputRayHit USkeletonEditingTool::CanBeginClickDragSequence(const FInputDeviceRay& InPressPos)
 {
+	if (Properties->bEnableComponentSelection)
+	{
+		return FInputRayHit();
+	}
+	
 	PendingFunction.Reset();
 	ParentIndex = INDEX_NONE;
 	
@@ -1118,6 +1608,16 @@ void USkeletonEditingTool::HandleSkeletalMeshModified(const TArray<FName>& InBon
 	Properties->Name = GetCurrentBone();
 }
 
+void USkeletonEditingTool::DrawHUD(FCanvas* Canvas, IToolsContextRenderAPI* RenderAPI)
+{
+	Super::DrawHUD(Canvas, RenderAPI);
+
+	if (SelectionMechanic)
+	{
+		SelectionMechanic->DrawHUD(Canvas, RenderAPI);
+	}
+}
+
 void USkeletonEditingTool::Render(IToolsContextRenderAPI* RenderAPI)
 {
 	// FIXME many things could be caches here and updated lazilly
@@ -1175,10 +1675,56 @@ void USkeletonEditingTool::Render(IToolsContextRenderAPI* RenderAPI)
 		HitProxies,
 		DrawConfig
 	);
+
+	if (SelectionMechanic && Properties->bEnableComponentSelection)
+	{
+		SelectionMechanic->Render(RenderAPI);
+	}
+}
+
+void USkeletonEditingTool::OnPropertyModified(UObject* PropertySet, FProperty* Property)
+{
+	if (PropertySet == SelectionMechanic->Properties && Property)
+	{
+		if (Property->GetFName() == GET_MEMBER_NAME_CHECKED(UMeshTopologySelectionMechanicProperties, bSelectVertices))
+		{
+			if (SelectionMechanic->Properties->bSelectVertices)
+			{
+				SelectionMechanic->Properties->bSelectEdges = SelectionMechanic->Properties->bSelectFaces = false;
+			}
+		}
+		else if (Property->GetFName() == GET_MEMBER_NAME_CHECKED(UMeshTopologySelectionMechanicProperties, bSelectEdges))
+		{
+			if (SelectionMechanic->Properties->bSelectEdges)
+			{
+				SelectionMechanic->Properties->bSelectVertices = SelectionMechanic->Properties->bSelectFaces = false;
+			}
+		}
+		else if (Property->GetFName() == GET_MEMBER_NAME_CHECKED(UMeshTopologySelectionMechanicProperties, bSelectFaces))
+		{
+			if (SelectionMechanic->Properties->bSelectFaces)
+			{
+				SelectionMechanic->Properties->bSelectEdges = SelectionMechanic->Properties->bSelectVertices = false;
+			}
+		}
+		SelectionMechanic->SetShowSelectableCorners(SelectionMechanic->Properties->bSelectVertices);
+	}
+	Super::OnPropertyModified(PropertySet, Property);
 }
 
 FBox USkeletonEditingTool::GetWorldSpaceFocusBox()
 {
+	if (Properties->bEnableComponentSelection)
+	{
+		const FGroupTopologySelection& ActiveSelection = SelectionMechanic->GetActiveSelection();
+		if (!ActiveSelection.IsEmpty())
+		{
+			static constexpr bool bWorld = true;
+			const FAxisAlignedBox3d Bounds = SelectionMechanic->GetSelectionBounds(bWorld);
+			return static_cast<FBox>(Bounds);
+		}
+	}
+	
 	const TArray<int32> BoneIndexes = GetSelectedBoneIndexes();
 	if (!BoneIndexes.IsEmpty())
 	{
diff --git a/Engine/Plugins/Experimental/MeshModelingToolsetExp/Source/MeshModelingToolsEditorOnlyExp/Public/SkeletalMesh/SkeletonEditingTool.h b/Engine/Plugins/Experimental/MeshModelingToolsetExp/Source/MeshModelingToolsEditorOnlyExp/Public/SkeletalMesh/SkeletonEditingTool.h
--- a/Engine/Plugins/Experimental/MeshModelingToolsetExp/Source/MeshModelingToolsEditorOnlyExp/Public/SkeletalMesh/SkeletonEditingTool.h
+++ b/Engine/Plugins/Experimental/MeshModelingToolsetExp/Source/MeshModelingToolsEditorOnlyExp/Public/SkeletalMesh/SkeletonEditingTool.h
@@ -3,6 +3,7 @@
 #pragma once
 
 #include "DynamicMeshBrushTool.h"
+#include "GroupTopology.h"
 #include "BaseTools/MeshSurfacePointMeshEditingTool.h"
 
 #include "SkeletalMeshNotifier.h"
@@ -10,7 +11,9 @@
 #include "SkeletalMesh/SkeletalMeshEditionInterface.h"
 #include "Engine/World.h"
 #include "Changes/ValueWatcher.h"
+#include "Mechanics/RectangleMarqueeMechanic.h"
 #include "Misc/EnumClassFlags.h"
+#include "Selection/MeshTopologySelectionMechanic.h"
 
 #include "SkeletonEditingTool.generated.h"
 
@@ -21,6 +24,7 @@ class USkeletonEditingTool;
 class USkeletonTransformProxy;
 class USkeletalMeshGizmoContextObjectBase;
 class USkeletalMeshGizmoWrapperBase;
+class UPolygonSelectionMechanic;
 
 namespace SkeletonEditingTool
 {
@@ -108,12 +112,15 @@ public:
 	virtual void Setup() override;
 	virtual void Shutdown(EToolShutdownType ShutdownType) override;
 
+	virtual void DrawHUD(FCanvas* Canvas, IToolsContextRenderAPI* RenderAPI) override;
 	virtual void Render(IToolsContextRenderAPI* RenderAPI) override;
 	
 	virtual bool HasCancel() const override { return true; }
 	virtual bool HasAccept() const override { return true; }
 	virtual bool CanAccept() const override { return true; }
 
+	virtual void OnPropertyModified(UObject* PropertySet, FProperty* Property) override;
+
 	// ICLickDragBehaviorTarget overrides
 	virtual FInputRayHit CanBeginClickDragSequence(const FInputDeviceRay& InPressPos) override;
 	
@@ -134,7 +141,8 @@ public:
 	void OrientBones();
 	void RemoveBones();
 	void UnParentBones();
-
+	void SnapBoneToComponentSelection(const bool bCreate);
+	
 	// ISkeletalMeshEditionInterface overrides
 	virtual TArray<FName> GetSelectedBones() const override;
 
@@ -151,6 +159,8 @@ public:
 
 	EEditingOperation GetOperation() const;
 	void SetOperation(const EEditingOperation InOperation, const bool bUpdateGizmo = true);
+
+	bool HasSelectedComponent() const;
 	
 protected:
 
@@ -164,8 +174,7 @@ protected:
 	void CreateNewBone();
 	void ParentBones(const FName& InParentName);
 
-	UPROPERTY()
-	TObjectPtr<USkeletonModifier> Modifier;
+public:
 	
 	UPROPERTY()
 	TObjectPtr<USkeletonEditingProperties> Properties;
@@ -176,12 +185,17 @@ protected:
 	UPROPERTY()
 	TObjectPtr<UMirroringProperties> MirroringProperties;
 
-public:
 	UPROPERTY()
 	TObjectPtr<UOrientingProperties> OrientingProperties;
+
+	UPROPERTY()
+	TObjectPtr<UPolygonSelectionMechanic> SelectionMechanic;
 	
 protected:
 	
+	UPROPERTY()
+	TObjectPtr<USkeletonModifier> Modifier = nullptr;
+	
 	UPROPERTY()
 	TObjectPtr<USingleClickInputBehavior> LeftClickBehavior = nullptr;
 
@@ -241,6 +255,15 @@ private:
 
 	void NormalizeSelection();
 
+	// setup
+	void SetupModifier(USkeletalMesh* InSkeletalMesh);
+	void SetupPreviewMesh();
+	void SetupProperties();
+	void SetupBehaviors();
+	void SetupGizmo(USkeletalMeshComponent* InComponent);
+	void SetupWatchers();
+	void SetupComponentsSelection();
+	
 	// actions
 	void RegisterCreateAction(FInteractiveToolActionSet& InOutActionSet, const int32 InActionId);
 	void RegisterDeleteAction(FInteractiveToolActionSet& InOutActionSet, const int32 InActionId);
@@ -250,7 +273,11 @@ private:
 	void RegisterCopyAction(FInteractiveToolActionSet& InOutActionSet, const int32 InActionId);
 	void RegisterPasteAction(FInteractiveToolActionSet& InOutActionSet, const int32 InActionId);
 	void RegisterDuplicateAction(FInteractiveToolActionSet& InOutActionSet, const int32 InActionId);
+	void RegisterSelectComponentsAction(FInteractiveToolActionSet& InOutActionSet, const int32 InActionId);
 
+	TArray<int32> GetSelectedComponents() const;
+
+	FTransform ComputeTransformFromComponents(const TArray<int32>& InIDs) const;
 	
 	TValueWatcher<TArray<FName>> SelectionWatcher;
 	TValueWatcher<EToolContextCoordinateSystem> CoordinateSystemWatcher; 
@@ -258,6 +285,8 @@ private:
 	int32 ParentIndex = INDEX_NONE;
 
 	TFunction<void()> PendingFunction;
+
+	TUniquePtr<UE::Geometry::FTriangleGroupTopology> Topology = nullptr;
 };
 
 ENUM_CLASS_FLAGS(USkeletonEditingTool::EBoneSelectionMode);
@@ -288,6 +317,9 @@ public:
 	UPROPERTY(EditAnywhere, Category = "Viewport Axis Settings",  meta = (DisplayPriority = 10, ClampMin = "0.0", UIMin = "0.0"))
 	float AxisThickness = 0.f;
 
+	UPROPERTY()
+	bool bEnableComponentSelection = false;
+
 #if WITH_EDITOR
 	virtual void PostEditChangeProperty(FPropertyChangedEvent &PropertyChangedEvent) override;
 #endif
