[GeometryProcessing] Fixing combine mesh instances including decorative mesh instances in coarse / voxwrapped LODs.

The coarse mesh was being overriden under the assumption that its source was the last simplified LOD, while it was in fact the last copied source LOD with the decorative instances correctly filtered.
Removing the mesh override fixes the issue.

diff --git a/Engine/Plugins/Experimental/MeshModelingToolsetExp/Source/GeometryProcessingAdapters/Private/CombineMeshInstancesImpl.cpp b/Engine/Plugins/Experimental/MeshModelingToolsetExp/Source/GeometryProcessingAdapters/Private/CombineMeshInstancesImpl.cpp
--- a/Engine/Plugins/Experimental/MeshModelingToolsetExp/Source/GeometryProcessingAdapters/Private/CombineMeshInstancesImpl.cpp
+++ b/Engine/Plugins/Experimental/MeshModelingToolsetExp/Source/GeometryProcessingAdapters/Private/CombineMeshInstancesImpl.cpp
@@ -3596,14 +3596,6 @@ void BuildCombinedMesh(
 		}
 	}
 
-	// Above we copied the last approximation LOD into the first VoxWrap LOD slot. However 
-	// approximate LOD may be very coarse, it's better to use a known-good starting mesh?
-	// Need to do this before starting hidden-removal task...
-	if (FirstVoxWrappedIndex < MeshLODs.Num())
-	{
-		MeshLODs[FirstVoxWrappedIndex].Mesh = MeshLODs[CombineOptions.NumCopiedLODs-1].Mesh;
-	}
-
 	// make a list of per-LOD geometric tolerances tha will drive additional optimization. 
 	// For copied and first simplified LODs, use Simplify Base Tolerance, and then increment
 	// for each successive LOD.  (todo: have a separate initial tolerance for Approx LODs?)
