Modeling Tools: Enables Polygroup preview when generating polygroups

diff --git a/Engine/Plugins/Experimental/MeshModelingToolsetExp/Source/MeshModelingToolsExp/Private/ConvertToPolygonsTool.cpp b/Engine/Plugins/Experimental/MeshModelingToolsetExp/Source/MeshModelingToolsExp/Private/ConvertToPolygonsTool.cpp
--- a/Engine/Plugins/Experimental/MeshModelingToolsetExp/Source/MeshModelingToolsExp/Private/ConvertToPolygonsTool.cpp
+++ b/Engine/Plugins/Experimental/MeshModelingToolsetExp/Source/MeshModelingToolsExp/Private/ConvertToPolygonsTool.cpp
@@ -610,7 +610,7 @@ void UConvertToPolygonsTool::UpdateVisualization()
 		int32 NumMaterials = MaterialTarget->GetNumMaterials();
 		for (int32 i = 0; i < NumMaterials; ++i)
 		{ 
-			MaterialSet.Materials.Add(ToolSetupUtil::GetSelectionMaterial(GetToolManager()));
+			MaterialSet.Materials.Add(ToolSetupUtil::GetVertexColorMaterial(GetToolManager()));
 		}
 		PreviewCompute->PreviewMesh->SetTriangleColorFunction([this](const FDynamicMesh3* Mesh, int TriangleID)
 		{
