Fixes bug skin paint tool properties not synchronized with viewport brush adjustments.
Fixes issue with brush settings not save per brush mode

diff --git a/Engine/Plugins/Experimental/MeshModelingToolsetExp/Source/MeshModelingToolsEditorOnlyExp/Private/SkeletalMesh/SkinWeightsPaintTool.cpp b/Engine/Plugins/Experimental/MeshModelingToolsetExp/Source/MeshModelingToolsEditorOnlyExp/Private/SkeletalMesh/SkinWeightsPaintTool.cpp
--- a/Engine/Plugins/Experimental/MeshModelingToolsetExp/Source/MeshModelingToolsEditorOnlyExp/Private/SkeletalMesh/SkinWeightsPaintTool.cpp
+++ b/Engine/Plugins/Experimental/MeshModelingToolsetExp/Source/MeshModelingToolsEditorOnlyExp/Private/SkeletalMesh/SkinWeightsPaintTool.cpp
@@ -2114,6 +2114,19 @@ void USkinWeightsPaintTool::OnPropertyModified(UObject* ModifiedObject, FPropert
 {
 	Super::OnPropertyModified(ModifiedObject, ModifiedProperty);
 
+	if (ModifiedProperty->GetName() == GET_MEMBER_NAME_STRING_CHECKED(USkinWeightsPaintToolProperties, BrushStrength))
+	{
+		WeightToolProperties->GetBrushConfig().Strength = WeightToolProperties->BrushStrength;
+	}
+	if (ModifiedProperty->GetName() == GET_MEMBER_NAME_STRING_CHECKED(USkinWeightsPaintToolProperties, BrushRadius))
+	{
+		WeightToolProperties->GetBrushConfig().Radius = WeightToolProperties->BrushRadius;
+	}
+	if (ModifiedProperty->GetName() == GET_MEMBER_NAME_STRING_CHECKED(USkinWeightsPaintToolProperties, BrushFalloffAmount))
+	{
+		WeightToolProperties->GetBrushConfig().Falloff = WeightToolProperties->BrushFalloffAmount;
+	}
+	
 	const FString NameOfModifiedProperty = ModifiedProperty->GetNameCPP();
 
 	// invalidate vertex color cache when any weight color properties are modified
diff --git a/Engine/Source/Runtime/InteractiveToolsFramework/Private/BaseTools/BaseBrushTool.cpp b/Engine/Source/Runtime/InteractiveToolsFramework/Private/BaseTools/BaseBrushTool.cpp
--- a/Engine/Source/Runtime/InteractiveToolsFramework/Private/BaseTools/BaseBrushTool.cpp
+++ b/Engine/Source/Runtime/InteractiveToolsFramework/Private/BaseTools/BaseBrushTool.cpp
@@ -105,6 +105,9 @@ void UBrushAdjusterInputBehavior::OnDragUpdate(FVector2D InScreenPosition)
 		NewRadius = FMath::Max(NewRadius, 0.01f);
 		BrushTool->BrushProperties->BrushRadius = NewRadius;
 		BrushTool->LastBrushStamp.Radius = NewRadius;
+		
+		FPropertyChangedEvent PropertyChangedEvent(UBrushBaseProperties::StaticClass()->FindPropertyByName(GET_MEMBER_NAME_CHECKED(UBrushBaseProperties, BrushRadius)));
+		BrushTool->BrushProperties->PostEditChangeProperty(PropertyChangedEvent);
 	}
 	else
 	{
@@ -112,6 +115,9 @@ void UBrushAdjusterInputBehavior::OnDragUpdate(FVector2D InScreenPosition)
 		float NewStrength = StartBrushStrength + VerticalDelta * -(StrengthAdjustSpeed * DPIScale);
 		NewStrength = FMath::Min(1.0f,FMath::Max(NewStrength, 0.f));
 		BrushTool->BrushProperties->BrushStrength = NewStrength;
+		
+		FPropertyChangedEvent PropertyChangedEvent(UBrushBaseProperties::StaticClass()->FindPropertyByName(GET_MEMBER_NAME_CHECKED(UBrushBaseProperties, BrushStrength)));
+		BrushTool->BrushProperties->PostEditChangeProperty(PropertyChangedEvent);
 	}	
 }
 
