MeshModelingTools: Display selection measurements in CubeGrid in viewport.

diff --git a/Engine/Plugins/Experimental/MeshModelingToolsetExp/Source/MeshModelingToolsExp/Private/CubeGridTool.cpp b/Engine/Plugins/Experimental/MeshModelingToolsetExp/Source/MeshModelingToolsExp/Private/CubeGridTool.cpp
--- a/Engine/Plugins/Experimental/MeshModelingToolsetExp/Source/MeshModelingToolsExp/Private/CubeGridTool.cpp
+++ b/Engine/Plugins/Experimental/MeshModelingToolsetExp/Source/MeshModelingToolsExp/Private/CubeGridTool.cpp
@@ -8,6 +8,7 @@
 #include "BaseGizmos/GizmoMath.h"
 #include "BaseGizmos/TransformGizmoUtil.h"
 #include "BaseGizmos/TransformProxy.h"
+#include "CanvasTypes.h"
 #include "CompositionOps/CubeGridBooleanOp.h"
 #include "Distance/DistLine3Ray3.h"
 #include "Drawing/PreviewGeometryActor.h"
@@ -15,6 +16,7 @@
 #include "DynamicMesh/DynamicMeshChangeTracker.h"
 #include "DynamicMesh/MeshTransforms.h"
 #include "DynamicMeshToMeshDescription.h"
+#include "Engine/Engine.h"  // GEngine->GetSmallFont()
 #include "InteractiveToolChange.h"
 #include "InteractiveToolManager.h"
 #include "InputState.h"
@@ -26,6 +28,7 @@
 #include "ModelingToolTargetUtil.h"
 #include "Properties/MeshMaterialProperties.h"
 #include "PropertySets/CreateMeshObjectTypeProperties.h"
+#include "SceneView.h"
 #include "Selection/ToolSelectionUtil.h"
 #include "ToolContextInterfaces.h"
 #include "ToolTargetManager.h"
@@ -480,6 +483,86 @@ namespace CubeGridToolLocals
 		DrawParallelInteriorLines(Dim2, Dim1);
 	}
 
+
+	void DisplayLengthsOfEveryNonzeroBoxSide(const FAxisAlignedBox3d& Box, const FTransform& BoxTransform, 
+		FCanvas& Canvas, const FSceneView& SceneView)
+	{
+		if (Box.IsEmpty())
+		{
+			return;
+		}
+
+		FVector3d LocalDimensions = Box.Max - Box.Min;
+
+		if (LocalDimensions.IsZero())
+		{
+			return;
+		}
+
+		FVector3d LocalCenter = Box.Center();
+		double DPIScale = Canvas.GetDPIScale();
+		UFont* UseFont = GEngine->GetSmallFont();
+
+		for (int MeasuredDim = 0; MeasuredDim < 3; ++MeasuredDim)
+		{
+			if (LocalDimensions[MeasuredDim] == 0)
+			{
+				// Flat on this side
+				continue;
+			}
+
+			double Length = FMath::Abs(LocalDimensions[MeasuredDim] * BoxTransform.GetScale3D()[MeasuredDim]);
+			FString String;
+			if (Length - static_cast<int32>(Length) < KINDA_SMALL_NUMBER)
+			{
+				String = FString::Printf(TEXT("%.0f"), Length);
+			}
+			else
+			{
+				// Two decimal places if we don't have a round number
+				String = FString::Printf(TEXT("%.2f"), Length);
+			}
+
+			// Iterate in a square across the other two dimensions
+			int OtherDim1 = MeasuredDim == 0 ? 1 : 0;
+			int OtherDim2 = MeasuredDim == 2 ? 1 : 2;
+			for (int i = 0; i < 2; ++i)
+			{
+				if (i == 1 && LocalDimensions[OtherDim1] == 0)
+				{
+					break;
+				}
+
+				for (int j = 0; j < 2; ++j)
+				{
+					if (j == 1 && LocalDimensions[OtherDim2] == 0)
+					{
+						break;
+					}
+
+					FVector3d LocalWriteLocation = Box.Min;
+					LocalWriteLocation[MeasuredDim] = LocalCenter[MeasuredDim];
+					if (i == 1)
+					{
+						LocalWriteLocation[OtherDim1] = Box.Max[OtherDim1];
+					}
+					if (j == 1)
+					{
+						LocalWriteLocation[OtherDim2] = Box.Max[OtherDim2];
+					}
+
+					FVector3d WorldPosition = BoxTransform.TransformPosition(LocalWriteLocation);
+					FVector2D PixelPosition;
+					SceneView.WorldToPixel(WorldPosition, PixelPosition);
+					Canvas.DrawShadowedString(PixelPosition.X / DPIScale, PixelPosition.Y / DPIScale, *String, 
+						UseFont, FLinearColor::White);
+				}
+			}
+			
+		}
+
+	}
+
 	/** Given a world hit, get a hit face. */
 	void ConvertToFaceHit(const FCubeGrid& CubeGrid, ECubeGridToolFaceSelectionMode SelectionMode, 
 		 const FRay& WorldRay, double HitT, const FVector3d& Normal, FCubeGrid::FCubeFace& FaceOut,
@@ -1308,6 +1391,21 @@ void UCubeGridTool::Render(IToolsContextRenderAPI* RenderAPI)
 	GridGizmoAlignmentMechanic->Render(RenderAPI);
 }
 
+void UCubeGridTool::DrawHUD(FCanvas* Canvas, IToolsContextRenderAPI* RenderAPI)
+{
+	using namespace CubeGridToolLocals;
+
+	if (bHaveSelection && Canvas && CubeGrid 
+		&& Settings && Settings->bShowSelectionMeasurements)
+	{
+		if (const FSceneView* SceneView = RenderAPI->GetSceneView())
+		{
+			DisplayLengthsOfEveryNonzeroBoxSide(Selection.Box, CubeGrid->GetFrame().ToFTransform(), 
+				*Canvas, *SceneView);
+		}
+	}
+}
+
 void UCubeGridTool::OnPropertyModified(UObject* PropertySet, FProperty* Property)
 {
 }
diff --git a/Engine/Plugins/Experimental/MeshModelingToolsetExp/Source/MeshModelingToolsExp/Public/CubeGridTool.h b/Engine/Plugins/Experimental/MeshModelingToolsetExp/Source/MeshModelingToolsExp/Public/CubeGridTool.h
--- a/Engine/Plugins/Experimental/MeshModelingToolsetExp/Source/MeshModelingToolsExp/Public/CubeGridTool.h
+++ b/Engine/Plugins/Experimental/MeshModelingToolsetExp/Source/MeshModelingToolsExp/Public/CubeGridTool.h
@@ -150,6 +150,10 @@ public:
 	UPROPERTY(EditAnywhere, Category = Options, AdvancedDisplay)
 	bool bKeepSideGroups = true;
 
+	/** When true, displays dimensions of the given selection in the viewport. */
+	UPROPERTY(EditAnywhere, Category = Options, AdvancedDisplay)
+	bool bShowSelectionMeasurements = true;
+
 	/** When performing selection, the tolerance to use when determining
 	 whether things lie in the same plane as a cube face. */
 	//~ This turned out to not be a useful setting, so it is no longer EditAnywhere. The only cases where it
@@ -373,6 +377,7 @@ public:
 
 	virtual void OnTick(float DeltaTime) override;
 	virtual void Render(IToolsContextRenderAPI* RenderAPI) override;
+	virtual void DrawHUD(FCanvas* Canvas, IToolsContextRenderAPI* RenderAPI) override;
 
 	virtual void OnPropertyModified(UObject* PropertySet, FProperty* Property) override;
 
