fix ensure on hole fill tool attempting to adding spans as loops -- check if the last and first vertices of the span are connected by an edge before attempting to add the loop

diff --git a/Engine/Plugins/Experimental/MeshModelingToolsetExp/Source/MeshModelingToolsExp/Private/HoleFillTool.cpp b/Engine/Plugins/Experimental/MeshModelingToolsetExp/Source/MeshModelingToolsExp/Private/HoleFillTool.cpp
--- a/Engine/Plugins/Experimental/MeshModelingToolsetExp/Source/MeshModelingToolsExp/Private/HoleFillTool.cpp
+++ b/Engine/Plugins/Experimental/MeshModelingToolsetExp/Source/MeshModelingToolsExp/Private/HoleFillTool.cpp
@@ -180,6 +180,11 @@ void UHoleFillTool::Setup()
 		// Since this tool can tolerate failing to fill a boundary loop, just convert the spans to loops and let them fail
 		for (const FEdgeSpan& Span : BoundaryLoops->Spans)
 		{
+			// skip span if there is no edge connecting the first and last vertices, as these cannot be initialized as loops
+			if (Span.Vertices.Num() < 2 || FDynamicMesh3::InvalidID == OriginalMesh->FindEdge(Span.Vertices[0], Span.Vertices.Last()))
+			{
+				continue;
+			}
 			FEdgeLoop SpanLoop(OriginalMesh.Get());
 			if (SpanLoop.InitializeFromVertices(Span.Vertices))
 			{
