MeshModelingTools: CubeGrid fixes/improvements:
- Fix mesh transform changing in undo after deleting full mesh.
- Add ability for CubeGrid to accept and start a new mesh without having to exit.
- Move the corner mode done/cancel buttons to viewport.
- Fix some tooltips
- Fix CubeGrid being able to make meshes empty (do same thing that Plane Cut tool does instead).

diff --git a/Engine/Plugins/Experimental/MeshModelingToolsetExp/Source/MeshModelingToolsExp/Private/CubeGridTool.cpp b/Engine/Plugins/Experimental/MeshModelingToolsetExp/Source/MeshModelingToolsExp/Private/CubeGridTool.cpp
--- a/Engine/Plugins/Experimental/MeshModelingToolsetExp/Source/MeshModelingToolsExp/Private/CubeGridTool.cpp
+++ b/Engine/Plugins/Experimental/MeshModelingToolsetExp/Source/MeshModelingToolsExp/Private/CubeGridTool.cpp
@@ -13,12 +13,14 @@
 #include "Distance/DistLine3Ray3.h"
 #include "Drawing/PreviewGeometryActor.h"
 #include "Drawing/LineSetComponent.h"
+#include "DynamicMeshEditor.h"
 #include "DynamicMesh/DynamicMeshChangeTracker.h"
 #include "DynamicMesh/MeshTransforms.h"
 #include "DynamicMeshToMeshDescription.h"
 #include "Engine/Engine.h"  // GEngine->GetSmallFont()
 #include "InteractiveToolChange.h"
 #include "InteractiveToolManager.h"
+#include "Input/Reply.h"
 #include "InputState.h"
 #include "Mechanics/DragAlignmentMechanic.h"
 #include "MeshOpPreviewHelpers.h" //UMeshOpPreviewWithBackgroundCompute
@@ -31,6 +33,7 @@
 #include "SceneView.h"
 #include "Selection/ToolSelectionUtil.h"
 #include "ToolContextInterfaces.h"
+#include "ToolHostCustomizationAPI.h"
 #include "ToolTargetManager.h"
 #include "ToolTargets/ToolTarget.h"
 #include "TargetInterfaces/MaterialProvider.h"
@@ -57,7 +60,11 @@ namespace CubeGridToolLocals
 		"Refer to side panel for shortcuts.");
 
 	FText CornerModeMessage = LOCTEXT("CubeGridCornerModeDescription", "Toggle corner selection for push/pulling by clicking or dragging. "
-		"Press Enter or click \"Done\" in the side panel to accept the result.");
+		"Press Enter or click \"Done\" to accept the result.");
+
+	FText CreatingAssetMessage = LOCTEXT("CreatingNewAssetLabel", "Creating new asset.");
+
+	FText EditingAssetMessage = LOCTEXT("EditingExistingAssetLabel", "Editing existing asset.");
 
 	const FText SelectionChangeTransactionName = LOCTEXT("SelectionChangeTransaction", "Cube Grid Selection Change");
 	const FText ModeChangeTransactionName = LOCTEXT("ModeChangeTransaction", "Cube Grid Mode Change");
@@ -121,6 +128,33 @@ namespace CubeGridToolLocals
 		TUniquePtr<UE::Geometry::FDynamicMeshChange> MeshChange;
 	};
 
+	class FCubeGridMeshTransformChange : public FToolCommandChange
+	{
+	public:
+		FCubeGridMeshTransformChange(const FTransform& BeforeIn, const FTransform& AfterIn)
+			: Before(BeforeIn)
+			, After(AfterIn)
+		{};
+
+		virtual void Apply(UObject* Object) override
+		{
+			Cast<UCubeGridTool>(Object)->SetCurrentMeshTransform(After);
+		}
+		virtual void Revert(UObject* Object) override
+		{
+			Cast<UCubeGridTool>(Object)->SetCurrentMeshTransform(Before);
+		}
+
+		virtual FString ToString() const override
+		{
+			return TEXT("CubeGridToolLocals::FCubeGridMeshTransformChange");
+		}
+
+	protected:
+		FTransform Before;
+		FTransform After;
+	};
+
 	/** Undoes selection changes */
 	class FCubeGridToolSelectionChange : public FToolCommandChange
 	{
@@ -273,21 +307,64 @@ namespace CubeGridToolLocals
 	class FCubeGridChangesMadeChange : public FToolCommandChange
 	{
 	public:
-		FCubeGridChangesMadeChange() {};
+		FCubeGridChangesMadeChange(bool bNowMadeIn = true) 
+			: bNowMade(bNowMadeIn) {};
 
 		virtual void Apply(UObject* Object) override
 		{
-			Cast<UCubeGridTool>(Object)->SetChangesMade(true);
+			Cast<UCubeGridTool>(Object)->SetChangesMade(bNowMade);
 		}
 		virtual void Revert(UObject* Object) override
 		{
-			Cast<UCubeGridTool>(Object)->SetChangesMade(false);
+			Cast<UCubeGridTool>(Object)->SetChangesMade(!bNowMade);
 		}
 
 		virtual FString ToString() const override
 		{
 			return TEXT("CubeGridToolLocals::FCubeGridChangesMadeChange");
 		}
+
+	protected:
+		bool bNowMade = true;
+	};
+
+	/** Deals with the tool target portion of an "Accept and Start New" change. The actual
+	  asset update and current mesh updates should happen alongside this change.*/
+	class FCubeGridTargetResetChange : public FToolCommandChange
+	{
+	public:
+		FCubeGridTargetResetChange(UToolTarget* ToolTargetIn)
+			: ToolTarget(ToolTargetIn)
+		{};
+
+		virtual void Apply(UObject* Object) override
+		{
+			UCubeGridTool* Tool = Cast<UCubeGridTool>(Object);
+			if (ToolTarget.IsValid())
+			{
+				UE::ToolTarget::SetSourceObjectVisible(ToolTarget.Get(), true);
+			}
+			Tool->SetTarget(nullptr);
+			Tool->GetToolManager()->DisplayMessage(CreatingAssetMessage, EToolMessageLevel::UserWarning);
+		}
+		virtual void Revert(UObject* Object) override
+		{
+			UCubeGridTool* Tool = Cast<UCubeGridTool>(Object);
+			Tool->SetTarget(ToolTarget.Get());
+			Tool->GetToolManager()->DisplayMessage(EditingAssetMessage, EToolMessageLevel::UserWarning);
+			if (ToolTarget.IsValid())
+			{
+				UE::ToolTarget::SetSourceObjectVisible(ToolTarget.Get(), false);
+			}
+		}
+
+		virtual FString ToString() const override
+		{
+			return TEXT("CubeGridToolLocals::FCubeGridAcceptAndStartNewChange");
+		}
+
+	protected:
+		TWeakObjectPtr<UToolTarget> ToolTarget;
 	};
 	
 	// Attach a frame to the box such that Z points along the given direction.
@@ -513,7 +590,7 @@ namespace CubeGridToolLocals
 
 			double Length = FMath::Abs(LocalDimensions[MeasuredDim] * BoxTransform.GetScale3D()[MeasuredDim]);
 			FString String;
-			if (Length - static_cast<int32>(Length) < KINDA_SMALL_NUMBER)
+			if (FMath::Abs(Length - FMath::RoundToDouble(Length)) < KINDA_SMALL_NUMBER)
 			{
 				String = FString::Printf(TEXT("%.0f"), Length);
 			}
@@ -696,6 +773,7 @@ void UCubeGridTool::InvalidatePreview(bool bUpdateCornerLineSet)
 		if (bPreviewMayDiffer)
 		{
 			Preview->PreviewMesh->UpdatePreview(CurrentMesh.Get());
+			Preview->PreviewMesh->SetTransform(CurrentMeshTransform);
 			bPreviewMayDiffer = false;
 		}
 		return;
@@ -855,11 +933,6 @@ void UCubeGridTool::Setup()
 
 	GetToolManager()->DisplayMessage(PushPullModeMessage, EToolMessageLevel::UserNotification);
 
-	DuringActivityActions = NewObject<UCubeGridDuringActivityActions>(this);
-	DuringActivityActions->Initialize(this);
-	AddToolPropertySource(DuringActivityActions);
-	SetToolPropertySourceEnabled(DuringActivityActions, false);
-
 	ToolActions = NewObject<UCubeGridToolActions>(this);
 	ToolActions->Initialize(this);
 	AddToolPropertySource(ToolActions);
@@ -1135,12 +1208,12 @@ void UCubeGridTool::Setup()
 
 	if (Target)
 	{
-		GetToolManager()->DisplayMessage(LOCTEXT("EditingExistingAssetLabel", "Editing existing asset."), 
+		GetToolManager()->DisplayMessage(EditingAssetMessage,
 			EToolMessageLevel::UserWarning);
 	}
 	else
 	{
-		GetToolManager()->DisplayMessage(LOCTEXT("CreatingNewAssetLabel", "Creating new asset."),
+		GetToolManager()->DisplayMessage(CreatingAssetMessage,
 			EToolMessageLevel::UserWarning);
 	}
 }
@@ -1177,7 +1250,7 @@ void UCubeGridTool::Shutdown(EToolShutdownType ShutdownType)
 
 	if (Target)
 	{
-		Cast<IPrimitiveComponentBackedTarget>(Target)->SetOwnerVisibility(true);
+		UE::ToolTarget::SetSourceObjectVisible(Target.Get(), true);
 	}
 
 	// CubeGrid might get used for long stretches at a time, and an accidental Esc hit could result in a fair
@@ -1200,68 +1273,164 @@ void UCubeGridTool::Shutdown(EToolShutdownType ShutdownType)
 
 	if (ShutdownType == EToolShutdownType::Accept && bChangesMade)
 	{
-		bool bCreatingNewAsset = !Target && CurrentMesh->TriangleCount() > 0;
-		if (Target)
+		OutputCurrentResults(/*bSetSelection =*/ true);
+	}
+
+	Preview->OnOpCompleted.RemoveAll(this);
+	Preview->Shutdown();
+
+	if (LineSets)
+	{
+		LineSets->Disconnect();
+		LineSets = nullptr;
+	}
+
+	GridGizmoAlignmentMechanic->Shutdown();
+
+	GetToolManager()->GetPairedGizmoManager()->DestroyAllGizmosByOwner(this);
+
+	// Probably unnecessary by this point, but just in case.
+	ClearViewportButtonCustomization();
+}
+
+// Applies the tool's results to the target or outputs the mesh. Meant to be called during shutdown
+// or when doing "Accept and Start New"
+void UCubeGridTool::OutputCurrentResults(bool bSetSelection)
+{
+	if (CurrentMesh->TriangleCount() <= 0)
+	{
+		if (Target && Target->IsValid())
 		{
-			if (Target->IsValid())
+			EAppReturnType::Type Ret = FMessageDialog::Open(EAppMsgType::YesNo,
+				LOCTEXT("DestroyComponentQuestion", "The tool has entirely cut away the source mesh, which is not permitted. "
+					"Do you actually want to delete the mesh component? Note that the actor will remain. Choosing No will leave "
+					"the component unchanged instead."),
+				LOCTEXT("DestroyComponentTitle", "Delete mesh components?"));
+			if (Ret == EAppReturnType::No || Ret == EAppReturnType::Cancel)
 			{
-				GetToolManager()->BeginUndoTransaction(LOCTEXT("CubeGridToolEditTransactionName", "Cube Grid Edit"));
-				FComponentMaterialSet OutputMaterialSet;
-				OutputMaterialSet.Materials = CurrentMeshMaterials;
-
-				UE::ToolTarget::CommitDynamicMeshUpdate(Target, *CurrentMesh, true,
-					FConversionToMeshDescriptionOptions(), &OutputMaterialSet);
-				GetToolManager()->EndUndoTransaction();
+				return;
 			}
-			else if (!Target->IsValid() && CurrentMesh->TriangleCount() > 0)
+
+			UPrimitiveComponent* Component = UE::ToolTarget::GetTargetComponent(Target);
+			if (ensure(Component))
 			{
-				EAppReturnType::Type Ret = FMessageDialog::Open(EAppMsgType::YesNo,
-					LOCTEXT("RecreateAssetQuestion", "The underlying asset that this tool was "
-						"operating on seems to no longer be valid (it was likely forcibly removed). "
-						"Would you like to recreate a new asset from the tool's current working "
-						"mesh? Selecting \"No\" or closing this window will discard the tool's "
-						"current work."), 
-					LOCTEXT("RecreateAssetTitle", "Recreate Mesh Asset?"));
-				if (Ret == EAppReturnType::Yes)
-				{
-					bCreatingNewAsset = true;
-				}
+				GetToolManager()->BeginUndoTransaction(LOCTEXT("DestroyComponentTransaction", "Delete Mesh Component"));
+				UE::ToolTarget::GetTargetComponent(Target)->DestroyComponent();
+				GetToolManager()->EndUndoTransaction();
 			}
 		}
 
-		if (bCreatingNewAsset)
+		return;
+	}
+
+	bool bCreatingNewAsset = !Target && CurrentMesh->TriangleCount() > 0;
+	if (Target)
+	{
+		if (Target->IsValid())
 		{
-			GetToolManager()->BeginUndoTransaction(LOCTEXT("CubeGridToolCreateTransactionName", "Cube Grid Create New"));
+			GetToolManager()->BeginUndoTransaction(LOCTEXT("CubeGridToolEditTransactionName", "Cube Grid Edit"));
+			FComponentMaterialSet OutputMaterialSet;
+			OutputMaterialSet.Materials = CurrentMeshMaterials;
 
-			FCreateMeshObjectParams NewMeshObjectParams;
-			NewMeshObjectParams.TargetWorld = TargetWorld;
-			NewMeshObjectParams.Transform = (FTransform)CurrentMeshTransform;
-			NewMeshObjectParams.BaseName = TEXT("CubeGridToolOutput");
-			NewMeshObjectParams.Materials = CurrentMeshMaterials;
-			NewMeshObjectParams.SetMesh(CurrentMesh.Get());
-			OutputTypeProperties->ConfigureCreateMeshObjectParams(NewMeshObjectParams);
-			FCreateMeshObjectResult Result = UE::Modeling::CreateMeshObject(GetToolManager(), MoveTemp(NewMeshObjectParams));
-			if (Result.IsOK() && Result.NewActor != nullptr)
+			UE::ToolTarget::CommitDynamicMeshUpdate(Target, *CurrentMesh, true,
+				FConversionToMeshDescriptionOptions(), &OutputMaterialSet);
+			GetToolManager()->EndUndoTransaction();
+			return;
+		}
+		else if (!Target->IsValid() && CurrentMesh->TriangleCount() > 0)
+		{
+			EAppReturnType::Type Ret = FMessageDialog::Open(EAppMsgType::YesNo,
+				LOCTEXT("RecreateAssetQuestion", "The underlying asset that this tool was "
+					"operating on seems to no longer be valid (it was likely forcibly removed). "
+					"Would you like to recreate a new asset from the tool's current working "
+					"mesh? Selecting \"No\" or closing this window will discard the tool's "
+					"current work."),
+				LOCTEXT("RecreateAssetTitle", "Recreate Mesh Asset?"));
+			if (Ret == EAppReturnType::Yes)
 			{
-				ToolSelectionUtil::SetNewActorSelection(GetToolManager(), Result.NewActor);
+				bCreatingNewAsset = true;
 			}
+		}
+	}
 
-			GetToolManager()->EndUndoTransaction();
+	if (bCreatingNewAsset)
+	{
+		GetToolManager()->BeginUndoTransaction(LOCTEXT("CubeGridToolCreateTransactionName", "Cube Grid Create New"));
+
+		FCreateMeshObjectParams NewMeshObjectParams;
+		NewMeshObjectParams.TargetWorld = TargetWorld;
+		NewMeshObjectParams.Transform = (FTransform)CurrentMeshTransform;
+		NewMeshObjectParams.BaseName = TEXT("CubeGridToolOutput");
+		NewMeshObjectParams.Materials = CurrentMeshMaterials;
+		NewMeshObjectParams.SetMesh(CurrentMesh.Get());
+		OutputTypeProperties->ConfigureCreateMeshObjectParams(NewMeshObjectParams);
+		FCreateMeshObjectResult Result = UE::Modeling::CreateMeshObject(GetToolManager(), MoveTemp(NewMeshObjectParams));
+		if (Result.IsOK() && Result.NewActor != nullptr && bSetSelection)
+		{
+			ToolSelectionUtil::SetNewActorSelection(GetToolManager(), Result.NewActor);
 		}
+
+		GetToolManager()->EndUndoTransaction();
 	}
+}
 
-	Preview->OnOpCompleted.RemoveAll(this);
-	Preview->Shutdown();
+void UCubeGridTool::AcceptToolAndStartNew()
+{
+	using namespace CubeGridToolLocals;
 
-	if (LineSets)
+	if (Mode == EMode::Corner)
 	{
-		LineSets->Disconnect();
-		LineSets = nullptr;
+		ApplyCornerMode(true);
 	}
 
-	GridGizmoAlignmentMechanic->Shutdown();
+	if (!bChangesMade)
+	{
+		return;
+	}
 
-	GetToolManager()->GetPairedGizmoManager()->DestroyAllGizmosByOwner(this);
+	FText TransactionText = LOCTEXT("AcceptAndStartNewTransaction", "Cube Grid Accept and Start New");
+	GetToolManager()->BeginUndoTransaction(TransactionText);
+
+	OutputCurrentResults(/*bSetSelection =*/ false);
+
+	GetToolManager()->EmitObjectChange(this, MakeUnique<FCubeGridChangesMadeChange>(false), TransactionText);
+
+	if (Target)
+	{
+		UE::ToolTarget::SetSourceObjectVisible(Target.Get(), true);
+		GetToolManager()->EmitObjectChange(this, MakeUnique<FCubeGridTargetResetChange>(Target), TransactionText);
+	}
+
+	// Clear out the current mesh
+	if (CurrentMesh->TriangleCount() > 0)
+	{
+		// TODO: There are probably better ways to do this than a dynamic mesh change tracker. For
+		// instance we could have a change that keeps a weak pointer to the mesh and keep a shared
+		// pointer in the tool, so that the mesh is thrown away once the tool exits (and the change
+		// is expired). Though we might need to be a bit more careful with retargeting the spatial, etc.
+		// For now, we do this easy approach.
+		TArray<int32> AllTids;
+		for (int32 Tid : CurrentMesh->TriangleIndicesItr())
+		{
+			AllTids.Add(Tid);
+		}
+		FDynamicMeshChangeTracker ChangeTracker(CurrentMesh.Get());
+		ChangeTracker.BeginChange();
+		ChangeTracker.SaveTriangles(AllTids, true /*bSaveVertices*/);
+		FDynamicMeshEditor Editor(CurrentMesh.Get());
+		Editor.RemoveTriangles(AllTids, true);
+		GetToolManager()->EmitObjectChange(this, MakeUnique<FCubeGridToolMeshChange>(ChangeTracker.EndChange()), TransactionText);
+	}
+
+	// Clear out any other structures/state
+	MeshSpatial->Build();
+	UpdateComputeInputs();
+	bWaitingToApplyPreview = false;
+	CurrentExtrudeAmount = 0;
+	Preview->CancelCompute();
+	Preview->PreviewMesh->UpdatePreview(CurrentMesh.Get());
+
+	GetToolManager()->EndUndoTransaction();
 }
 
 void UCubeGridTool::OnTick(float DeltaTime)
@@ -1319,6 +1488,9 @@ void UCubeGridTool::ApplyPreview()
 		// The transform might have changed if we added to an empty mesh
 		if (!FTransform(CurrentMeshTransform).Equals(Preview->PreviewMesh->GetTransform()))
 		{
+			GetToolManager()->EmitObjectChange(this, MakeUnique<FCubeGridMeshTransformChange>(
+				CurrentMeshTransform, Preview->PreviewMesh->GetTransform()), TransactionText);
+
 			CurrentMeshTransform = Preview->PreviewMesh->GetTransform();
 			GridGizmoAlignmentMechanic->InitializeDeformedMeshRayCast([this]() { return MeshSpatial.Get(); }, 
 				CurrentMeshTransform, nullptr);
@@ -2269,7 +2441,7 @@ void UCubeGridTool::SetGridPowerClamped(int32 GridPower)
 				}
 				else
 				{
-					// For FiveAndTen, we multiply  by 2 half the time and by 5 the second half, rounding up for 2's.
+					// For FiveAndTen, we multiply by 2 half the time and by 5 the second half, rounding up for 2's.
 					uint8 FloorHalfGridPower = GridPower / 2;
 					uint32 TwoMultiplier = static_cast<uint32>(1) << (GridPower - FloorHalfGridPower);
 					return TwoMultiplier * FMath::Pow(5.0, static_cast<double>(FloorHalfGridPower));
@@ -2310,14 +2482,6 @@ void UCubeGridToolActions::PostAction(ECubeGridToolAction Action)
 	}
 }
 
-void UCubeGridDuringActivityActions::PostAction(ECubeGridToolAction Action)
-{
-	if (ParentTool.IsValid())
-	{
-		ParentTool->RequestAction(Action);
-	}
-}
-
 void UCubeGridTool::RequestAction(ECubeGridToolAction ActionType)
 {
 	if (PendingAction == ECubeGridToolAction::NoAction)
@@ -2359,19 +2523,6 @@ void UCubeGridTool::ApplyAction(ECubeGridToolAction ActionType)
 	//case ECubeGridToolAction::FitGrid:
 	//	StartFitGrid();
 	//	break;
-	case ECubeGridToolAction::Done:
-		if (Mode == EMode::Corner)
-		{
-			ApplyCornerMode();
-		}
-		else if (Mode == EMode::FitGrid)
-		{
-			//CancelFitGrid();
-		}
-		break;
-	case ECubeGridToolAction::Cancel:
-		RevertToDefaultMode();
-		break;
 	case ECubeGridToolAction::ResetFromActor:
 		if (ToolActions->GridSourceActor)
 		{
@@ -2380,6 +2531,9 @@ void UCubeGridTool::ApplyAction(ECubeGridToolAction ActionType)
 			UpdateGridGizmo(TransformToUse);
 		}
 		break;
+	case ECubeGridToolAction::AcceptAndStartNew:
+		AcceptToolAndStartNew();
+		break;
 	}
 }
 
@@ -2541,7 +2695,38 @@ void UCubeGridTool::StartCornerMode()
 	Mode = EMode::Corner;
 
 	SetToolPropertySourceEnabled(ToolActions, false);
-	SetToolPropertySourceEnabled(DuringActivityActions, true);
+	// Customize the tool accept/cancel buttons to the current activity.
+	if (IToolHostCustomizationAPI* ButtonCustomizer = IToolHostCustomizationAPI::Find(GetToolManager()).GetInterface())
+	{
+		IToolHostCustomizationAPI::FAcceptCancelButtonOverrideParams Params;
+		Params.Label = LOCTEXT("CornerModeActivityLabel","Corner Mode");
+		Params.OverrideAcceptButtonText = LOCTEXT("DoneButton", "Done");
+		Params.OverrideAcceptButtonTooltip = LOCTEXT("DoneButtonTooltip", "Apply the current change and exit corner mode.");
+		Params.OverrideCancelButtonText = LOCTEXT("CancelButton", "Cancel");
+		Params.OverrideCancelButtonTooltip = LOCTEXT("CancelButtonTooltip", "Cancel the current change and exit corner mode.");
+		Params.CanAccept = [this]() { return true; };
+		Params.OnAcceptCancelTriggered = [this](bool bAccept)
+		{
+			if (bAccept)
+			{
+				if (Mode == EMode::Corner)
+				{
+					ApplyCornerMode();
+				}
+				else if (Mode == EMode::FitGrid)
+				{
+					//AcceptFitGrid();
+				}
+			}
+			else
+			{
+				RevertToDefaultMode();
+			}
+			return FReply::Handled();
+		};
+
+		ButtonCustomizer->RequestAcceptCancelButtonOverride(Params);
+	}
 
 	Settings->bInCornerMode = true;
 	NotifyOfPropertyChangeByTool(Settings);
@@ -2589,7 +2774,7 @@ void UCubeGridTool::ApplyCornerMode(bool bDontWaitForTick)
 	Mode = EMode::PushPull;
 	GetToolManager()->DisplayMessage(PushPullModeMessage, EToolMessageLevel::UserNotification);
 	SetToolPropertySourceEnabled(ToolActions, true);
-	SetToolPropertySourceEnabled(DuringActivityActions, false);
+	ClearViewportButtonCustomization();
 
 	Settings->bInCornerMode = false;
 	NotifyOfPropertyChangeByTool(Settings);
@@ -2606,7 +2791,7 @@ void UCubeGridTool::CancelCornerMode()
 	Mode = EMode::PushPull;
 	GetToolManager()->DisplayMessage(PushPullModeMessage, EToolMessageLevel::UserNotification);
 	SetToolPropertySourceEnabled(ToolActions, true);
-	SetToolPropertySourceEnabled(DuringActivityActions, false);
+	ClearViewportButtonCustomization();
 
 	CurrentExtrudeAmount = 0;
 	InvalidatePreview();
@@ -2649,11 +2834,25 @@ void UCubeGridTool::RevertToDefaultMode()
 	}
 }
 
+void UCubeGridTool::ClearViewportButtonCustomization()
+{
+	if (IToolHostCustomizationAPI* ButtonCustomizer = IToolHostCustomizationAPI::Find(GetToolManager()).GetInterface())
+	{
+		ButtonCustomizer->ClearButtonOverrides();
+	}
+}
+
 void UCubeGridTool::SetChangesMade(bool bChangesMadeIn)
 {
 	bChangesMade = bChangesMadeIn;
 }
 
+void UCubeGridTool::SetCurrentMeshTransform(const FTransform& TransformIn)
+{
+	CurrentMeshTransform = TransformIn;
+	InvalidatePreview();
+}
+
 // For use by Undo/Redo during corner mode
 void UCubeGridTool::SetCurrentExtrudeAmount(int32 ExtrudeAmount)
 {
diff --git a/Engine/Plugins/Experimental/MeshModelingToolsetExp/Source/MeshModelingToolsExp/Public/CubeGridTool.h b/Engine/Plugins/Experimental/MeshModelingToolsetExp/Source/MeshModelingToolsExp/Public/CubeGridTool.h
--- a/Engine/Plugins/Experimental/MeshModelingToolsetExp/Source/MeshModelingToolsExp/Public/CubeGridTool.h
+++ b/Engine/Plugins/Experimental/MeshModelingToolsetExp/Source/MeshModelingToolsExp/Public/CubeGridTool.h
@@ -97,16 +97,17 @@ public:
 
 	//~ Unfortunately it seems that there isn't a way to make the automatic tooltip
 	//~ here platform-specific, so we can't be specific about the hotkeys
-	/** Determines cube grid scale. Can also be adjusted with hotkeys. */
+	/** Determines cube grid scale. Can also be adjusted with hotkeys. This changes Current
+	 Block Size to match the current Base Block Size and Grid Power (differently depending
+	 on the Power of Two setting). */
 	UPROPERTY(EditAnywhere, Category = Options, meta = (
 		EditCondition = "bAllowedToEditGrid", HideEditConditionToggle,
 		UIMin = "0", UIMax = "10", ClampMin = "0", ClampMax = "31"))
 	uint8 GridPower = DEFAULT_GRID_POWER;
 
 	/** 
-	 * Sets the size of a block at the current grid power. This is done by changing the 
-	 * base block size (i.e. the size at grid power 0) such that the target size is achieved at 
-	 * the the current value of Grid Power.
+	 * Sets the size of a block at the current grid power. This changes Base Block Size
+	 * such that the given value is achieved at the the current value of Grid Power.
 	 */
 	UPROPERTY(EditAnywhere, Category = Options, meta = (
 		EditCondition = "bAllowedToEditGrid", HideEditConditionToggle,
@@ -119,7 +120,7 @@ public:
 	int32 BlocksPerStep = 1;
 
 	/** 
-	 * When true, block sizes change by powers of two as grid power is changed. When false, block
+	 * When true, block sizes change by powers of two as Grid Power is changed. When false, block
 	 * sizes change by twos and fives, much like the default editor grid snapping options (for
 	 * instance, sizes might increase from 10 to 50 to 100 to 500).
 	 * Note that toggling this option will reset Grid Power and Current Block Size to default values.
@@ -130,8 +131,8 @@ public:
 	// Must match ClampMax in GridPower, used to make hotkeys not exceed it.
 	const uint8 MaxGridPower = 31;
 
-	/** Smallest block size to use in the grid. For instance, 3.125 results in
-	 blocks that are 100 sized at 5 power of two since 3.125 * 2^5 = 100. */
+	/** Smallest block size to use in the grid, i.e. the block size at Grid Power 0. This
+	 changes Current Block Size according to current value of Grid Power. */
 	UPROPERTY(EditAnywhere, Category = Options, AdvancedDisplay, meta = (
 		EditCondition = "bAllowedToEditGrid", HideEditConditionToggle,
 		UIMin = "0.1", UIMax = "10", ClampMin = "0.001", ClampMax = "1000"))
@@ -230,8 +231,7 @@ enum class ECubeGridToolAction
 	CornerMode,
 	// FitGrid,
 	ResetFromActor,
-	Done,
-	Cancel,
+	AcceptAndStartNew
 };
 
 UCLASS()
@@ -285,27 +285,16 @@ public:
 	 */
 	UFUNCTION(CallInEditor, Category = GridReinitialization)
 	void ResetGridFromActor () { PostAction(ECubeGridToolAction::ResetFromActor); }
-};
-
-UCLASS()
-class MESHMODELINGTOOLSEXP_API UCubeGridDuringActivityActions : public UInteractiveToolPropertySet
-{
-	GENERATED_BODY()
-
-public:
-	TWeakObjectPtr<UCubeGridTool> ParentTool;
-
-	void Initialize(UCubeGridTool* ParentToolIn) { ParentTool = ParentToolIn; }
-
-	void PostAction(ECubeGridToolAction Action);
 
-	/** Accept and complete current action. */
-	UFUNCTION(CallInEditor, Category = Actions)
-	void Done() { PostAction(ECubeGridToolAction::Done); }
-
-	/** Cancel and exit current action */
-	UFUNCTION(CallInEditor, Category = Actions)
-	void Cancel() { PostAction(ECubeGridToolAction::Cancel); }
+	/**
+	 * Accepts the output of the current tool and restarts it in "create new asset" mode.
+	 * Used to avoid creating one huge mesh when working on different parts of a level.
+	 * Note that undoing object creation will not remove any created assets on disc (i.e.
+	 * undoing the creation of a static mesh will remove the actor from the world but not
+	 * remove the static mesh asset from disc).
+	 */
+	UFUNCTION(CallInEditor, Category = AssetActions)
+	void AcceptAndStartNew() { PostAction(ECubeGridToolAction::AcceptAndStartNew); }
 };
 
 /** Tool that allows for blocky boolean operations on an orientable power-of-two grid. */
@@ -394,6 +383,7 @@ public:
 	virtual bool IsInCornerMode() const;
 	virtual void RevertToDefaultMode();
 	virtual void SetChangesMade(bool bChangesMadeIn);
+	virtual void SetCurrentMeshTransform(const FTransform& TransformIn);
 	virtual void SetCurrentExtrudeAmount(int32 ExtrudeAmount);
 	virtual void SetCornerSelection(bool CornerSelectedFlags[4]);
 
@@ -459,9 +449,6 @@ protected:
 	UPROPERTY()
 	TObjectPtr<UCubeGridToolActions> ToolActions = nullptr;
 
-	UPROPERTY()
-	TObjectPtr<UCubeGridDuringActivityActions> DuringActivityActions = nullptr;
-
 	UPROPERTY()
 	TObjectPtr<UNewMeshMaterialProperties> MaterialProperties = nullptr;
 
@@ -561,6 +548,7 @@ protected:
 
 	ECubeGridToolAction PendingAction = ECubeGridToolAction::NoAction;
 	void ApplyAction(ECubeGridToolAction ActionType);
+	void ClearViewportButtonCustomization();
 
 	int32 GridPowerWatcherIdx;
 	int32 BlockBaseSizeWatcherIdx;
@@ -620,4 +608,7 @@ protected:
 
 	// Used to see if we need to update the asset that we've been modifying
 	bool bChangesMade = false;
+
+	void OutputCurrentResults(bool bSetSelection);
+	void AcceptToolAndStartNew();
 };
