Modeling Tools: Fixes bug that ignored scale transforms when creating new objects while mirroring

diff --git a/Engine/Plugins/Experimental/MeshModelingToolsetExp/Source/MeshModelingToolsExp/Private/MirrorTool.cpp b/Engine/Plugins/Experimental/MeshModelingToolsetExp/Source/MeshModelingToolsExp/Private/MirrorTool.cpp
--- a/Engine/Plugins/Experimental/MeshModelingToolsetExp/Source/MeshModelingToolsExp/Private/MirrorTool.cpp
+++ b/Engine/Plugins/Experimental/MeshModelingToolsetExp/Source/MeshModelingToolsExp/Private/MirrorTool.cpp
@@ -393,6 +393,7 @@ void UMirrorTool::OnShutdown(EToolShutdownType ShutdownType)
 			FTransform Transform = (FTransform)UE::ToolTarget::GetLocalToWorldTransform(Targets[PreviewIndex]);
 			Results.Emplace(Preview->Shutdown());
 			MeshTransforms::ApplyTransformInverse(*(Results.Last().Mesh), MirrorTool_Local::OnlyScale(Transform), true);
+			Results.Last().Transform.SetScale(Transform.GetScale3D() * Results.Last().Transform.GetScale3D());
 		}
 
 		// Convert to output. This will also edit the selection.
