fix crash in draw spline tool if user recompiles their blueprint while the tool is open, then continues to use the tool

diff --git a/Engine/Plugins/Experimental/MeshModelingToolsetExp/Source/MeshModelingToolsEditorOnlyExp/Private/DrawSplineTool.cpp b/Engine/Plugins/Experimental/MeshModelingToolsetExp/Source/MeshModelingToolsEditorOnlyExp/Private/DrawSplineTool.cpp
--- a/Engine/Plugins/Experimental/MeshModelingToolsetExp/Source/MeshModelingToolsEditorOnlyExp/Private/DrawSplineTool.cpp
+++ b/Engine/Plugins/Experimental/MeshModelingToolsetExp/Source/MeshModelingToolsEditorOnlyExp/Private/DrawSplineTool.cpp
@@ -984,6 +984,12 @@ void UDrawSplineTool::OnTick(float DeltaTime)
 		PlaneMechanic->Tick(DeltaTime);
 	}
 
+	// check if we've invalidated the WorkingSpline
+	if (PreviewActor && !WorkingSpline.IsValid())
+	{
+		bNeedToRerunConstructionScript = true;
+	}
+
 	if (bNeedToRerunConstructionScript)
 	{
 		bNeedToRerunConstructionScript = false;
