make draw spline tool fall back to empty actor if it fails to create actor w/ the requested blueprint class

diff --git a/Engine/Plugins/Experimental/MeshModelingToolsetExp/Source/MeshModelingToolsEditorOnlyExp/Private/DrawSplineTool.cpp b/Engine/Plugins/Experimental/MeshModelingToolsetExp/Source/MeshModelingToolsEditorOnlyExp/Private/DrawSplineTool.cpp
--- a/Engine/Plugins/Experimental/MeshModelingToolsetExp/Source/MeshModelingToolsEditorOnlyExp/Private/DrawSplineTool.cpp
+++ b/Engine/Plugins/Experimental/MeshModelingToolsetExp/Source/MeshModelingToolsEditorOnlyExp/Private/DrawSplineTool.cpp
@@ -471,6 +471,11 @@ void UDrawSplineTool::TransitionOutputMode()
 				// Important that we don't use the default (RF_Transactional) here, or else we'll end up
 				// issuing an undo transaction in this call.
 				EObjectFlags::RF_Transient);
+			if (!PreviewActor)
+			{
+				FallbackSplinePlacement();
+				break;
+			}
 
 			WorkingSpline = GetOrCreateTargetSpline(PreviewActor, Settings->ExistingSplineIndexToReplace);
 			break;
@@ -627,6 +632,11 @@ void UDrawSplineTool::GenerateAsset()
 		AActor* NewActor = FActorFactoryAssetProxy::AddActorForAsset(
 			Settings->BlueprintToCreate.Get(),
 			/*bSelectActors =*/ false);
+		if (!NewActor)
+		{
+			CreateSplineInEmptyActor();
+			break;
+		}
 
 		OutputSpline = GetOrCreateTargetSpline(NewActor, Settings->ExistingSplineIndexToReplace, true);
 		CopySplineToSpline(*WorkingSpline, *OutputSpline, true);
