Fixes issue, brush radius does not update in realtime in viewport

diff --git a/Engine/Plugins/Experimental/MeshModelingToolsetExp/Source/ModelingEditorUI/Private/SkinWeightDetailCustomization.cpp b/Engine/Plugins/Experimental/MeshModelingToolsetExp/Source/ModelingEditorUI/Private/SkinWeightDetailCustomization.cpp
--- a/Engine/Plugins/Experimental/MeshModelingToolsetExp/Source/ModelingEditorUI/Private/SkinWeightDetailCustomization.cpp
+++ b/Engine/Plugins/Experimental/MeshModelingToolsetExp/Source/ModelingEditorUI/Private/SkinWeightDetailCustomization.cpp
@@ -168,6 +168,12 @@ void FSkinWeightDetailCustomization::AddBrushUI(IDetailLayoutBuilder& DetailBuil
 			.OnValueChanged_Lambda([this](EWeightEditOperation Mode)
 			{
 				SkinToolSettings->BrushMode = Mode;
+
+				// sync base tool settings with the mode specific saved values
+				// these are the source of truth for the base class viewport rendering of brush
+				SkinToolSettings->BrushRadius = SkinToolSettings->GetBrushConfig().Radius;
+				SkinToolSettings->BrushStrength = SkinToolSettings->GetBrushConfig().Strength;
+				SkinToolSettings->BrushFalloffAmount = SkinToolSettings->GetBrushConfig().Falloff;
 			})
 			+SSegmentedControl<EWeightEditOperation>::Slot(EWeightEditOperation::Add)
 			.Text(LOCTEXT("BrushAddMode", "Add"))
diff --git a/Engine/Source/Runtime/InteractiveToolsFramework/Private/BaseTools/BaseBrushTool.cpp b/Engine/Source/Runtime/InteractiveToolsFramework/Private/BaseTools/BaseBrushTool.cpp
--- a/Engine/Source/Runtime/InteractiveToolsFramework/Private/BaseTools/BaseBrushTool.cpp
+++ b/Engine/Source/Runtime/InteractiveToolsFramework/Private/BaseTools/BaseBrushTool.cpp
@@ -514,7 +514,7 @@ void UBaseBrushTool::UpdateBrushStampIndicator()
 			BrushStampIndicator->LineColor = BrushAdjusterBehavior->IsBrushBeingAdjusted() ? FLinearColor::White : FLinearColor::Green;	
 		}
 		
-		BrushStampIndicator->Update(LastBrushStamp.Radius, LastBrushStamp.WorldPosition, LastBrushStamp.WorldNormal, LastBrushStamp.Falloff);
+		BrushStampIndicator->Update(BrushProperties->BrushRadius, LastBrushStamp.WorldPosition, LastBrushStamp.WorldNormal, LastBrushStamp.Falloff);
 	}
 }
 
