Modeling Tools: When manually updating Brush Size in Group Paint Tool, new size is now reflected in viewport immediately, without having to click in viewport to see updated size

diff --git a/Engine/Plugins/Experimental/MeshModelingToolsetExp/Source/MeshModelingToolsExp/Private/MeshGroupPaintTool.cpp b/Engine/Plugins/Experimental/MeshModelingToolsetExp/Source/MeshModelingToolsExp/Private/MeshGroupPaintTool.cpp
--- a/Engine/Plugins/Experimental/MeshModelingToolsetExp/Source/MeshModelingToolsExp/Private/MeshGroupPaintTool.cpp
+++ b/Engine/Plugins/Experimental/MeshModelingToolsetExp/Source/MeshModelingToolsExp/Private/MeshGroupPaintTool.cpp
@@ -169,7 +169,7 @@ void UMeshGroupPaintTool::Setup()
 	FilterProperties->WatchProperty(FilterProperties->SubToolType,
 		[this](EMeshGroupPaintInteractionType NewType) { UpdateSubToolType(NewType); });
 	FilterProperties->WatchProperty(FilterProperties->BrushSize,
-		[this](float NewSize) { UMeshSculptToolBase::BrushProperties->BrushSize.AdaptiveSize = NewSize; });
+		[this](float NewSize) { UMeshSculptToolBase::BrushProperties->BrushSize.AdaptiveSize = NewSize; CalculateBrushRadius(); });
 	FilterProperties->WatchProperty(FilterProperties->bHitBackFaces,
 		[this](bool bNewValue) { UMeshSculptToolBase::BrushProperties->bHitBackFaces = bNewValue; });
 	FilterProperties->RestoreProperties(this);
@@ -327,13 +327,6 @@ TUniquePtr<FMeshSculptBrushOp>& UMeshGroupPaintTool::GetActiveBrushOp()
 	}
 }
 
-
-void UMeshGroupPaintTool::OnPropertyModified(UObject* PropertySet, FProperty* Property)
-{
-	CalculateBrushRadius();
-}
-
-
 void UMeshGroupPaintTool::IncreaseBrushRadiusAction()
 {
 	Super::IncreaseBrushRadiusAction();
diff --git a/Engine/Plugins/Experimental/MeshModelingToolsetExp/Source/MeshModelingToolsExp/Public/MeshGroupPaintTool.h b/Engine/Plugins/Experimental/MeshModelingToolsetExp/Source/MeshModelingToolsExp/Public/MeshGroupPaintTool.h
--- a/Engine/Plugins/Experimental/MeshModelingToolsetExp/Source/MeshModelingToolsExp/Public/MeshGroupPaintTool.h
+++ b/Engine/Plugins/Experimental/MeshModelingToolsetExp/Source/MeshModelingToolsExp/Public/MeshGroupPaintTool.h
@@ -309,8 +309,6 @@ public:
 
 	virtual bool OnUpdateHover(const FInputDeviceRay& DevicePos) override;
 
-	virtual void OnPropertyModified(UObject* PropertySet, FProperty* Property) override;
-
 	bool IsInBrushSubMode() const;
 
 	virtual void CommitResult(UBaseDynamicMeshComponent* Component, bool bModifiedTopology) override;
