toggle off the input mesh preview when showing the complex collision mesh, so that users can see the preview mesh

diff --git a/Engine/Plugins/Experimental/MeshModelingToolsetExp/Source/MeshModelingToolsExp/Private/Physics/ExtractCollisionGeometryTool.cpp b/Engine/Plugins/Experimental/MeshModelingToolsetExp/Source/MeshModelingToolsExp/Private/Physics/ExtractCollisionGeometryTool.cpp
--- a/Engine/Plugins/Experimental/MeshModelingToolsetExp/Source/MeshModelingToolsExp/Private/Physics/ExtractCollisionGeometryTool.cpp
+++ b/Engine/Plugins/Experimental/MeshModelingToolsetExp/Source/MeshModelingToolsExp/Private/Physics/ExtractCollisionGeometryTool.cpp
@@ -82,11 +82,22 @@ void UExtractCollisionGeometryTool::Setup()
 	Settings = NewObject<UExtractCollisionToolProperties>(this);
 	Settings->RestoreProperties(this);
 	AddToolPropertySource(Settings);
+	// Update input mesh visibility w/ logic that toggles it off when the complex preview is shown
+	auto UpdateInputMeshVisibility = [this]()
+	{
+		bool bShowInput = Settings->bShowInputMesh && (Settings->CollisionType != EExtractCollisionOutputType::Complex || !Settings->bShowPreview);
+		UE::ToolTarget::SetSourceObjectVisible(Target, bShowInput);
+	};
+	Settings->WatchProperty(Settings->CollisionType, [this, UpdateInputMeshVisibility](EExtractCollisionOutputType NewValue) { UpdateInputMeshVisibility(); });
 	Settings->WatchProperty(Settings->bWeldEdges, [this](bool bNewValue) { bResultValid = false; });
-	Settings->WatchProperty(Settings->bShowPreview, [this](bool bNewValue) { PreviewMesh->SetVisible(bNewValue); });
+	Settings->WatchProperty(Settings->bShowPreview, [this, UpdateInputMeshVisibility](bool bNewValue)
+	{
+		PreviewMesh->SetVisible(bNewValue); 
+		UpdateInputMeshVisibility();
+	});
 	PreviewMesh->SetVisible(Settings->bShowPreview);
-	Settings->WatchProperty(Settings->bShowInputMesh, [this](bool bNewValue) { UE::ToolTarget::SetSourceObjectVisible(Target, bNewValue); });
-	UE::ToolTarget::SetSourceObjectVisible(Target, Settings->bShowInputMesh);
+	Settings->WatchProperty(Settings->bShowInputMesh, [this, UpdateInputMeshVisibility](bool bNewValue) { UpdateInputMeshVisibility(); });
+	UpdateInputMeshVisibility();
 
 	VizSettings = NewObject<UCollisionGeometryVisualizationProperties>(this);
 	VizSettings->bEnableShowSolid = false; // This solid visualization is redundant to the 'show preview' option in the general settings section of this tool
diff --git a/Engine/Plugins/Experimental/MeshModelingToolsetExp/Source/MeshModelingToolsExp/Public/Physics/ExtractCollisionGeometryTool.h b/Engine/Plugins/Experimental/MeshModelingToolsetExp/Source/MeshModelingToolsExp/Public/Physics/ExtractCollisionGeometryTool.h
--- a/Engine/Plugins/Experimental/MeshModelingToolsetExp/Source/MeshModelingToolsExp/Public/Physics/ExtractCollisionGeometryTool.h
+++ b/Engine/Plugins/Experimental/MeshModelingToolsetExp/Source/MeshModelingToolsExp/Public/Physics/ExtractCollisionGeometryTool.h
@@ -53,7 +53,7 @@ public:
 	bool bShowPreview = false;
 
 	/** Show/Hide input mesh */
-	UPROPERTY(EditAnywhere, Category = Options)
+	UPROPERTY(EditAnywhere, Category = Options, meta = (EditCondition = "CollisionType != EExtractCollisionOutputType::Complex || !bShowPreview"))
 	bool bShowInputMesh = true;
 
 	/** Whether or not to weld coincident border edges of the Complex Collision Mesh (if possible) */
