Tool Presets: Changed flyout dialog in the New Preset window to use a list presentation instead of a tile presentation, allowing for more clarity in reading names on long preset collection labels.

diff --git a/Engine/Plugins/Editor/ModelingToolsEditorMode/Source/ModelingToolsEditorMode/Private/ModelingToolsEditorModeToolkit.cpp b/Engine/Plugins/Editor/ModelingToolsEditorMode/Source/ModelingToolsEditorMode/Private/ModelingToolsEditorModeToolkit.cpp
--- a/Engine/Plugins/Editor/ModelingToolsEditorMode/Source/ModelingToolsEditorMode/Private/ModelingToolsEditorModeToolkit.cpp
+++ b/Engine/Plugins/Editor/ModelingToolsEditorMode/Source/ModelingToolsEditorMode/Private/ModelingToolsEditorModeToolkit.cpp
@@ -1248,6 +1248,7 @@ TSharedRef<SWidget> FModelingToolsEditorModeToolkit::GetPresetCreateButtonConten
 					.AssetThumbnailLabel(EThumbnailLabel::AssetName)
 					.bForceShowPluginContent(true)
 					.bForceShowEngineContent(true)
+					.AssetViewType(EAssetViewType::List)
 				]
 				+ SHorizontalBox::Slot()
 				[
diff --git a/Engine/Plugins/Experimental/MeshModelingToolsetExp/Source/ModelingEditorUI/Private/ModelingWidgets/SToolInputAssetComboPanel.cpp b/Engine/Plugins/Experimental/MeshModelingToolsetExp/Source/ModelingEditorUI/Private/ModelingWidgets/SToolInputAssetComboPanel.cpp
--- a/Engine/Plugins/Experimental/MeshModelingToolsetExp/Source/ModelingEditorUI/Private/ModelingWidgets/SToolInputAssetComboPanel.cpp
+++ b/Engine/Plugins/Experimental/MeshModelingToolsetExp/Source/ModelingEditorUI/Private/ModelingWidgets/SToolInputAssetComboPanel.cpp
@@ -119,8 +119,9 @@ void SToolInputAssetComboPanel::Construct(const FArguments& InArgs)
 	EThumbnailLabel::Type TileThumbnailLabel = InArgs._AssetThumbnailLabel;
 	bool bForceShowEngineContent = InArgs._bForceShowEngineContent;
 	bool bForceShowPluginContent = InArgs._bForceShowPluginContent;
+	EAssetViewType::Type AssetViewType = InArgs._AssetViewType;
 
-	ComboButton->SetOnGetMenuContent(FOnGetContent::CreateLambda([this, TileThumbnailLabel,bForceShowEngineContent, bForceShowPluginContent]()
+	ComboButton->SetOnGetMenuContent(FOnGetContent::CreateLambda([this, TileThumbnailLabel,bForceShowEngineContent, bForceShowPluginContent, AssetViewType]()
 	{
 		// Configure filter for asset picker
 		FAssetPickerConfig Config;
@@ -129,7 +130,7 @@ void SToolInputAssetComboPanel::Construct(const FArguments& InArgs)
 		Config.Filter.ClassPaths.Add( this->AssetClassType->GetClassPathName() );
 		Config.Filter.bRecursivePaths = true;
 		//Config.Filter.bIncludeOnlyOnDiskAssets = true;
-		Config.InitialAssetViewType = EAssetViewType::Tile;
+		Config.InitialAssetViewType = AssetViewType;
 		Config.bFocusSearchBoxWhenOpened = true;
 		Config.bAllowNullSelection = true;
 		Config.bAllowDragging = false;
diff --git a/Engine/Plugins/Experimental/MeshModelingToolsetExp/Source/ModelingEditorUI/Public/ModelingWidgets/SToolInputAssetComboPanel.h b/Engine/Plugins/Experimental/MeshModelingToolsetExp/Source/ModelingEditorUI/Public/ModelingWidgets/SToolInputAssetComboPanel.h
--- a/Engine/Plugins/Experimental/MeshModelingToolsetExp/Source/ModelingEditorUI/Public/ModelingWidgets/SToolInputAssetComboPanel.h
+++ b/Engine/Plugins/Experimental/MeshModelingToolsetExp/Source/ModelingEditorUI/Public/ModelingWidgets/SToolInputAssetComboPanel.h
@@ -10,6 +10,7 @@
 #include "AssetRegistry/AssetData.h"
 #include "CollectionManagerTypes.h"
 #include "AssetThumbnail.h"
+#include "IContentBrowserSingleton.h"
 
 class SComboButton;
 class FAssetThumbnailPool;
@@ -20,7 +21,6 @@ class ITableRow;
 class STableViewBase;
 class SToolInputAssetPicker;
 
-
 /**
 * SToolInputAssetComboPanel provides a similar UI to SComboPanel but
 * specifically for picking Assets. The standard widget is a SComboButton
@@ -70,6 +70,7 @@ public:
 		, _AssetThumbnailLabel(EThumbnailLabel::NoLabel)
 		, _bForceShowEngineContent(false)
 		, _bForceShowPluginContent(false)
+		, _AssetViewType(EAssetViewType::Tile)
 	{
 		}
 
@@ -112,6 +113,9 @@ public:
 		/** Indicates if plugin content should always be shown */
 		SLATE_ARGUMENT(bool, bForceShowPluginContent)
 
+		/** Sets the type of display we want the asset view to take. */
+		SLATE_ARGUMENT(EAssetViewType::Type, AssetViewType)
+
 	SLATE_END_ARGS()
 
 
