Modeling Tools: Displace Tool Selection Support for 'Flat' Subdivision type

diff --git a/Engine/Plugins/Experimental/MeshModelingToolsetExp/Source/MeshModelingToolsExp/Private/DisplaceMeshTool.cpp b/Engine/Plugins/Experimental/MeshModelingToolsetExp/Source/MeshModelingToolsExp/Private/DisplaceMeshTool.cpp
--- a/Engine/Plugins/Experimental/MeshModelingToolsetExp/Source/MeshModelingToolsExp/Private/DisplaceMeshTool.cpp
+++ b/Engine/Plugins/Experimental/MeshModelingToolsetExp/Source/MeshModelingToolsExp/Private/DisplaceMeshTool.cpp
@@ -26,6 +26,9 @@
 // needed to disable normals recalculation on the underlying asset
 #include "AssetUtils/MeshDescriptionUtil.h"
 #include "Components/StaticMeshComponent.h"
+#include "DynamicSubmesh3.h"
+#include "Selections/GeometrySelection.h"
+#include "Selections/GeometrySelectionUtil.h"
 #include "Engine/StaticMesh.h"
 
 using namespace UE::Geometry;
@@ -300,6 +303,7 @@ namespace DisplaceMeshToolLocals{
 	{
 		int SubdivisionsCount = -1;
 		TSharedPtr<FIndexedWeightMap, ESPMode::ThreadSafe> WeightMap = nullptr;
+		TArray<int> SelectedTriangles;
 
 		// Optional selection parameters
 		TOptional<EDisplaceMeshToolTriangleSelectionType> SelectionType;
@@ -354,7 +358,8 @@ namespace DisplaceMeshToolLocals{
 
 	void FSubdivideMeshOp::CalculateResult(FProgressCancel* ProgressCancel)
 	{
-		int SubdivisionsCount = Parameters.SubdivisionsCount;
+		const int SubdivisionsCount = Parameters.SubdivisionsCount;
+		bool bUsingSelection = Parameters.SelectedTriangles.Num() != ResultMesh->TriangleCount();
 		if (SubdivisionType == EDisplaceMeshToolSubdivisionType::Flat) 
 		{
 			if (Parameters.SelectionType.IsSet()) // user wants to only tessellate a subset of the triangles
@@ -364,7 +369,35 @@ namespace DisplaceMeshToolLocals{
 					Parameters.ActiveMaterialID.IsSet() && 
 					Parameters.ActiveMaterialID.GetValue() != INDEX_NONE)
 				{
-					TUniquePtr<FTessellationPattern> Pattern = FSelectiveTessellate::CreateConcentricRingsPatternFromMaterial(ResultMesh.Get(), SubdivisionsCount, Parameters.ActiveMaterialID.GetValue());;
+					TUniquePtr<FTessellationPattern> Pattern;
+					if (bUsingSelection)
+					{
+						Pattern = FSelectiveTessellate::CreateConcentricRingsPatternFromSelectionAndMaterial(ResultMesh.Get(), SubdivisionsCount, Parameters.ActiveMaterialID.GetValue(), Parameters.SelectedTriangles);
+					}
+					else
+					{
+						Pattern = FSelectiveTessellate::CreateConcentricRingsPatternFromMaterial(ResultMesh.Get(), SubdivisionsCount, Parameters.ActiveMaterialID.GetValue());
+					}
+					
+					FDynamicMesh3 OutMesh;
+					FSelectiveTessellate Tessellator(ResultMesh.Get(), &OutMesh);
+					Tessellator.Progress = ProgressCancel;
+					Tessellator.SetPattern(Pattern.Get());
+
+					VerticesToDisplace = MakeUnique<TArray<int>>();
+					Tessellator.TessInfo.SelectedVertices = VerticesToDisplace.Get();
+					
+					if (ensureMsgf(Tessellator.Validate() == EOperationValidationResult::Ok, TEXT("The tessellator parameters are invalid.")))
+					{
+						if (Tessellator.Compute())
+						{
+							*ResultMesh = MoveTemp(OutMesh);
+						}
+					}
+				} // EDisplaceMeshToolTriangleSelectionType::None but still tessellating only selected geometry
+				else if (bUsingSelection)
+				{
+					TUniquePtr<FTessellationPattern> Pattern =FSelectiveTessellate::CreateConcentricRingsTessellationPattern(ResultMesh.Get(), SubdivisionsCount, Parameters.SelectedTriangles);
 					
 					FDynamicMesh3 OutMesh;
 					FSelectiveTessellate Tessellator(ResultMesh.Get(), &OutMesh);
@@ -890,7 +923,7 @@ namespace DisplaceMeshToolLocals{
 /*
  * ToolBuilder
  */
-USingleSelectionMeshEditingTool* UDisplaceMeshToolBuilder::CreateNewTool(const FToolBuilderState& SceneState) const
+USingleTargetWithSelectionTool* UDisplaceMeshToolBuilder::CreateNewTool(const FToolBuilderState& SceneState) const
 {
 	return NewObject<UDisplaceMeshTool>(SceneState.ToolManager);
 }
@@ -969,6 +1002,24 @@ void UDisplaceMeshTool::Setup()
 
 	DynamicMeshComponent->SetTangentsType(EDynamicMeshComponentTangentsMode::AutoCalculated);
 	DynamicMeshComponent->SetMesh(UE::ToolTarget::GetDynamicMeshCopy(Target));
+
+	// retrieve the selected triangles
+	TArray<int> TrianglesSelectedInMesh;
+	if (HasGeometrySelection())
+	{
+		const FGeometrySelection& InputSelection = GetGeometrySelection();
+		UE::Geometry::EnumerateSelectionTriangles(InputSelection, *DynamicMeshComponent->GetMesh(),
+			[&](int32 TriangleID){ TrianglesSelectedInMesh.Add(TriangleID); });
+	}
+	else //if the whole object is selected, SelectedTriangles will be all the triangles in the Object
+	{
+		for (int32 TriangleID : DynamicMeshComponent->GetMesh()->TriangleIndicesItr())
+		{
+			TrianglesSelectedInMesh.Add(TriangleID);
+		}
+	}
+	NumTrianglesInSelection = TrianglesSelectedInMesh.Num();
+
 	OriginalMesh.Copy(*DynamicMeshComponent->GetMesh());
 	OriginalMeshSpatial.SetMesh(&OriginalMesh, true);
 
@@ -1054,6 +1105,7 @@ void UDisplaceMeshTool::Setup()
 	SubParameters.SubdivisionsCount = CommonProperties->Subdivisions;
 	SubParameters.WeightMap = ActiveWeightMap;
 	SubParameters.SelectionType = SelectiveTessellationProperties->SelectionType;
+	SubParameters.SelectedTriangles = MoveTemp(TrianglesSelectedInMesh);
 	
 	if (SelectiveTessellationProperties->ActiveMaterial.IsNone() == false) 
 	{
@@ -1148,7 +1200,7 @@ void UDisplaceMeshTool::ValidateSubdivisions()
 	bool bIsInitialized = (Subdivider != nullptr);
 
 	constexpr int MaxTriangles = 3000000;
-	double NumTriangles = OriginalMesh.MaxTriangleID();
+	double NumTriangles = NumTrianglesInSelection;
 	int MaxSubdivisions = (int)(FMath::Sqrt(MaxTriangles/NumTriangles) - 1);
 	if (CommonProperties->Subdivisions > MaxSubdivisions)
 	{
diff --git a/Engine/Plugins/Experimental/MeshModelingToolsetExp/Source/MeshModelingToolsExp/Public/DisplaceMeshTool.h b/Engine/Plugins/Experimental/MeshModelingToolsetExp/Source/MeshModelingToolsExp/Public/DisplaceMeshTool.h
--- a/Engine/Plugins/Experimental/MeshModelingToolsetExp/Source/MeshModelingToolsExp/Public/DisplaceMeshTool.h
+++ b/Engine/Plugins/Experimental/MeshModelingToolsetExp/Source/MeshModelingToolsExp/Public/DisplaceMeshTool.h
@@ -13,6 +13,7 @@
 #include "Spatial/SampledScalarField2.h"
 #include "UObject/NoExportTypes.h"
 #include "WeightMapUtil.h"
+#include "BaseTools/SingleTargetWithSelectionTool.h"
 #include "DisplaceMeshTool.generated.h"
 
 struct FMeshDescription;
@@ -289,18 +290,18 @@ public:
  * Builder for Simple Mesh Displacement Tool
  */
 UCLASS()
-class MESHMODELINGTOOLSEXP_API UDisplaceMeshToolBuilder : public USingleSelectionMeshEditingToolBuilder
+class MESHMODELINGTOOLSEXP_API UDisplaceMeshToolBuilder : public USingleTargetWithSelectionToolBuilder
 {
 	GENERATED_BODY()
 public:
-	USingleSelectionMeshEditingTool* CreateNewTool(const FToolBuilderState& SceneState) const final;
+	virtual USingleTargetWithSelectionTool* CreateNewTool(const FToolBuilderState& SceneState) const override final;
 };
 
 /**
  * Simple Mesh Displacement Tool
  */
 UCLASS()
-class MESHMODELINGTOOLSEXP_API UDisplaceMeshTool : public USingleSelectionMeshEditingTool
+class MESHMODELINGTOOLSEXP_API UDisplaceMeshTool : public USingleTargetWithSelectionTool
 {
 	GENERATED_BODY()
 public:
@@ -365,6 +366,7 @@ private:
 
 	TSharedPtr<UE::Geometry::FDynamicMesh3, ESPMode::ThreadSafe> SubdividedMesh = nullptr;
 	TSharedPtr<TArray<int>, ESPMode::ThreadSafe> VerticesToDisplace = nullptr;
+	int32 NumTrianglesInSelection;
 	
 	UPROPERTY()
 	TObjectPtr<AInternalToolFrameworkActor> PreviewMeshActor = nullptr;
diff --git a/Engine/Plugins/Runtime/GeometryProcessing/Source/DynamicMesh/Private/Operations/SelectiveTessellate.cpp b/Engine/Plugins/Runtime/GeometryProcessing/Source/DynamicMesh/Private/Operations/SelectiveTessellate.cpp
--- a/Engine/Plugins/Runtime/GeometryProcessing/Source/DynamicMesh/Private/Operations/SelectiveTessellate.cpp
+++ b/Engine/Plugins/Runtime/GeometryProcessing/Source/DynamicMesh/Private/Operations/SelectiveTessellate.cpp
@@ -2057,6 +2057,29 @@ TUniquePtr<FTessellationPattern> FSelectiveTessellate::CreateConcentricRingsPatt
 	return FSelectiveTessellate::CreateConcentricRingsTessellationPattern(InMesh, InTessellationLevel, TriangleList);
 }
 
+TUniquePtr<FTessellationPattern> FSelectiveTessellate::CreateConcentricRingsPatternFromSelectionAndMaterial(const FDynamicMesh3* InMesh,
+																								 const int InTessellationLevel,
+																								 const int MaterialID,
+																								 const TArray<int>& SelectedTriangles) 
+{
+	if (InMesh->HasAttributes() == false || InMesh->Attributes()->HasMaterialID() == false || SelectedTriangles.IsEmpty())
+	{
+		return nullptr;
+	}
+
+	TArray<int> SelectedTrianglesWithMaterial;
+	for (const int TID : SelectedTriangles)
+	{
+		const int TriangleMaterialID = InMesh->Attributes()->GetMaterialID()->GetValue(TID);
+		if (TriangleMaterialID == MaterialID)
+		{
+			SelectedTrianglesWithMaterial.Add(TID);
+		}
+	}
+
+	return FSelectiveTessellate::CreateConcentricRingsTessellationPattern(InMesh, InTessellationLevel, SelectedTrianglesWithMaterial);
+}
+
 //
 // Inner uniform pattern
 //
diff --git a/Engine/Plugins/Runtime/GeometryProcessing/Source/DynamicMesh/Public/Operations/SelectiveTessellate.h b/Engine/Plugins/Runtime/GeometryProcessing/Source/DynamicMesh/Public/Operations/SelectiveTessellate.h
--- a/Engine/Plugins/Runtime/GeometryProcessing/Source/DynamicMesh/Public/Operations/SelectiveTessellate.h
+++ b/Engine/Plugins/Runtime/GeometryProcessing/Source/DynamicMesh/Public/Operations/SelectiveTessellate.h
@@ -426,6 +426,15 @@ public:
 																					 const int MaterialID);
 
 
+	/**
+	* Tessellate only triangles that are selected AND which belong to a given material (FDynamicMeshMaterialAttribute).
+	*
+	*  @return nullptr if InMesh has no attributes, no material attribute, if no material attribute with the given id exists, or if selection is invalid
+	*/
+	static TUniquePtr<FTessellationPattern> CreateConcentricRingsPatternFromSelectionAndMaterial(const FDynamicMesh3* InMesh,
+																					  const int InTessellationLevel,
+																					  const int MaterialID,
+																					  const TArray<int>& SelectedTriangles);
 
 
     /**
