make tools use a more consistent path to get dynamic meshes from tool targets

diff --git a/Engine/Plugins/Experimental/MeshModelingToolsetExp/Source/MeshModelingToolsEditorOnlyExp/Private/PolygonOnMeshTool.cpp b/Engine/Plugins/Experimental/MeshModelingToolsetExp/Source/MeshModelingToolsEditorOnlyExp/Private/PolygonOnMeshTool.cpp
--- a/Engine/Plugins/Experimental/MeshModelingToolsetExp/Source/MeshModelingToolsEditorOnlyExp/Private/PolygonOnMeshTool.cpp
+++ b/Engine/Plugins/Experimental/MeshModelingToolsetExp/Source/MeshModelingToolsEditorOnlyExp/Private/PolygonOnMeshTool.cpp
@@ -147,8 +147,7 @@ void UPolygonOnMeshTool::Setup()
 
 	// Convert input mesh description to dynamic mesh
 	OriginalDynamicMesh = MakeShared<FDynamicMesh3, ESPMode::ThreadSafe>();
-	FMeshDescriptionToDynamicMesh Converter;
-	Converter.Convert( UE::ToolTarget::GetMeshDescription(Target), *OriginalDynamicMesh);
+	*OriginalDynamicMesh = UE::ToolTarget::GetDynamicMeshCopy(Target);
 	// TODO: consider adding an AABB tree construction here?  tradeoff vs doing a raycast against full every time a param change happens ...
 
 	LastDrawnPolygon = FPolygon2d();
diff --git a/Engine/Plugins/Experimental/MeshModelingToolsetExp/Source/MeshModelingToolsEditorOnlyExp/Private/SkeletalMesh/SkinWeightsBindingTool.cpp b/Engine/Plugins/Experimental/MeshModelingToolsetExp/Source/MeshModelingToolsEditorOnlyExp/Private/SkeletalMesh/SkinWeightsBindingTool.cpp
--- a/Engine/Plugins/Experimental/MeshModelingToolsetExp/Source/MeshModelingToolsEditorOnlyExp/Private/SkeletalMesh/SkinWeightsBindingTool.cpp
+++ b/Engine/Plugins/Experimental/MeshModelingToolsetExp/Source/MeshModelingToolsEditorOnlyExp/Private/SkeletalMesh/SkinWeightsBindingTool.cpp
@@ -137,8 +137,7 @@ void USkinWeightsBindingTool::Setup()
 							  [HandlePropertyChange](int32) { HandlePropertyChange(true); });
 
 	OriginalMesh = MakeShared<FDynamicMesh3, ESPMode::ThreadSafe>();
-	FMeshDescriptionToDynamicMesh Converter;
-	Converter.Convert(UE::ToolTarget::GetMeshDescription(Targets[0]), *OriginalMesh);
+	*OriginalMesh = UE::ToolTarget::GetDynamicMeshCopy(Targets[0]);
 
 	// Enable or override vertex colors on the original mesh.
 	OriginalMesh->EnableAttributes();
diff --git a/Engine/Plugins/Experimental/MeshModelingToolsetExp/Source/MeshModelingToolsEditorOnlyExp/Private/SubdividePolyTool.cpp b/Engine/Plugins/Experimental/MeshModelingToolsetExp/Source/MeshModelingToolsEditorOnlyExp/Private/SubdividePolyTool.cpp
--- a/Engine/Plugins/Experimental/MeshModelingToolsetExp/Source/MeshModelingToolsEditorOnlyExp/Private/SubdividePolyTool.cpp
+++ b/Engine/Plugins/Experimental/MeshModelingToolsetExp/Source/MeshModelingToolsEditorOnlyExp/Private/SubdividePolyTool.cpp
@@ -213,8 +213,7 @@ void USubdividePolyTool::Setup()
 
 	bool bWantVertexNormals = false;
 	OriginalMesh = MakeShared<FDynamicMesh3, ESPMode::ThreadSafe>(bWantVertexNormals, false, false, false);
-	FMeshDescriptionToDynamicMesh Converter;
-	Converter.Convert(UE::ToolTarget::GetMeshDescription(Target), *OriginalMesh);
+	*OriginalMesh = UE::ToolTarget::GetDynamicMeshCopy(Target);
 
 	Topology = MakeShared<FGroupTopology>(OriginalMesh.Get(), false);
 	auto ShouldAddExtraCornerAtVert = [this](const FGroupTopology& GroupTopology, int32 Vid, const FIndex2i& AttachedGroupEdgeEids)
diff --git a/Engine/Plugins/Experimental/MeshModelingToolsetExp/Source/MeshModelingToolsExp/Private/EditNormalsTool.cpp b/Engine/Plugins/Experimental/MeshModelingToolsetExp/Source/MeshModelingToolsExp/Private/EditNormalsTool.cpp
--- a/Engine/Plugins/Experimental/MeshModelingToolsetExp/Source/MeshModelingToolsExp/Private/EditNormalsTool.cpp
+++ b/Engine/Plugins/Experimental/MeshModelingToolsetExp/Source/MeshModelingToolsExp/Private/EditNormalsTool.cpp
@@ -221,8 +221,7 @@ void UEditNormalsTool::UpdateNumPreviews()
 			OpFactory->Tool = this;
 			OpFactory->ComponentIndex = PreviewIdx;
 			OriginalDynamicMeshes[PreviewIdx] = MakeShared<FDynamicMesh3, ESPMode::ThreadSafe>();
-			FMeshDescriptionToDynamicMesh Converter;
-			Converter.Convert(UE::ToolTarget::GetMeshDescription(Targets[PreviewIdx]), *OriginalDynamicMeshes[PreviewIdx]);
+			*OriginalDynamicMeshes[PreviewIdx] = UE::ToolTarget::GetDynamicMeshCopy(Targets[PreviewIdx]);
 
 			UMeshOpPreviewWithBackgroundCompute* Preview = Previews.Add_GetRef(NewObject<UMeshOpPreviewWithBackgroundCompute>(OpFactory, "Preview"));
 			Preview->Setup(GetTargetWorld(), OpFactory);
diff --git a/Engine/Plugins/Experimental/MeshModelingToolsetExp/Source/MeshModelingToolsExp/Private/HoleFillTool.cpp b/Engine/Plugins/Experimental/MeshModelingToolsetExp/Source/MeshModelingToolsExp/Private/HoleFillTool.cpp
--- a/Engine/Plugins/Experimental/MeshModelingToolsetExp/Source/MeshModelingToolsExp/Private/HoleFillTool.cpp
+++ b/Engine/Plugins/Experimental/MeshModelingToolsetExp/Source/MeshModelingToolsExp/Private/HoleFillTool.cpp
@@ -133,8 +133,7 @@ void UHoleFillTool::Setup()
 
 	// create mesh to operate on
 	OriginalMesh = MakeShared<FDynamicMesh3, ESPMode::ThreadSafe>();
-	FMeshDescriptionToDynamicMesh Converter;
-	Converter.Convert(UE::ToolTarget::GetMeshDescription(Target), *OriginalMesh);
+	*OriginalMesh = UE::ToolTarget::GetDynamicMeshCopy(Target);
 
 	// initialize properties
 	Properties = NewObject<UHoleFillToolProperties>(this, TEXT("Hole Fill Settings"));
diff --git a/Engine/Plugins/Experimental/MeshModelingToolsetExp/Source/MeshModelingToolsExp/Private/LatticeDeformerTool.cpp b/Engine/Plugins/Experimental/MeshModelingToolsetExp/Source/MeshModelingToolsExp/Private/LatticeDeformerTool.cpp
--- a/Engine/Plugins/Experimental/MeshModelingToolsetExp/Source/MeshModelingToolsExp/Private/LatticeDeformerTool.cpp
+++ b/Engine/Plugins/Experimental/MeshModelingToolsetExp/Source/MeshModelingToolsExp/Private/LatticeDeformerTool.cpp
@@ -172,8 +172,7 @@ void ULatticeDeformerTool::Setup()
 		"Drag the lattice control points to deform the mesh"), EToolMessageLevel::UserNotification);
 
 	OriginalMesh = MakeShared<FDynamicMesh3, ESPMode::ThreadSafe>();
-	FMeshDescriptionToDynamicMesh Converter;
-	Converter.Convert(UE::ToolTarget::GetMeshDescription(Target), *OriginalMesh);
+	*OriginalMesh = UE::ToolTarget::GetDynamicMeshCopy(Target);
 
 	// Note: Mesh will be implicitly transformed to world space by transforming the lattice; we account for whether that would invert the mesh here
 	MeshTransforms::ReverseOrientationIfNeeded(*OriginalMesh, (Cast<IPrimitiveComponentBackedTarget>(Target)->GetWorldTransform()));
diff --git a/Engine/Plugins/Experimental/MeshModelingToolsetExp/Source/MeshModelingToolsExp/Private/MeshBoundaryToolBase.cpp b/Engine/Plugins/Experimental/MeshModelingToolsetExp/Source/MeshModelingToolsExp/Private/MeshBoundaryToolBase.cpp
--- a/Engine/Plugins/Experimental/MeshModelingToolsetExp/Source/MeshModelingToolsExp/Private/MeshBoundaryToolBase.cpp
+++ b/Engine/Plugins/Experimental/MeshModelingToolsetExp/Source/MeshModelingToolsExp/Private/MeshBoundaryToolBase.cpp
@@ -29,8 +29,7 @@ void UMeshBoundaryToolBase::Setup()
 
 	// create mesh to operate on
 	OriginalMesh = MakeShared<FDynamicMesh3, ESPMode::ThreadSafe>();
-	FMeshDescriptionToDynamicMesh Converter;
-	Converter.Convert( UE::ToolTarget::GetMeshDescription(Target), *OriginalMesh);
+	*OriginalMesh = UE::ToolTarget::GetDynamicMeshCopy(Target);
 
 	// initialize hit query
 	MeshSpatial.SetMesh(OriginalMesh.Get());
diff --git a/Engine/Plugins/Experimental/MeshModelingToolsetExp/Source/MeshModelingToolsExp/Private/MeshSpaceDeformerTool.cpp b/Engine/Plugins/Experimental/MeshModelingToolsetExp/Source/MeshModelingToolsExp/Private/MeshSpaceDeformerTool.cpp
--- a/Engine/Plugins/Experimental/MeshModelingToolsetExp/Source/MeshModelingToolsExp/Private/MeshSpaceDeformerTool.cpp
+++ b/Engine/Plugins/Experimental/MeshModelingToolsetExp/Source/MeshModelingToolsExp/Private/MeshSpaceDeformerTool.cpp
@@ -140,8 +140,7 @@ void UMeshSpaceDeformerTool::Setup()
 	// populate the OriginalDynamicMesh with a conversion of the input mesh.
 	{
 		OriginalDynamicMesh = MakeShared<FDynamicMesh3, ESPMode::ThreadSafe>();
-		FMeshDescriptionToDynamicMesh Converter;
-		Converter.Convert(UE::ToolTarget::GetMeshDescription(Target), *OriginalDynamicMesh);
+		*OriginalDynamicMesh = UE::ToolTarget::GetDynamicMeshCopy(Target);
 	}
 
 	IPrimitiveComponentBackedTarget* TargetComponent = Cast<IPrimitiveComponentBackedTarget>(Target);
diff --git a/Engine/Plugins/Experimental/MeshModelingToolsetExp/Source/MeshModelingToolsExp/Private/MirrorTool.cpp b/Engine/Plugins/Experimental/MeshModelingToolsetExp/Source/MeshModelingToolsExp/Private/MirrorTool.cpp
--- a/Engine/Plugins/Experimental/MeshModelingToolsetExp/Source/MeshModelingToolsExp/Private/MirrorTool.cpp
+++ b/Engine/Plugins/Experimental/MeshModelingToolsetExp/Source/MeshModelingToolsExp/Private/MirrorTool.cpp
@@ -202,8 +202,7 @@ void UMirrorTool::Setup()
 	{
 		// Convert into dynamic mesh
 		TSharedPtr<FDynamicMesh3, ESPMode::ThreadSafe> DynamicMesh = MakeShared<FDynamicMesh3, ESPMode::ThreadSafe>();
-		FMeshDescriptionToDynamicMesh Converter;
-		Converter.Convert(UE::ToolTarget::GetMeshDescription(Targets[i]), *DynamicMesh);
+		*DynamicMesh = UE::ToolTarget::GetDynamicMeshCopy(Targets[i]);
 		// Bake the scale part of the transform
 		FTransform Transform = (FTransform)UE::ToolTarget::GetLocalToWorldTransform(Targets[i]);
 		MeshTransforms::ApplyTransform(*DynamicMesh, MirrorTool_Local::OnlyScale(Transform), true);
diff --git a/Engine/Plugins/Experimental/MeshModelingToolsetExp/Source/MeshModelingToolsExp/Private/PlaneCutTool.cpp b/Engine/Plugins/Experimental/MeshModelingToolsetExp/Source/MeshModelingToolsExp/Private/PlaneCutTool.cpp
--- a/Engine/Plugins/Experimental/MeshModelingToolsetExp/Source/MeshModelingToolsExp/Private/PlaneCutTool.cpp
+++ b/Engine/Plugins/Experimental/MeshModelingToolsetExp/Source/MeshModelingToolsExp/Private/PlaneCutTool.cpp
@@ -81,12 +81,11 @@ void UPlaneCutTool::Setup()
 			EToolMessageLevel::UserWarning);
 	}
 
-	// Convert input mesh descriptions to dynamic mesh
+	// Convert input targets to dynamic mesh
 	for (int Idx = 0; Idx < Targets.Num(); Idx++)
 	{
 		FDynamicMesh3* OriginalDynamicMesh = new FDynamicMesh3;
-		FMeshDescriptionToDynamicMesh Converter;
-		Converter.Convert(UE::ToolTarget::GetMeshDescription(Targets[Idx]), *OriginalDynamicMesh);
+		*OriginalDynamicMesh = UE::ToolTarget::GetDynamicMeshCopy(Targets[Idx]);
 		OriginalDynamicMesh->EnableAttributes();
 		TDynamicMeshScalarTriangleAttribute<int>* SubObjectIDs = new TDynamicMeshScalarTriangleAttribute<int>(OriginalDynamicMesh);
 		SubObjectIDs->Initialize(0);
diff --git a/Engine/Plugins/Experimental/MeshModelingToolsetExp/Source/MeshModelingToolsExp/Private/ProjectToTargetTool.cpp b/Engine/Plugins/Experimental/MeshModelingToolsetExp/Source/MeshModelingToolsExp/Private/ProjectToTargetTool.cpp
--- a/Engine/Plugins/Experimental/MeshModelingToolsetExp/Source/MeshModelingToolsExp/Private/ProjectToTargetTool.cpp
+++ b/Engine/Plugins/Experimental/MeshModelingToolsetExp/Source/MeshModelingToolsExp/Private/ProjectToTargetTool.cpp
@@ -30,10 +30,9 @@ UMultiSelectionMeshEditingTool* UProjectToTargetToolBuilder::CreateNewTool(const
 void UProjectToTargetTool::Setup()
 {
 	// ProjectionTarget and ProjectionTargetSpatial are setup before calling the parent class's Setup
-	FMeshDescriptionToDynamicMesh ProjectionConverter;
 	check(Targets.Num() == 2);
 	ProjectionTarget = MakeUnique<FDynamicMesh3>();
-	ProjectionConverter.Convert(UE::ToolTarget::GetMeshDescription(Targets[1]), *ProjectionTarget);
+	*ProjectionTarget = UE::ToolTarget::GetDynamicMeshCopy(Targets[1]);
 	ProjectionTargetSpatial = MakeUnique<FDynamicMeshAABBTree3>(ProjectionTarget.Get(), true);
 
 	// Now setup parent RemeshMeshTool class
diff --git a/Engine/Plugins/Experimental/MeshModelingToolsetExp/Source/MeshModelingToolsExp/Private/RemoveOccludedTrianglesTool.cpp b/Engine/Plugins/Experimental/MeshModelingToolsetExp/Source/MeshModelingToolsExp/Private/RemoveOccludedTrianglesTool.cpp
--- a/Engine/Plugins/Experimental/MeshModelingToolsetExp/Source/MeshModelingToolsExp/Private/RemoveOccludedTrianglesTool.cpp
+++ b/Engine/Plugins/Experimental/MeshModelingToolsetExp/Source/MeshModelingToolsExp/Private/RemoveOccludedTrianglesTool.cpp
@@ -226,8 +226,7 @@ void URemoveOccludedTrianglesTool::SetupPreviews()
 			OpFactory->Tool = this;
 			OpFactory->PreviewIdx = PreviewIdx;
 			OriginalDynamicMeshes[PreviewIdx] = MakeShared<FDynamicMesh3, ESPMode::ThreadSafe>();
-			FMeshDescriptionToDynamicMesh Converter;
-			Converter.Convert(UE::ToolTarget::GetMeshDescription(Targets[TargetIdx]), *OriginalDynamicMeshes[PreviewIdx]);
+			*OriginalDynamicMeshes[PreviewIdx] = UE::ToolTarget::GetDynamicMeshCopy(Targets[TargetIdx]);
 
 			UMeshOpPreviewWithBackgroundCompute* Preview = Previews.Add_GetRef(NewObject<UMeshOpPreviewWithBackgroundCompute>(OpFactory, "Preview"));
 			Preview->Setup(GetTargetWorld(), OpFactory);
diff --git a/Engine/Plugins/Experimental/MeshModelingToolsetExp/Source/MeshModelingToolsExp/Private/SelfUnionMeshesTool.cpp b/Engine/Plugins/Experimental/MeshModelingToolsetExp/Source/MeshModelingToolsExp/Private/SelfUnionMeshesTool.cpp
--- a/Engine/Plugins/Experimental/MeshModelingToolsetExp/Source/MeshModelingToolsExp/Private/SelfUnionMeshesTool.cpp
+++ b/Engine/Plugins/Experimental/MeshModelingToolsetExp/Source/MeshModelingToolsExp/Private/SelfUnionMeshesTool.cpp
@@ -110,9 +110,7 @@ void USelfUnionMeshesTool::ConvertInputsAndSetPreviewMaterials(bool bSetPreviewM
 	bool bNeedColorAttr = false;
 	for (int ComponentIdx = 0; ComponentIdx < Targets.Num(); ComponentIdx++)
 	{
-		FDynamicMesh3 ComponentMesh;
-		FMeshDescriptionToDynamicMesh Converter;
-		Converter.Convert(UE::ToolTarget::GetMeshDescription(Targets[ComponentIdx]), ComponentMesh);
+		FDynamicMesh3 ComponentMesh = UE::ToolTarget::GetDynamicMeshCopy(Targets[ComponentIdx]);
 		bNeedColorAttr = bNeedColorAttr || (ComponentMesh.Attributes()->HasPrimaryColors());
 		// ensure materials and attributes are always enabled
 		ComponentMesh.EnableAttributes();
diff --git a/Engine/Plugins/Runtime/MeshModelingToolset/Source/MeshModelingTools/Private/CSGMeshesTool.cpp b/Engine/Plugins/Runtime/MeshModelingToolset/Source/MeshModelingTools/Private/CSGMeshesTool.cpp
--- a/Engine/Plugins/Runtime/MeshModelingToolset/Source/MeshModelingTools/Private/CSGMeshesTool.cpp
+++ b/Engine/Plugins/Runtime/MeshModelingToolset/Source/MeshModelingTools/Private/CSGMeshesTool.cpp
@@ -219,8 +219,7 @@ void UCSGMeshesTool::ConvertInputsAndSetPreviewMaterials(bool bSetPreviewMesh)
 	for (int ComponentIdx = 0; ComponentIdx < Targets.Num(); ComponentIdx++)
 	{
 		OriginalDynamicMeshes[ComponentIdx] = MakeShared<FDynamicMesh3, ESPMode::ThreadSafe>();
-		FMeshDescriptionToDynamicMesh Converter;
-		Converter.Convert(UE::ToolTarget::GetMeshDescription(Targets[ComponentIdx]), *OriginalDynamicMeshes[ComponentIdx]);
+		*OriginalDynamicMeshes[ComponentIdx] = UE::ToolTarget::GetDynamicMeshCopy(Targets[ComponentIdx]);
 
 		// ensure materials and attributes are always enabled
 		OriginalDynamicMeshes[ComponentIdx]->EnableAttributes();
diff --git a/Engine/Plugins/Runtime/MeshModelingToolset/Source/MeshModelingTools/Private/CutMeshWithMeshTool.cpp b/Engine/Plugins/Runtime/MeshModelingToolset/Source/MeshModelingTools/Private/CutMeshWithMeshTool.cpp
--- a/Engine/Plugins/Runtime/MeshModelingToolset/Source/MeshModelingTools/Private/CutMeshWithMeshTool.cpp
+++ b/Engine/Plugins/Runtime/MeshModelingToolset/Source/MeshModelingTools/Private/CutMeshWithMeshTool.cpp
@@ -117,8 +117,7 @@ void UCutMeshWithMeshTool::ConvertInputsAndSetPreviewMaterials(bool bSetPreviewM
 	for (int ComponentIdx = 0; ComponentIdx < Targets.Num(); ComponentIdx++)
 	{
 		TSharedPtr<FDynamicMesh3, ESPMode::ThreadSafe> Mesh = MakeShared<FDynamicMesh3, ESPMode::ThreadSafe>();
-		FMeshDescriptionToDynamicMesh Converter;
-		Converter.Convert(UE::ToolTarget::GetMeshDescription(Targets[ComponentIdx]), *Mesh);
+		*Mesh = UE::ToolTarget::GetDynamicMeshCopy(Targets[ComponentIdx]);
 
 		// ensure materials and attributes are always enabled
 		Mesh->EnableAttributes();
diff --git a/Engine/Plugins/Runtime/MeshModelingToolset/Source/MeshModelingTools/Private/UVLayoutTool.cpp b/Engine/Plugins/Runtime/MeshModelingToolset/Source/MeshModelingTools/Private/UVLayoutTool.cpp
--- a/Engine/Plugins/Runtime/MeshModelingToolset/Source/MeshModelingTools/Private/UVLayoutTool.cpp
+++ b/Engine/Plugins/Runtime/MeshModelingToolset/Source/MeshModelingTools/Private/UVLayoutTool.cpp
@@ -132,8 +132,7 @@ void UUVLayoutTool::UpdateNumPreviews()
 		for (int32 PreviewIdx = CurrentNumPreview; PreviewIdx < TargetNumPreview; PreviewIdx++)
 		{
 			OriginalDynamicMeshes[PreviewIdx] = MakeShared<FDynamicMesh3, ESPMode::ThreadSafe>();
-			FMeshDescriptionToDynamicMesh Converter;
-			Converter.Convert(UE::ToolTarget::GetMeshDescription(Targets[PreviewIdx]), *OriginalDynamicMeshes[PreviewIdx]);
+			*OriginalDynamicMeshes[PreviewIdx] = UE::ToolTarget::GetDynamicMeshCopy(Targets[PreviewIdx]);
 
 			Factories[PreviewIdx]= NewObject<UUVLayoutOperatorFactory>();
 			Factories[PreviewIdx]->OriginalMesh = OriginalDynamicMeshes[PreviewIdx];
diff --git a/Engine/Plugins/Runtime/MeshModelingToolset/Source/ModelingComponents/Private/BaseTools/BaseMeshProcessingTool.cpp b/Engine/Plugins/Runtime/MeshModelingToolset/Source/ModelingComponents/Private/BaseTools/BaseMeshProcessingTool.cpp
--- a/Engine/Plugins/Runtime/MeshModelingToolset/Source/ModelingComponents/Private/BaseTools/BaseMeshProcessingTool.cpp
+++ b/Engine/Plugins/Runtime/MeshModelingToolset/Source/ModelingComponents/Private/BaseTools/BaseMeshProcessingTool.cpp
@@ -84,8 +84,7 @@ void UBaseMeshProcessingTool::Setup()
 	ToolPropertyObjects.Add(this);
 
 	// populate the BaseMesh with a conversion of the input mesh.
-	FMeshDescriptionToDynamicMesh Converter;
-	Converter.Convert(UE::ToolTarget::GetMeshDescription(Target), InitialMesh);
+	InitialMesh = UE::ToolTarget::GetDynamicMeshCopy(Target);
 
 	if (RequiresScaleNormalization())
 	{
diff --git a/Engine/Plugins/Runtime/MeshModelingToolset/Source/ModelingComponents/Private/BaseTools/BaseVoxelTool.cpp b/Engine/Plugins/Runtime/MeshModelingToolset/Source/ModelingComponents/Private/BaseTools/BaseVoxelTool.cpp
--- a/Engine/Plugins/Runtime/MeshModelingToolset/Source/ModelingComponents/Private/BaseTools/BaseVoxelTool.cpp
+++ b/Engine/Plugins/Runtime/MeshModelingToolset/Source/ModelingComponents/Private/BaseTools/BaseVoxelTool.cpp
@@ -63,8 +63,7 @@ void UBaseVoxelTool::ConvertInputsAndSetPreviewMaterials(bool bSetPreviewMesh)
 	for (int ComponentIdx = 0; ComponentIdx < Targets.Num(); ComponentIdx++)
 	{
 		OriginalDynamicMeshes[ComponentIdx] = MakeShared<FDynamicMesh3, ESPMode::ThreadSafe>();
-		FMeshDescriptionToDynamicMesh Converter;
-		Converter.Convert(UE::ToolTarget::GetMeshDescription(Targets[ComponentIdx]), *OriginalDynamicMeshes[ComponentIdx]);
+		*OriginalDynamicMeshes[ComponentIdx] = UE::ToolTarget::GetDynamicMeshCopy(Targets[ComponentIdx]);
 	}
 
 	//if (bSetPreviewMesh)
