Skeleton Editing Tool: Set average normal with multiple vertices selected

#rnx

diff --git a/Engine/Plugins/Experimental/MeshModelingToolsetExp/Source/MeshModelingToolsEditorOnlyExp/Private/SkeletalMesh/SkeletonEditingTool.cpp b/Engine/Plugins/Experimental/MeshModelingToolsetExp/Source/MeshModelingToolsEditorOnlyExp/Private/SkeletalMesh/SkeletonEditingTool.cpp
--- a/Engine/Plugins/Experimental/MeshModelingToolsetExp/Source/MeshModelingToolsEditorOnlyExp/Private/SkeletalMesh/SkeletonEditingTool.cpp
+++ b/Engine/Plugins/Experimental/MeshModelingToolsetExp/Source/MeshModelingToolsEditorOnlyExp/Private/SkeletalMesh/SkeletonEditingTool.cpp
@@ -1148,7 +1148,8 @@ namespace SkeletonEditingTool
 
 		FTransform GetVerticesTransform(const TArray<int32>& InVertexIDs) const
 		{
-			switch(InVertexIDs.Num())
+			const int32 NumVertices = InVertexIDs.Num();
+			switch(NumVertices)
 			{
 			case 0:
 				return FTransform::Identity;
@@ -1162,16 +1163,37 @@ namespace SkeletonEditingTool
 				break;
 			}
 
-			// compute centroid only
-			FVector Centroid(0.0);
+			// compute centroid and normal
+			FVector Centroid(0.0), Normal(0.0);
 			for (const int32 VertexID: InVertexIDs)
 			{
 				Centroid += Mesh.GetVertexRef(VertexID);
+				Normal += FMeshNormals::ComputeVertexNormal(Mesh, VertexID);
+			}
+			
+			const double InvNumVertices = 1.0 / static_cast<double>(NumVertices);
+			Centroid *= InvNumVertices;
+			Normal *= InvNumVertices;
+			
+			const FVector3d& V0 = Mesh.GetVertexRef(InVertexIDs[0]);
+			FVector3d X = (Centroid - V0).GetSafeNormal();
+			if (X.IsNearlyZero())
+			{
+				for (int32 VertexID = 0; VertexID < NumVertices && X.IsNearlyZero(); VertexID++)
+				{
+					X = (Centroid - Mesh.GetVertexRef(InVertexIDs[VertexID])).GetSafeNormal();
+				}	
 			}
-			Centroid /= static_cast<double>(InVertexIDs.Num());
 
-			FTransform Transform(Centroid);
-			return MoveTemp(Transform);
+			FQuat Rotation = FQuat::Identity;
+			const bool bParallel = FMath::IsNearlyEqual(FMath::Abs(FVector::DotProduct(Normal, X)), 1.0);
+			if (!bParallel)
+			{
+				const FVector3d Y = Normal.Cross(X);
+				Rotation = FRotationMatrix::MakeFromXY(X, Y).ToQuat();
+			}
+			
+			return FTransform(Rotation, Centroid);
 		}
 		
 		FTransform GetEdgesTransform(const TArray<int32>& InEdgeIDs) const
