fix polygroup colors not rendering due to incorrect material in subdivide, tri select, and materials tools

diff --git a/Engine/Plugins/Experimental/MeshModelingToolsetExp/Source/MeshModelingToolsEditorOnlyExp/Private/SubdividePolyTool.cpp b/Engine/Plugins/Experimental/MeshModelingToolsetExp/Source/MeshModelingToolsEditorOnlyExp/Private/SubdividePolyTool.cpp
--- a/Engine/Plugins/Experimental/MeshModelingToolsetExp/Source/MeshModelingToolsEditorOnlyExp/Private/SubdividePolyTool.cpp
+++ b/Engine/Plugins/Experimental/MeshModelingToolsetExp/Source/MeshModelingToolsEditorOnlyExp/Private/SubdividePolyTool.cpp
@@ -348,7 +348,7 @@ void USubdividePolyTool::Setup()
 	{
 		if (bNewRenderGroups)
 		{
-			PreviewMesh->SetOverrideRenderMaterial(ToolSetupUtil::GetSelectionMaterial(GetToolManager()));
+			PreviewMesh->SetOverrideRenderMaterial(ToolSetupUtil::GetVertexColorMaterial(GetToolManager()));
 			PreviewMesh->SetTriangleColorFunction([](const FDynamicMesh3* Mesh, int TriangleID)
 			{
 				return LinearColors::SelectFColor(Mesh->GetTriangleGroup(TriangleID));
diff --git a/Engine/Plugins/Experimental/MeshModelingToolsetExp/Source/MeshModelingToolsExp/Private/MeshSelectionTool.cpp b/Engine/Plugins/Experimental/MeshModelingToolsetExp/Source/MeshModelingToolsExp/Private/MeshSelectionTool.cpp
--- a/Engine/Plugins/Experimental/MeshModelingToolsetExp/Source/MeshModelingToolsExp/Private/MeshSelectionTool.cpp
+++ b/Engine/Plugins/Experimental/MeshModelingToolsetExp/Source/MeshModelingToolsExp/Private/MeshSelectionTool.cpp
@@ -887,7 +887,7 @@ void UMeshSelectionTool::UpdateVisualization(bool bSelectionModified)
 	{
 		if (SelectionProps->FaceColorMode != EMeshFacesColorMode::None)
 		{
-			PreviewMesh->SetOverrideRenderMaterial(ToolSetupUtil::GetSelectionMaterial(GetToolManager()));
+			PreviewMesh->SetOverrideRenderMaterial(ToolSetupUtil::GetVertexColorMaterial(GetToolManager()));
 			PreviewMesh->SetTriangleColorFunction([this](const FDynamicMesh3* Mesh, int TriangleID)
 			{
 				return GetCurrentFaceColor(Mesh, TriangleID);
