Skeletal Mesh Binding Tool: do not copy color overlay when setting new weights
    - do not modify the original mesh during configuration to avoid deleting incomig data
    - use the mesh description instead of the dyna mesh to submit the new skin weights
    - use VertexColorMaterial as overlay material instead of changing the standard ones
    - update the color overlay with an empty selection when the tool is started. The notifier will update the visualization when the selection is changed.

#rnx

diff --git a/Engine/Plugins/Experimental/MeshModelingToolsetExp/Source/MeshModelingToolsEditorOnlyExp/Private/SkeletalMesh/SkinWeightsBindingTool.cpp b/Engine/Plugins/Experimental/MeshModelingToolsetExp/Source/MeshModelingToolsEditorOnlyExp/Private/SkeletalMesh/SkinWeightsBindingTool.cpp
--- a/Engine/Plugins/Experimental/MeshModelingToolsetExp/Source/MeshModelingToolsEditorOnlyExp/Private/SkeletalMesh/SkinWeightsBindingTool.cpp
+++ b/Engine/Plugins/Experimental/MeshModelingToolsetExp/Source/MeshModelingToolsEditorOnlyExp/Private/SkeletalMesh/SkinWeightsBindingTool.cpp
@@ -78,7 +78,6 @@ void USkinWeightsBindingTool::Setup()
 	{
 		USkeletalMesh* SkeletalMesh = SkeletalMeshComponent->GetSkeletalMeshAsset();
 		ReferenceSkeleton = SkeletalMesh->GetRefSkeleton();
-		Properties->CurrentBone = ReferenceSkeleton.GetBoneName(0);
 	}
 	
 	UE::ToolTarget::HideSourceObject(Targets[0]);
@@ -90,20 +89,11 @@ void USkinWeightsBindingTool::Setup()
 		UpdateVisualization();
 	});
 	
-	FComponentMaterialSet MaterialSet = UE::ToolTarget::GetMaterialSet(Targets[0]);
+	const FComponentMaterialSet MaterialSet = UE::ToolTarget::GetMaterialSet(Targets[0]);
+	Preview->ConfigureMaterials(MaterialSet.Materials, ToolSetupUtil::GetDefaultWorkingMaterial(GetToolManager()));
 
 	UMaterialInterface* VtxColorMaterial = GetToolManager()->GetContextQueriesAPI()->GetStandardMaterial(EStandardToolContextMaterials::VertexColorMaterial);
-	if (VtxColorMaterial != nullptr)
-	{
-		for (UMaterialInterface*& Material: MaterialSet.Materials)
-		{
-			Material = VtxColorMaterial;
-		}
-	}
-	
-	Preview->ConfigureMaterials( MaterialSet.Materials,
-								 ToolSetupUtil::GetDefaultWorkingMaterial(GetToolManager())
-	);
+	Preview->OverrideMaterial = VtxColorMaterial;
 
 	auto UpdateOccupancy = [this]()
 	{
@@ -136,15 +126,13 @@ void USkinWeightsBindingTool::Setup()
 	Properties->WatchProperty(Properties->VoxelResolution,
 							  [HandlePropertyChange](int32) { HandlePropertyChange(true); });
 
+	EditedMeshDescription = MakeUnique<FMeshDescription>();
+	*EditedMeshDescription = *UE::ToolTarget::GetMeshDescription(Targets[0]); 
+	
 	OriginalMesh = MakeShared<FDynamicMesh3, ESPMode::ThreadSafe>();
 	*OriginalMesh = UE::ToolTarget::GetDynamicMeshCopy(Targets[0]);
 
-	// Enable or override vertex colors on the original mesh.
-	OriginalMesh->EnableAttributes();
-	OriginalMesh->Attributes()->DisablePrimaryColors();
-	OriginalMesh->Attributes()->EnablePrimaryColors();
-	// Create an overlay that has no split elements, init with zero value.
-	OriginalMesh->Attributes()->PrimaryColors()->CreateFromPredicate([](int ParentVID, int TriIDA, int TriIDB){return true;}, 0.f);
+
 
 	Preview->PreviewMesh->SetTransform((FTransform) UE::ToolTarget::GetLocalToWorldTransform(Targets[0]));
 	Preview->PreviewMesh->SetTangentsMode(EDynamicMeshComponentTangentsMode::AutoCalculated);
@@ -292,14 +280,51 @@ TUniquePtr<UE::Geometry::FDynamicMeshOperator> USkinWeightsBindingTool::MakeNewO
 	return Op;
 }
 
+bool USkinWeightsBindingTool::UpdateSkinWeightsFromDynamicMesh(FDynamicMesh3& InResultMesh) const
+{
+	using namespace UE::AnimationCore;
+	using namespace UE::Geometry;
+
+	if (!InResultMesh.HasAttributes())
+	{
+		return false;
+	}
+	
+	const FDynamicMeshVertexSkinWeightsAttribute* SkinWeights = InResultMesh.Attributes()->GetSkinWeightsAttribute(FSkeletalMeshAttributes::DefaultSkinWeightProfileName);
+	if (!SkinWeights)
+	{
+		return false;
+	}
+
+	FSkeletalMeshAttributes MeshAttribs(*EditedMeshDescription);
+	FSkinWeightsVertexAttributesRef VertexSkinWeights = MeshAttribs.GetVertexSkinWeights();
+	
+	const int32 NumVertices = EditedMeshDescription->Vertices().Num();
+	for (int32 VertexIndex = 0; VertexIndex < NumVertices; VertexIndex++)
+	{
+		FBoneWeights Weights;
+		SkinWeights->GetValue(VertexIndex, Weights);
+		VertexSkinWeights.Set(FVertexID(VertexIndex), Weights);
+	}
+
+	return true;
+}
 
 void USkinWeightsBindingTool::GenerateAsset(const FDynamicMeshOpResult& Result)
 {
-	// TODO: Update FDynamicMeshToMeshDescription to allow update the skin weights only.
+	using namespace UE::Geometry;
+
+	FDynamicMesh3* ResultMesh = Result.Mesh.Get();
+	check(ResultMesh != nullptr);
+
+	if (!UpdateSkinWeightsFromDynamicMesh(*ResultMesh))
+	{
+		return;
+	}
+	
 	GetToolManager()->BeginUndoTransaction(LOCTEXT("SkinWeightsBindingToolTransactionName", "Create Rigid Binding"));
 
-	check(Result.Mesh.Get() != nullptr);
-	UE::ToolTarget::CommitMeshDescriptionUpdateViaDynamicMesh(Targets[0], *Result.Mesh.Get(), true);
+	UE::ToolTarget::CommitMeshDescriptionUpdate(Targets[0], EditedMeshDescription.Get());
 
 	GetToolManager()->EndUndoTransaction();	
 }
@@ -320,9 +345,9 @@ void USkinWeightsBindingTool::UpdateVisualization(bool bInForce)
 	using namespace UE::AnimationCore;
 	using namespace UE::Geometry;
 
-	const int32 RawBoneIndex = ReferenceSkeleton.FindRawBoneIndex(Properties->CurrentBone);
-	if ((bInForce || Preview->HaveValidNonEmptyResult()) && RawBoneIndex != INDEX_NONE)
+	if (bInForce || Preview->HaveValidNonEmptyResult())
 	{
+		const int32 RawBoneIndex = ReferenceSkeleton.FindRawBoneIndex(Properties->CurrentBone);
 		const FBoneIndexType BoneIndex = static_cast<FBoneIndexType>(RawBoneIndex);
 
 		// update mesh with new value colors
@@ -339,26 +364,34 @@ void USkinWeightsBindingTool::UpdateVisualization(bool bInForce)
 				ColorOverlay = InMesh.Attributes()->PrimaryColors();
 				ColorOverlay->CreateFromPredicate([](int /*ParentVID*/, int /*TriIDA*/, int /*TriIDB*/){return true;}, 0.f);
 			}
+
+			static const FVector4f DefaultColor(WeightToColor(0.f));
 			
 			FBoneWeights BoneWeights;
-			
 			for (const int32 ElementId : ColorOverlay->ElementIndicesItr())
 			{
-				const int32 VertexId = ColorOverlay->GetParentVertex(ElementId);
-				SkinWeights->GetValue(VertexId, BoneWeights);
-
-				float Weight = 0.0f;
-				for (const FBoneWeight BoneWeight: BoneWeights)
+				if (RawBoneIndex != INDEX_NONE)
 				{
-					if (BoneWeight.GetBoneIndex() == BoneIndex)
+					const int32 VertexId = ColorOverlay->GetParentVertex(ElementId);
+					SkinWeights->GetValue(VertexId, BoneWeights);
+
+					float Weight = 0.0f;
+					for (const FBoneWeight BoneWeight: BoneWeights)
 					{
-						Weight = BoneWeight.GetWeight();
-						break;
+						if (BoneWeight.GetBoneIndex() == BoneIndex)
+						{
+							Weight = BoneWeight.GetWeight();
+							break;
+						}
 					}
+			
+					const FVector4f Color(WeightToColor(Weight));
+					ColorOverlay->SetElement(ElementId, Color);
+				}
+				else
+				{
+					ColorOverlay->SetElement(ElementId, DefaultColor);
 				}
-				
-				const FVector4f Color(WeightToColor(Weight));
-				ColorOverlay->SetElement(ElementId, Color);
 			}
 		});
 	}
diff --git a/Engine/Plugins/Experimental/MeshModelingToolsetExp/Source/MeshModelingToolsEditorOnlyExp/Public/SkeletalMesh/SkinWeightsBindingTool.h b/Engine/Plugins/Experimental/MeshModelingToolsetExp/Source/MeshModelingToolsEditorOnlyExp/Public/SkeletalMesh/SkinWeightsBindingTool.h
--- a/Engine/Plugins/Experimental/MeshModelingToolsetExp/Source/MeshModelingToolsEditorOnlyExp/Public/SkeletalMesh/SkinWeightsBindingTool.h
+++ b/Engine/Plugins/Experimental/MeshModelingToolsetExp/Source/MeshModelingToolsEditorOnlyExp/Public/SkeletalMesh/SkinWeightsBindingTool.h
@@ -113,6 +113,8 @@ protected:
 	TSharedPtr<UE::Geometry::FOccupancyGrid3> Occupancy;
 	
 	TSharedPtr<UE::Geometry::FDynamicMesh3, ESPMode::ThreadSafe> OriginalMesh;
+	
+	TUniquePtr<FMeshDescription> EditedMeshDescription;
 
 	FReferenceSkeleton ReferenceSkeleton;
 	
@@ -124,6 +126,9 @@ protected:
 	UPROPERTY()
 	TWeakObjectPtr<USkeletalMeshEditorContextObjectBase> EditorContext = nullptr;
 
+	// updates the skin weights using the operator's result
+	bool UpdateSkinWeightsFromDynamicMesh(UE::Geometry::FDynamicMesh3& InResultMesh) const;
+	
 	// ISkeletalMeshEditionInterface
 	virtual void HandleSkeletalMeshModified(const TArray<FName>& InBoneNames, const ESkeletalMeshNotifyType InNotifyType) override;
 };
