MeshAutoUV - Fix array out of bound crash that could occur during the mesh description -> dynamic mesh -> mesh description roundtrip.
* The conversion to a dynamic mesh may result in less triangles as it could remove duplicates or degenerates.

diff --git a/Engine/Plugins/Experimental/MeshModelingToolsetExp/Source/GeometryProcessingAdapters/Private/GeometryProcessingAdaptersModule.cpp b/Engine/Plugins/Experimental/MeshModelingToolsetExp/Source/GeometryProcessingAdapters/Private/GeometryProcessingAdaptersModule.cpp
--- a/Engine/Plugins/Experimental/MeshModelingToolsetExp/Source/GeometryProcessingAdapters/Private/GeometryProcessingAdaptersModule.cpp
+++ b/Engine/Plugins/Experimental/MeshModelingToolsetExp/Source/GeometryProcessingAdapters/Private/GeometryProcessingAdaptersModule.cpp
@@ -28,7 +28,7 @@ void FGeometryProcessingAdaptersModule::StartupModule()
 	}
 
 	MeshAutoUV = MakeShared<FMeshAutoUVImpl>();
-	if (CombineMeshInstances.IsValid())
+	if (MeshAutoUV.IsValid())
 	{
 		IModularFeatures::Get().RegisterModularFeature(IGeometryProcessing_MeshAutoUV::GetModularFeatureName(), MeshAutoUV.Get());
 	}
diff --git a/Engine/Plugins/Experimental/MeshModelingToolsetExp/Source/GeometryProcessingAdapters/Private/MeshAutoUVImpl.cpp b/Engine/Plugins/Experimental/MeshModelingToolsetExp/Source/GeometryProcessingAdapters/Private/MeshAutoUVImpl.cpp
--- a/Engine/Plugins/Experimental/MeshModelingToolsetExp/Source/GeometryProcessingAdapters/Private/MeshAutoUVImpl.cpp
+++ b/Engine/Plugins/Experimental/MeshModelingToolsetExp/Source/GeometryProcessingAdapters/Private/MeshAutoUVImpl.cpp
@@ -3,6 +3,8 @@
 #include "GeometryProcessing/MeshAutoUVImpl.h"
 #include "MeshDescriptionToDynamicMesh.h"
 #include "DynamicMeshToMeshDescription.h"
+#include "MeshAttributes.h"
+#include "StaticMeshAttributes.h"
 #include "ParameterizationOps/ParameterizeMeshOp.h"
 
 
@@ -50,20 +52,43 @@ void FMeshAutoUVImpl::GenerateUVs(FMeshDescription& InOutMesh, const IGeometryPr
 	}
 
 	ParameterizeMeshOp.CalculateResult(nullptr);
+	
+	TUniquePtr<UE::Geometry::FDynamicMesh3> DynamicMeshWithUVs = ParameterizeMeshOp.ExtractResult();
+	if (DynamicMeshWithUVs)
+	{
+		// The DynamicMesh has now been populated with valid UVs for each vertex instance.
+		// Rather than performing a full conversion of the dynamic mesh back to a mesh description, 
+		// we can use the mapping in TriIDMap which was generated during the MeshDescription -> DynamicMesh conversion
+		// This is actually required as the conversion to a DynamicMesh may have removed duplicate triangles.
+		FDynamicMeshUVOverlay* DynamicMeshUVs = DynamicMeshWithUVs->Attributes()->PrimaryUV();
+		TVertexInstanceAttributesRef<FVector2f> MeshDescriptionUVs = InOutMesh.VertexInstanceAttributes().GetAttributesRef<FVector2f>(MeshAttribute::VertexInstance::TextureCoordinate);
+
+		// For each triangle of the dynamic mesh.
+		for (const int DynamicMeshTID : DynamicMeshWithUVs->TriangleIndicesItr())
+		{
+			// Map triangle from the DynamicMesh to it's equivalent in the MeshDescription.
+			int32 MeshDescriptionTID = MeshDescriptionToDynamicMesh.TriIDMap[DynamicMeshTID];
 
-	TUniquePtr<UE::Geometry::FDynamicMesh3> ResultMesh = ParameterizeMeshOp.ExtractResult();
+			// Grab the vertex instances (3) used by that triangle.
+			const FIndex3i DynamicMeshTriVIDs = DynamicMeshWithUVs->GetTriangle(DynamicMeshTID);
 
-	ResultsOut.ResultCode = IGeometryProcessing_MeshAutoUV::EResultCode::Success;
+			// For each triangle corner index...
+			for (int Index = 0; Index < 3; ++Index)
+			{
+				// Grab the computed UV in the dynamic mesh for that vertex instance.
+				const int DynamicMeshVID = DynamicMeshTriVIDs[Index];
+				FVector2f UV = DynamicMeshUVs->GetElementAtVertex(DynamicMeshTID, DynamicMeshVID);
 
-	FConversionToMeshDescriptionOptions ConversionOptions;
-	ConversionOptions.bSetPolyGroups = false;
-	ConversionOptions.bUpdatePositions = false;
-	ConversionOptions.bUpdateNormals = false;
-	ConversionOptions.bUpdateTangents = false;
-	ConversionOptions.bUpdateUVs = true;
-	ConversionOptions.bUpdateVtxColors = false;
-	ConversionOptions.bTransformVtxColorsSRGBToLinear = false;
+				// Assign the computed UV to the mesh description vertex instance.
+				FVertexInstanceID MeshDescriptionVID = InOutMesh.GetTriangleVertexInstance(MeshDescriptionTID, Index);
+				MeshDescriptionUVs.Set(MeshDescriptionVID, UV);
+			}
+		}
 
-	FDynamicMeshToMeshDescription Converter(ConversionOptions);
-	Converter.UpdateUsingConversionOptions(ResultMesh.Get(), InOutMesh);
+		ResultsOut.ResultCode = IGeometryProcessing_MeshAutoUV::EResultCode::Success;
+	}
+	else
+	{
+		ResultsOut.ResultCode = IGeometryProcessing_MeshAutoUV::EResultCode::UnknownError;
+	}
 }
\ No newline at end of file
diff --git a/Engine/Source/Runtime/StaticMeshDescription/Private/StaticMeshOperations.cpp b/Engine/Source/Runtime/StaticMeshDescription/Private/StaticMeshOperations.cpp
--- a/Engine/Source/Runtime/StaticMeshDescription/Private/StaticMeshOperations.cpp
+++ b/Engine/Source/Runtime/StaticMeshDescription/Private/StaticMeshOperations.cpp
@@ -2299,7 +2299,8 @@ bool FStaticMeshOperations::GenerateUV(const FMeshDescription& MeshDescription,
 	if (OutTexCoords.IsEmpty())
 	{
 		IGeometryProcessingInterfacesModule* GeomProcInterfaces = FModuleManager::Get().GetModulePtr<IGeometryProcessingInterfacesModule>("GeometryProcessingInterfaces");
-		if (GeomProcInterfaces)
+		IGeometryProcessing_MeshAutoUV* MeshAutoUV = GeomProcInterfaces ? GeomProcInterfaces->GetMeshAutoUVImplementation() : nullptr;
+		if (MeshAutoUV)
 		{
 		    FMeshDescription MeshCopy = MeshDescription;
 			TArray<FTriangleID>	RemapTriangles;
@@ -2333,8 +2334,6 @@ bool FStaticMeshOperations::GenerateUV(const FMeshDescription& MeshDescription,
 				}
 			};
 			    
-			IGeometryProcessing_MeshAutoUV* MeshAutoUV = GeomProcInterfaces->GetMeshAutoUVImplementation();
-
 		    IGeometryProcessing_MeshAutoUV::FOptions Options = MeshAutoUV->ConstructDefaultOptions();
 			Options.Method = GetAutoUVMethod(GenerateUVOptions.UVMethod);
 
