Modeling Tools: Adds Geometry Selection Support for the Smooth Tool

diff --git a/Engine/Plugins/Experimental/MeshModelingToolsetExp/Source/MeshModelingToolsExp/Private/OffsetMeshTool.cpp b/Engine/Plugins/Experimental/MeshModelingToolsetExp/Source/MeshModelingToolsExp/Private/OffsetMeshTool.cpp
--- a/Engine/Plugins/Experimental/MeshModelingToolsetExp/Source/MeshModelingToolsExp/Private/OffsetMeshTool.cpp
+++ b/Engine/Plugins/Experimental/MeshModelingToolsetExp/Source/MeshModelingToolsExp/Private/OffsetMeshTool.cpp
@@ -14,6 +14,15 @@ using namespace UE::Geometry;
 
 #define LOCTEXT_NAMESPACE "UOffsetMeshTool"
 
+/*
+ * Tool Builder
+ */
+
+USingleTargetWithSelectionTool* UOffsetMeshToolBuilder::CreateNewTool(const FToolBuilderState& SceneState) const
+{
+	return NewObject<UOffsetMeshTool>(SceneState.ToolManager);
+}
+
 /*
  * Tool
  */
diff --git a/Engine/Plugins/Experimental/MeshModelingToolsetExp/Source/MeshModelingToolsExp/Private/SmoothMeshTool.cpp b/Engine/Plugins/Experimental/MeshModelingToolsetExp/Source/MeshModelingToolsExp/Private/SmoothMeshTool.cpp
--- a/Engine/Plugins/Experimental/MeshModelingToolsetExp/Source/MeshModelingToolsExp/Private/SmoothMeshTool.cpp
+++ b/Engine/Plugins/Experimental/MeshModelingToolsetExp/Source/MeshModelingToolsExp/Private/SmoothMeshTool.cpp
@@ -7,6 +7,8 @@
 #include "DynamicMesh/MeshNormals.h"
 
 // Smoothing operators
+#include "DynamicSubmesh3.h"
+#include "Selections/GeometrySelectionUtil.h"
 #include "SmoothingOps/IterativeSmoothingOp.h"
 #include "SmoothingOps/CotanSmoothingOp.h"
 
@@ -16,6 +18,15 @@ using namespace UE::Geometry;
 
 #define LOCTEXT_NAMESPACE "USmoothMeshTool"
 
+/*
+ * Tool Builder
+ */
+ 
+USingleTargetWithSelectionTool* USmoothMeshToolBuilder::CreateNewTool(const FToolBuilderState& SceneState) const
+{
+	return NewObject<USmoothMeshTool>(SceneState.ToolManager);
+}
+
 /*
  * Tool
  */
@@ -79,6 +90,23 @@ TUniquePtr<FDynamicMeshOperator> USmoothMeshTool::MakeNewOperator()
 	
 	const FDynamicMesh3* Mesh = &GetInitialMesh();
 
+	bool bUsingSelection = false;
+
+	UE::Geometry::FDynamicSubmesh3 SubmeshToSmooth;
+	
+	if (HasGeometrySelection())
+	{
+		// retrieves all triangles in the current selected geometry
+		TSet<int> TrianglesSelected;
+		const FGeometrySelection& InputSelection = GetGeometrySelection();
+		UE::Geometry::EnumerateSelectionTriangles(InputSelection, *Mesh,
+		[&](int32 TriangleID){TrianglesSelected.Add(TriangleID); });
+
+		// creates a Submesh from the selection
+		SubmeshToSmooth = FDynamicSubmesh3(Mesh,TrianglesSelected.Array());
+		bUsingSelection = true;
+	}
+
 	FSmoothingOpBase::FOptions Options;
 	Options.BaseNormals = this->GetInitialVtxNormals();
 
@@ -99,7 +127,14 @@ TUniquePtr<FDynamicMeshOperator> USmoothMeshTool::MakeNewOperator()
 		Options.bSmoothBoundary = IterativeProperties->bSmoothBoundary;
 		Options.bUniform = true;
 		Options.bUseImplicit = false;
-		MeshOp = MakeUnique<FIterativeSmoothingOp>(Mesh, Options);
+		if (bUsingSelection)
+		{
+			MeshOp = MakeUnique<FIterativeSmoothingOp>(Mesh, Options, SubmeshToSmooth);
+		}
+		else
+		{
+			MeshOp = MakeUnique<FIterativeSmoothingOp>(Mesh, Options);
+		}
 		break;
 
 	case ESmoothMeshToolSmoothType::Diffusion:
@@ -108,7 +143,14 @@ TUniquePtr<FDynamicMeshOperator> USmoothMeshTool::MakeNewOperator()
 		Options.Iterations = DiffusionProperties->Steps;
 		Options.bUniform = DiffusionProperties->bPreserveUVs == false;
 		Options.bUseImplicit = true;
-		MeshOp = MakeUnique<FIterativeSmoothingOp>(Mesh, Options);
+		if (bUsingSelection)
+		{
+			MeshOp = MakeUnique<FIterativeSmoothingOp>(Mesh, Options, SubmeshToSmooth);
+		}
+		else
+		{
+			MeshOp = MakeUnique<FIterativeSmoothingOp>(Mesh, Options);
+		}
 		break;
 
 	case ESmoothMeshToolSmoothType::Implicit:
@@ -122,7 +164,14 @@ TUniquePtr<FDynamicMeshOperator> USmoothMeshTool::MakeNewOperator()
 		Options.bUniform = ImplicitProperties->bPreserveUVs == false;
 		Options.bUseImplicit = true;
 		Options.NormalOffset = ImplicitProperties->VolumeCorrection;
-		MeshOp = MakeUnique<FCotanSmoothingOp>(Mesh, Options);
+		if (bUsingSelection)
+		{
+			MeshOp = MakeUnique<FCotanSmoothingOp>(Mesh, Options, SubmeshToSmooth);
+		}
+		else
+		{
+			MeshOp = MakeUnique<FCotanSmoothingOp>(Mesh, Options);
+		}
 		}
 		break;
 	}
diff --git a/Engine/Plugins/Experimental/MeshModelingToolsetExp/Source/MeshModelingToolsExp/Public/OffsetMeshTool.h b/Engine/Plugins/Experimental/MeshModelingToolsetExp/Source/MeshModelingToolsExp/Public/OffsetMeshTool.h
--- a/Engine/Plugins/Experimental/MeshModelingToolsetExp/Source/MeshModelingToolsExp/Public/OffsetMeshTool.h
+++ b/Engine/Plugins/Experimental/MeshModelingToolsetExp/Source/MeshModelingToolsExp/Public/OffsetMeshTool.h
@@ -148,7 +148,5 @@ class MESHMODELINGTOOLSEXP_API UOffsetMeshToolBuilder : public UBaseMeshProcessi
 {
 	GENERATED_BODY()
 public:
-	virtual UBaseMeshProcessingTool* MakeNewToolInstance(UObject* Outer) const {
-		return NewObject<UOffsetMeshTool>(Outer);
-	}
+	virtual USingleTargetWithSelectionTool* CreateNewTool(const FToolBuilderState& SceneState) const override;
 };
diff --git a/Engine/Plugins/Experimental/MeshModelingToolsetExp/Source/MeshModelingToolsExp/Public/SmoothMeshTool.h b/Engine/Plugins/Experimental/MeshModelingToolsetExp/Source/MeshModelingToolsExp/Public/SmoothMeshTool.h
--- a/Engine/Plugins/Experimental/MeshModelingToolsetExp/Source/MeshModelingToolsExp/Public/SmoothMeshTool.h
+++ b/Engine/Plugins/Experimental/MeshModelingToolsetExp/Source/MeshModelingToolsExp/Public/SmoothMeshTool.h
@@ -177,7 +177,5 @@ class MESHMODELINGTOOLSEXP_API USmoothMeshToolBuilder : public UBaseMeshProcessi
 {
 	GENERATED_BODY()
 public:
-	virtual UBaseMeshProcessingTool* MakeNewToolInstance(UObject* Outer) const {
-		return NewObject<USmoothMeshTool>(Outer);
-	}
+	virtual USingleTargetWithSelectionTool* CreateNewTool(const FToolBuilderState& SceneState) const override;
 };
diff --git a/Engine/Plugins/Runtime/MeshModelingToolset/Source/ModelingComponents/Private/BaseTools/BaseMeshProcessingTool.cpp b/Engine/Plugins/Runtime/MeshModelingToolset/Source/ModelingComponents/Private/BaseTools/BaseMeshProcessingTool.cpp
--- a/Engine/Plugins/Runtime/MeshModelingToolset/Source/ModelingComponents/Private/BaseTools/BaseMeshProcessingTool.cpp
+++ b/Engine/Plugins/Runtime/MeshModelingToolset/Source/ModelingComponents/Private/BaseTools/BaseMeshProcessingTool.cpp
@@ -20,6 +20,7 @@
 #include "TargetInterfaces/PrimitiveComponentBackedTarget.h"
 #include "ToolTargetManager.h"
 #include "ModelingToolTargetUtil.h"
+#include "Selection/StoredMeshSelectionUtil.h"
 
 #include UE_INLINE_GENERATED_CPP_BY_NAME(BaseMeshProcessingTool)
 
@@ -47,16 +48,11 @@ bool UBaseMeshProcessingToolBuilder::CanBuildTool(const FToolBuilderState& Scene
 	return SceneState.TargetManager->CountSelectedAndTargetable(SceneState, GetTargetRequirements()) == 1;
 }
 
-UInteractiveTool* UBaseMeshProcessingToolBuilder::BuildTool(const FToolBuilderState& SceneState) const
+USingleTargetWithSelectionTool* UBaseMeshProcessingToolBuilder::CreateNewTool(const FToolBuilderState& SceneState) const
 {
-	UBaseMeshProcessingTool* NewTool = MakeNewToolInstance(SceneState.ToolManager);
-
-	UToolTarget* Target = SceneState.TargetManager->BuildFirstSelectedTargetable(SceneState, GetTargetRequirements());
-	check(Target);
-	NewTool->SetTarget(Target);
-	NewTool->SetWorld(SceneState.World);
-
-	return NewTool;
+	PRAGMA_DISABLE_DEPRECATION_WARNINGS
+	return  MakeNewToolInstance(SceneState.ToolManager);
+	PRAGMA_ENABLE_DEPRECATION_WARNINGS
 }
 
 
@@ -66,12 +62,6 @@ UInteractiveTool* UBaseMeshProcessingToolBuilder::BuildTool(const FToolBuilderSt
  * Tool
  */
 
-void UBaseMeshProcessingTool::SetWorld(UWorld* World)
-{
-	this->TargetWorld = World;
-}
-
-
 void UBaseMeshProcessingTool::Setup()
 {
 	UInteractiveTool::Setup();
@@ -143,7 +133,7 @@ void UBaseMeshProcessingTool::Setup()
 
 	// Construct the preview object and set the material on it.
 	Preview = NewObject<UMeshOpPreviewWithBackgroundCompute>(this, "Preview");
-	Preview->Setup(this->TargetWorld, this); // Adds the actual functional tool in the Preview object
+	Preview->Setup(GetTargetWorld(), this); // Adds the actual functional tool in the Preview object
 	Preview->PreviewMesh->SetTangentsMode(EDynamicMeshComponentTangentsMode::AutoCalculated);
 	ToolSetupUtil::ApplyRenderingConfigurationToPreview(Preview->PreviewMesh, Target);
 
diff --git a/Engine/Plugins/Runtime/MeshModelingToolset/Source/ModelingComponents/Public/BaseTools/BaseMeshProcessingTool.h b/Engine/Plugins/Runtime/MeshModelingToolset/Source/ModelingComponents/Public/BaseTools/BaseMeshProcessingTool.h
--- a/Engine/Plugins/Runtime/MeshModelingToolset/Source/ModelingComponents/Public/BaseTools/BaseMeshProcessingTool.h
+++ b/Engine/Plugins/Runtime/MeshModelingToolset/Source/ModelingComponents/Public/BaseTools/BaseMeshProcessingTool.h
@@ -3,8 +3,7 @@
 #pragma once
 
 #include "CoreMinimal.h"
-#include "SingleSelectionTool.h"
-#include "InteractiveToolBuilder.h"
+#include "SingleTargetWithSelectionTool.h"
 #include "DynamicMesh/DynamicMesh3.h"
 #include "WeightMapTypes.h"
 #include "MeshOpPreviewHelpers.h"
@@ -24,16 +23,19 @@ using UE::Geometry::FIndexedWeightMap1f;
  * ToolBuilder for UBaseMeshProcessingTool
  */
 UCLASS()
-class MODELINGCOMPONENTS_API UBaseMeshProcessingToolBuilder : public UInteractiveToolWithToolTargetsBuilder
+class MODELINGCOMPONENTS_API UBaseMeshProcessingToolBuilder : public USingleTargetWithSelectionToolBuilder
 {
 	GENERATED_BODY()
 
 public:
 	virtual bool CanBuildTool(const FToolBuilderState& SceneState) const override;
-	virtual UInteractiveTool* BuildTool(const FToolBuilderState& SceneState) const override;
+	virtual bool RequiresInputSelection() const override { return false; }
+	virtual USingleTargetWithSelectionTool* CreateNewTool(const FToolBuilderState& SceneState) const override;
+
 
 public:
 	// subclass must override!
+	UE_DEPRECATED(5.5, "MakeNewToolInstance is deprecated, please use CreateNewTool instead")
 	virtual UBaseMeshProcessingTool* MakeNewToolInstance(UObject* Outer) const { check(false); return nullptr; }
 
 protected:
@@ -74,7 +76,7 @@ protected:
  *
  */
 UCLASS()
-class MODELINGCOMPONENTS_API UBaseMeshProcessingTool : public USingleSelectionTool, public UE::Geometry::IDynamicMeshOperatorFactory
+class MODELINGCOMPONENTS_API UBaseMeshProcessingTool : public USingleTargetWithSelectionTool, public UE::Geometry::IDynamicMeshOperatorFactory
 {
 	GENERATED_BODY()
 protected:
@@ -82,8 +84,6 @@ protected:
 public:
 	UBaseMeshProcessingTool() = default;
 
-	virtual void SetWorld(UWorld* World);
-
 	//
 	// InteractiveTool API - generally does not need to be modified by subclasses
 	//
@@ -186,8 +186,6 @@ protected:
 	virtual void SavePropertySets();
 
 
-	UWorld* TargetWorld = nullptr;
-
 	// Preview object holds temporary Actor with preview mesh component
 	UPROPERTY()
 	TObjectPtr<UMeshOpPreviewWithBackgroundCompute> Preview = nullptr;
diff --git a/Engine/Plugins/Runtime/MeshModelingToolset/Source/ModelingOperators/Private/SmoothingOps/CotanSmoothingOp.cpp b/Engine/Plugins/Runtime/MeshModelingToolset/Source/ModelingOperators/Private/SmoothingOps/CotanSmoothingOp.cpp
--- a/Engine/Plugins/Runtime/MeshModelingToolset/Source/ModelingOperators/Private/SmoothingOps/CotanSmoothingOp.cpp
+++ b/Engine/Plugins/Runtime/MeshModelingToolset/Source/ModelingOperators/Private/SmoothingOps/CotanSmoothingOp.cpp
@@ -12,6 +12,11 @@ FCotanSmoothingOp::FCotanSmoothingOp(const FDynamicMesh3* Mesh, const FSmoothing
 {
 }
 
+FCotanSmoothingOp::FCotanSmoothingOp(const FDynamicMesh3* Mesh, const FSmoothingOpBase::FOptions& OptionsIn, const FDynamicSubmesh3& Submesh) :
+	FSmoothingOpBase(Mesh, OptionsIn, Submesh)
+{
+}
+
 double FCotanSmoothingOp::GetSmoothPower(int32 VertexID, bool bIsBoundary)
 {
 	double UsePower = SmoothOptions.SmoothPower;
diff --git a/Engine/Plugins/Runtime/MeshModelingToolset/Source/ModelingOperators/Private/SmoothingOps/IterativeSmoothingOp.cpp b/Engine/Plugins/Runtime/MeshModelingToolset/Source/ModelingOperators/Private/SmoothingOps/IterativeSmoothingOp.cpp
--- a/Engine/Plugins/Runtime/MeshModelingToolset/Source/ModelingOperators/Private/SmoothingOps/IterativeSmoothingOp.cpp
+++ b/Engine/Plugins/Runtime/MeshModelingToolset/Source/ModelingOperators/Private/SmoothingOps/IterativeSmoothingOp.cpp
@@ -10,6 +10,12 @@ FIterativeSmoothingOp::FIterativeSmoothingOp(const FDynamicMesh3* Mesh, const FS
 {
 }
 
+FIterativeSmoothingOp::FIterativeSmoothingOp(const FDynamicMesh3* Mesh, const FSmoothingOpBase::FOptions& OptionsIn,  const FDynamicSubmesh3& Submesh) :
+	FSmoothingOpBase(Mesh, OptionsIn, Submesh)
+{
+}
+
+
 void FIterativeSmoothingOp::CalculateResult(FProgressCancel* Progress)
 {
 	// Update the values in the position buffer with smoothed positions.
diff --git a/Engine/Plugins/Runtime/MeshModelingToolset/Source/ModelingOperators/Private/SmoothingOps/SmoothingOpBase.cpp b/Engine/Plugins/Runtime/MeshModelingToolset/Source/ModelingOperators/Private/SmoothingOps/SmoothingOpBase.cpp
--- a/Engine/Plugins/Runtime/MeshModelingToolset/Source/ModelingOperators/Private/SmoothingOps/SmoothingOpBase.cpp
+++ b/Engine/Plugins/Runtime/MeshModelingToolset/Source/ModelingOperators/Private/SmoothingOps/SmoothingOpBase.cpp
@@ -2,6 +2,7 @@
 
 #include "SmoothingOps/SmoothingOpBase.h"
 
+#include "DynamicSubmesh3.h"
 #include "DynamicMesh/DynamicMesh3.h"
 #include "DynamicMesh/MeshNormals.h"
 
@@ -16,9 +17,29 @@ FSmoothingOpBase::FSmoothingOpBase(const FDynamicMesh3* Mesh, const FOptions& Op
 	
 	PositionBuffer.SetNum(ResultMesh->MaxVertexID());
 
-	for (int vid : ResultMesh->VertexIndicesItr())
+	for (int VID : ResultMesh->VertexIndicesItr())
 	{
-		PositionBuffer[vid] = ResultMesh->GetVertex(vid);
+		PositionBuffer[VID] = ResultMesh->GetVertex(VID);
+	}
+}
+
+FSmoothingOpBase::FSmoothingOpBase(const FDynamicMesh3* Mesh, const FOptions& OptionsIn, const FDynamicSubmesh3& Submesh) :
+	FDynamicMeshOperator(),
+	SmoothOptions(OptionsIn)
+{
+	// copy the src mesh into the SavedMesh so that positions of unselected vertices are already saved in ResultMesh
+	SavedMesh = MakeUnique<FDynamicMesh3>(*Mesh);
+	
+	ResultMesh = MakeUnique<FDynamicMesh3>(Submesh.GetSubmesh());
+	const int32 NumVertsToSmooth = ResultMesh->MaxVertexID(); // the number of vertices in the submesh, which is the region to be smoothed
+	PositionBuffer.SetNum(NumVertsToSmooth);
+	SmoothedToOriginalMap.SetNum(NumVertsToSmooth);
+
+	for (int SubmeshVID : ResultMesh->VertexIndicesItr())
+	{
+		const FVector3d VertexPosition = ResultMesh->GetVertex(SubmeshVID);
+		PositionBuffer[SubmeshVID] = VertexPosition;
+		SmoothedToOriginalMap[SubmeshVID] = Submesh.MapVertexToBaseMesh(SubmeshVID);
 	}
 }
 
@@ -30,11 +51,31 @@ void FSmoothingOpBase::SetTransform(const FTransformSRT3d& XForm)
 
 void FSmoothingOpBase::UpdateResultMesh()
 {
-	// move the verts to the new location
-	for (int32 vid : ResultMesh->VertexIndicesItr())
+	// if SavedMesh was populated, we are operating on a submesh
+	if (SavedMesh != nullptr)
+	{
+		for (int32 VID : SavedMesh->VertexIndicesItr())
+		{
+			// if vertex is in submesh, it was smoothed and has new position
+			if (SmoothedToOriginalMap.Contains(VID))
+			{
+				int32 SubmeshIndex = SmoothedToOriginalMap.Find(VID);
+				const FVector3d Pos = PositionBuffer[SubmeshIndex];
+				SavedMesh->SetVertex(VID, Pos);
+			}
+			// else, can keep vertex's original position as initially stored in SavedMesh, no need to reassign
+		}
+		ResultMesh = MoveTemp(SavedMesh);
+	}
+	// otherwise, the entire mesh is smoothed
+	else
 	{
-		const FVector3d Pos = PositionBuffer[vid];
-		ResultMesh->SetVertex(vid, Pos);
+		// move all the vertices to their new location
+		for (int32 VID : ResultMesh->VertexIndicesItr())
+		{
+			const FVector3d Pos = PositionBuffer[VID];
+			ResultMesh->SetVertex(VID, Pos);
+		}
 	}
 
 	// recalculate normals
diff --git a/Engine/Plugins/Runtime/MeshModelingToolset/Source/ModelingOperators/Public/SmoothingOps/CotanSmoothingOp.h b/Engine/Plugins/Runtime/MeshModelingToolset/Source/ModelingOperators/Public/SmoothingOps/CotanSmoothingOp.h
--- a/Engine/Plugins/Runtime/MeshModelingToolset/Source/ModelingOperators/Public/SmoothingOps/CotanSmoothingOp.h
+++ b/Engine/Plugins/Runtime/MeshModelingToolset/Source/ModelingOperators/Public/SmoothingOps/CotanSmoothingOp.h
@@ -14,6 +14,9 @@ class MODELINGOPERATORS_API  FCotanSmoothingOp : public FSmoothingOpBase
 {
 public:
 	FCotanSmoothingOp(const FDynamicMesh3* Mesh, const FSmoothingOpBase::FOptions& OptionsIn);
+	
+	// Support for smoothing only selected geometry
+	FCotanSmoothingOp(const FDynamicMesh3* Mesh, const FSmoothingOpBase::FOptions& OptionsIn, const FDynamicSubmesh3& Submesh);
 
 	~FCotanSmoothingOp() override {};
 
diff --git a/Engine/Plugins/Runtime/MeshModelingToolset/Source/ModelingOperators/Public/SmoothingOps/IterativeSmoothingOp.h b/Engine/Plugins/Runtime/MeshModelingToolset/Source/ModelingOperators/Public/SmoothingOps/IterativeSmoothingOp.h
--- a/Engine/Plugins/Runtime/MeshModelingToolset/Source/ModelingOperators/Public/SmoothingOps/IterativeSmoothingOp.h
+++ b/Engine/Plugins/Runtime/MeshModelingToolset/Source/ModelingOperators/Public/SmoothingOps/IterativeSmoothingOp.h
@@ -15,6 +15,9 @@ class MODELINGOPERATORS_API  FIterativeSmoothingOp : public FSmoothingOpBase
 public:
 	FIterativeSmoothingOp(const FDynamicMesh3* Mesh, const FSmoothingOpBase::FOptions& OptionsIn);
 
+	// Support for smoothing only selected geometry
+	FIterativeSmoothingOp(const FDynamicMesh3* Mesh, const FSmoothingOpBase::FOptions& OptionsIn,  const FDynamicSubmesh3& Submesh);
+
 	~FIterativeSmoothingOp() override {};
 
 	// Apply smoothing. results in an updated ResultMesh
diff --git a/Engine/Plugins/Runtime/MeshModelingToolset/Source/ModelingOperators/Public/SmoothingOps/SmoothingOpBase.h b/Engine/Plugins/Runtime/MeshModelingToolset/Source/ModelingOperators/Public/SmoothingOps/SmoothingOpBase.h
--- a/Engine/Plugins/Runtime/MeshModelingToolset/Source/ModelingOperators/Public/SmoothingOps/SmoothingOpBase.h
+++ b/Engine/Plugins/Runtime/MeshModelingToolset/Source/ModelingOperators/Public/SmoothingOps/SmoothingOpBase.h
@@ -3,6 +3,7 @@
 #pragma once
 
 #include "CoreMinimal.h"
+#include "DynamicSubmesh3.h"
 #include "Util/ProgressCancel.h"
 #include "ModelingOperators.h"
 #include "VectorTypes.h"
@@ -61,6 +62,9 @@ public:
 
 	FSmoothingOpBase(const FDynamicMesh3* Mesh, const FOptions& OptionsIn);
 
+	// Support for smoothing only selected geometry
+	FSmoothingOpBase(const FDynamicMesh3* Mesh, const FOptions& OptionsIn, const FDynamicSubmesh3& Submesh);
+
 	virtual ~FSmoothingOpBase() override {}
 
 	// set ability on protected transform.
@@ -77,6 +81,12 @@ protected:
 
 	TArray<FVector3d> PositionBuffer;
 
+	// a copy of the original mesh to save information regarding the non-selected (/non-smoothed) mesh when applicable
+	TUniquePtr<FDynamicMesh3> SavedMesh = nullptr;
+
+	// maps the VertexId in the smoothed mesh to the VertexId in the original mesh
+	TArray<int32> SmoothedToOriginalMap;
+
 };
 
 } // end namespace UE::Geometry
