harden draw spline tool against case of a BP actor that deletes itself + use recommended method to delete the working/temp actors

diff --git a/Engine/Plugins/Experimental/MeshModelingToolsetExp/Source/MeshModelingToolsEditorOnlyExp/Private/DrawSplineTool.cpp b/Engine/Plugins/Experimental/MeshModelingToolsetExp/Source/MeshModelingToolsEditorOnlyExp/Private/DrawSplineTool.cpp
--- a/Engine/Plugins/Experimental/MeshModelingToolsetExp/Source/MeshModelingToolsEditorOnlyExp/Private/DrawSplineTool.cpp
+++ b/Engine/Plugins/Experimental/MeshModelingToolsetExp/Source/MeshModelingToolsEditorOnlyExp/Private/DrawSplineTool.cpp
@@ -539,21 +539,40 @@ void UDrawSplineTool::Shutdown(EToolShutdownType ShutdownType)
 		HiddenSpline = nullptr;
 	}
 
-	int32 NumSplinePoints = WorkingSpline->GetNumberOfSplinePoints();
-	if (ShutdownType == EToolShutdownType::Accept && NumSplinePoints > 0)
+	if (ShutdownType == EToolShutdownType::Accept && WorkingSpline.IsValid() && WorkingSpline->GetNumberOfSplinePoints() > 0)
 	{
 		GenerateAsset();
 	}
 
 	PlaneMechanic->Shutdown();
+	
+	
+	if (WorkingSpline.IsValid())
+	{
+		WorkingSpline->DestroyComponent();
+	}
 
+	auto DestroyActor = [&](AActor* Actor)
+	{
+		if (UWorld* ActorWorld = Actor->GetWorld())
+		{
+			if (GIsEditor && GUnrealEd)
+			{
+				GUnrealEd->DeleteActors(TArray{ Actor }, ActorWorld, GUnrealEd->GetSelectedActors()->GetElementSelectionSet());
+			}
+			else
+			{
+				ActorWorld->DestroyActor(Actor);
+			}
+		}
+	};
 	if (PreviewActor)
 	{
-		PreviewActor->Destroy();
+		DestroyActor(PreviewActor);
 	}
 	if (PreviewRootActor)
 	{
-		PreviewRootActor->Destroy();
+		DestroyActor(PreviewRootActor);
 	}
 
 	Super::Shutdown(ShutdownType);
@@ -657,6 +676,10 @@ void UDrawSplineTool::GenerateAsset()
 void UDrawSplineTool::AddSplinePoint(const FVector3d& HitLocation, const FVector3d& HitNormal)
 {
 	using namespace DrawSplineToolLocals;
+	if (!WorkingSpline.IsValid())
+	{
+		return;
+	}
 
 	int32 NumSplinePoints = WorkingSpline->GetNumberOfSplinePoints();
 	FVector3d UpVectorToUse = GetUpVectorToUse(HitLocation, HitNormal, NumSplinePoints);
@@ -1029,6 +1052,11 @@ void UDrawSplineTool::OnTick(float DeltaTime)
 			}
 		}
 	}
+
+	if (!WorkingSpline.IsValid())
+	{
+		GetToolManager()->PostActiveToolShutdownRequest(this, EToolShutdownType::Cancel, true, LOCTEXT("LostWorkingSpline", "The Draw Spline tool must close because the in-progress spline has been unexpectedly deleted."));
+	}
 }
 
 void UDrawSplineTool::Render(IToolsContextRenderAPI* RenderAPI)
