Add preview visualization to the split tool
+ Make the Split tool support give the option to split by mesh topology, vertex overlap (w/ a distance tolerance), material ID or polygroup (Note for meshes w/ selection, it will not have these options and will continue to split by selection)

diff --git a/Engine/Plugins/Experimental/MeshModelingToolsetExp/Source/MeshModelingToolsExp/Private/SplitMeshesTool.cpp b/Engine/Plugins/Experimental/MeshModelingToolsetExp/Source/MeshModelingToolsExp/Private/SplitMeshesTool.cpp
--- a/Engine/Plugins/Experimental/MeshModelingToolsetExp/Source/MeshModelingToolsExp/Private/SplitMeshesTool.cpp
+++ b/Engine/Plugins/Experimental/MeshModelingToolsetExp/Source/MeshModelingToolsExp/Private/SplitMeshesTool.cpp
@@ -2,6 +2,7 @@
 
 #include "SplitMeshesTool.h"
 #include "ComponentSourceInterfaces.h"
+#include "Drawing/PreviewGeometryActor.h"
 #include "InteractiveToolManager.h"
 #include "ToolTargetManager.h"
 #include "ToolBuilderUtil.h"
@@ -9,10 +10,12 @@
 #include "ModelingToolTargetUtil.h"
 #include "ModelingObjectsCreationAPI.h"
 #include "Selection/ToolSelectionUtil.h"
-#include "Selections/MeshConnectedComponents.h"
+#include "VertexConnectedComponents.h"
+#include "DynamicMesh/DynamicMeshAttributeSet.h"
 #include "DynamicMesh/MeshTransforms.h"
 #include "DynamicSubmesh3.h"
 #include "Selections/GeometrySelectionUtil.h"
+#include "Util/ColorConstants.h"
 
 #include "TargetInterfaces/MeshDescriptionProvider.h"
 #include "TargetInterfaces/PrimitiveComponentBackedTarget.h"
@@ -61,23 +64,47 @@ void USplitMeshesTool::Setup()
 
 	BasicProperties = NewObject<USplitMeshesToolProperties>(this);
 	BasicProperties->RestoreProperties(this);
+	BasicProperties->WatchProperty(BasicProperties->SplitMethod, [this](ESplitMeshesMethod) { UpdateSplitMeshes(); });
+	BasicProperties->WatchProperty(BasicProperties->ConnectVerticesThreshold, [this](double) { UpdateSplitMeshes(); });
+	BasicProperties->WatchProperty(BasicProperties->bShowPreview, [this](bool bShowPreview) { UpdatePreviewVisibility(bShowPreview); });
 	AddToolPropertySource(BasicProperties);
 
 	SourceMeshes.SetNum(Targets.Num());
+	bool bHasSelection = false;
 	for (int32 k = 0; k < Targets.Num(); ++k)
 	{
 		SourceMeshes[k].Mesh = UE::ToolTarget::GetDynamicMeshCopy(Targets[k], true);
 		SourceMeshes[k].Materials = UE::ToolTarget::GetMaterialSet(Targets[k]).Materials;
+		bHasSelection = bHasSelection || HasGeometrySelection(k);
 	}
+	BasicProperties->bIsInSelectionMode = bHasSelection;
+
+	PerTargetPreviews.Reserve(Targets.Num());
+	for (int32 TargetIdx = 0; TargetIdx < Targets.Num(); ++TargetIdx)
+	{
+		UPreviewGeometry* PreviewGeom = PerTargetPreviews.Add_GetRef(NewObject<UPreviewGeometry>(this));
+		PreviewGeom->CreateInWorld(UE::ToolTarget::GetTargetActor(Targets[TargetIdx])->GetWorld(), UE::ToolTarget::GetLocalToWorldTransform(Targets[TargetIdx]));
+	}
+	PreviewMaterial = ToolSetupUtil::GetVertexColorMaterial(GetToolManager(), false);
 
 	UpdateSplitMeshes();
 
+	UpdatePreviewVisibility(BasicProperties->bShowPreview);
+
 	SetToolDisplayName(LOCTEXT("ToolName", "Split"));
 	GetToolManager()->DisplayMessage(
 		LOCTEXT("OnStartTool", "Split Meshes into parts"),
 		EToolMessageLevel::UserNotification);
 }
 
+void USplitMeshesTool::UpdatePreviewVisibility(bool bShowPreview)
+{
+	for (int32 PreviewIdx = 0; PreviewIdx < PerTargetPreviews.Num(); ++PreviewIdx)
+	{
+		PerTargetPreviews[PreviewIdx]->SetAllVisible(bShowPreview);
+	}
+}
+
 bool USplitMeshesTool::CanAccept() const
 {
 	return Super::CanAccept();
@@ -85,6 +112,11 @@ bool USplitMeshesTool::CanAccept() const
 
 void USplitMeshesTool::OnShutdown(EToolShutdownType ShutdownType)
 {
+	for (UPreviewGeometry* PreviewGeom : PerTargetPreviews)
+	{
+		PreviewGeom->Disconnect();
+	}
+
 	OutputTypeProperties->SaveProperties(this, TEXT("OutputTypeFromInputTool"));
 	BasicProperties->SaveProperties(this);
 
@@ -163,15 +195,39 @@ void USplitMeshesTool::UpdateSplitMeshes()
 	SplitMeshes.SetNum(SourceMeshes.Num());
 	NoSplitCount = 0;
 
+	int32 VisColorIdx = 0;
+
 	for (int32 si = 0; si < SourceMeshes.Num(); ++si)
 	{
 		FComponentsInfo& SplitInfo = SplitMeshes[si];
 		const FDynamicMesh3* SourceMesh = &SourceMeshes[si].Mesh;
 		const TArray<UMaterialInterface*> SourceMaterials = SourceMeshes[si].Materials;
 
-		FMeshConnectedComponents MeshComponents(SourceMesh);
-		
-		int32 NumComponents;
+		TArray<TArray<int32>> ComponentTriIndices;
+		int32 NumComponents = 0;
+
+		auto FillComponentTriIndicesFromTriIDs = [&ComponentTriIndices, &NumComponents, SourceMesh](TFunctionRef<int32(int32)> TIDtoID)
+			{
+				TMap<int32, int32> ComponentIDMap;
+				for (int32 TID : SourceMesh->TriangleIndicesItr())
+				{
+					int32 CompID = TIDtoID(TID);
+					int32* FoundIdx = ComponentIDMap.Find(CompID);
+					int32 UseIdx = -1;
+					if (FoundIdx)
+					{
+						UseIdx = *FoundIdx;
+					}
+					else
+					{
+						UseIdx = ComponentTriIndices.AddDefaulted();
+						ComponentIDMap.Add(CompID, UseIdx);
+					}
+					ComponentTriIndices[UseIdx].Add(TID);
+				}
+				NumComponents = ComponentTriIndices.Num();
+			};
+
 		const bool bMeshHasGeometrySelection = HasGeometrySelection(si);
 		if (bMeshHasGeometrySelection)
 		{
@@ -179,14 +235,34 @@ void USplitMeshesTool::UpdateSplitMeshes()
 			// decide where to split the mesh; instead split into 2 meshes regardless: the selected geometry, and everything else
 			NumComponents = 2;
 		}
-		else
+		else if (BasicProperties->SplitMethod == ESplitMeshesMethod::ByMeshTopology || BasicProperties->SplitMethod == ESplitMeshesMethod::ByVertexOverlap)
+		{
+			FVertexConnectedComponents Components(SourceMesh->MaxVertexID());
+			Components.ConnectTriangles(*SourceMesh);
+			if (BasicProperties->SplitMethod == ESplitMeshesMethod::ByVertexOverlap)
+			{
+				Components.ConnectCloseVertices(*SourceMesh, BasicProperties->ConnectVerticesThreshold, 2);
+			}
+			FillComponentTriIndicesFromTriIDs([&](int32 TID)->int32 { return Components.GetComponent(SourceMesh->GetTriangle(TID).A); });
+		}
+		else if (BasicProperties->SplitMethod == ESplitMeshesMethod::ByPolyGroup)
+		{
+			FillComponentTriIndicesFromTriIDs([&](int32 TID)->int32 { return SourceMesh->GetTriangleGroup(TID); });
+		}
+		else if (BasicProperties->SplitMethod == ESplitMeshesMethod::ByMaterialID)
 		{
-			MeshComponents.FindConnectedTriangles();
-			NumComponents = MeshComponents.Num();
+			if (SourceMesh->HasAttributes())
+			{
+				if (const FDynamicMeshMaterialAttribute* MaterialID = SourceMesh->Attributes()->GetMaterialID())
+				{
+					FillComponentTriIndicesFromTriIDs([&](int32 TID)->int32 { return MaterialID->GetValue(TID); });
+				}
+			}
 		}
 		
 		if (NumComponents < 2)
 		{
+			PerTargetPreviews[si]->RemoveAllTriangleSets();
 			SplitInfo.bNoComponents = true;
 			NoSplitCount++;
 			continue;
@@ -231,9 +307,9 @@ void USplitMeshesTool::UpdateSplitMeshes()
 			else
 			{
 				// if statement should always be true- components should always have been calculated & populated when there's no geometry selection
-				if (!MeshComponents.Components.IsEmpty())
+				if (ensure(!ComponentTriIndices.IsEmpty()))
 				{
-					SubmeshCalc = FDynamicSubmesh3(SourceMesh, MeshComponents[k].Indices);
+					SubmeshCalc = FDynamicSubmesh3(SourceMesh, ComponentTriIndices[k]);
 				}
 			}
 
@@ -278,6 +354,25 @@ void USplitMeshesTool::UpdateSplitMeshes()
 			SplitInfo.Materials[k] = MoveTemp(NewMaterials);
 			SplitInfo.Origins[k] = Origin;
 		}
+
+		PerTargetPreviews[si]->CreateOrUpdateTriangleSet(TEXT("Components"), 1, [&](int32, TArray<FRenderableTriangle>& Triangles)
+			{	
+				for (const UE::Geometry::FDynamicMesh3& Mesh : SplitInfo.Meshes)
+				{
+					++VisColorIdx;
+					FColor MeshColor = LinearColors::SelectFColor(VisColorIdx);
+
+					for (int32 TID : Mesh.TriangleIndicesItr())
+					{
+						FVector3d Normal = Mesh.GetTriNormal(TID);
+						FIndex3i Tri = Mesh.GetTriangle(TID);
+						FRenderableTriangleVertex A(Mesh.GetVertex(Tri.A), FVector2D(0, 0), Normal, MeshColor);
+						FRenderableTriangleVertex B(Mesh.GetVertex(Tri.B), FVector2D(1, 0), Normal, MeshColor);
+						FRenderableTriangleVertex C(Mesh.GetVertex(Tri.C), FVector2D(1, 1), Normal, MeshColor);
+						Triangles.Add(FRenderableTriangle(PreviewMaterial, A, B, C));
+					}
+				}
+			}, SourceMeshes[si].Mesh.TriangleCount());
 	}
 
 	if (NoSplitCount > 0)
diff --git a/Engine/Plugins/Experimental/MeshModelingToolsetExp/Source/MeshModelingToolsExp/Public/SplitMeshesTool.h b/Engine/Plugins/Experimental/MeshModelingToolsetExp/Source/MeshModelingToolsExp/Public/SplitMeshesTool.h
--- a/Engine/Plugins/Experimental/MeshModelingToolsetExp/Source/MeshModelingToolsExp/Public/SplitMeshesTool.h
+++ b/Engine/Plugins/Experimental/MeshModelingToolsetExp/Source/MeshModelingToolsExp/Public/SplitMeshesTool.h
@@ -24,15 +24,39 @@ protected:
 	virtual const FToolTargetTypeRequirements& GetTargetRequirements() const override;
 };
 
-
+// Methods for splitting meshes
+UENUM()
+enum class ESplitMeshesMethod : uint8
+{
+	// Split meshes based on the triangle-connected regions of the mesh
+	ByMeshTopology,
+	// Split meshes based on triangle-connected regions, and consider vertices to be connected if they are within a tolerance distance
+	ByVertexOverlap,
+	// Split meshes based on material ID
+	ByMaterialID,
+	// Split meshes based on PolyGroup ID
+	ByPolyGroup
+};
 
 UCLASS()
 class MESHMODELINGTOOLSEXP_API USplitMeshesToolProperties : public UInteractiveToolPropertySet
 {
 	GENERATED_BODY()
 public:
+	UPROPERTY(EditAnywhere, Category = Options, meta = (EditCondition = "!bIsInSelectionMode", EditConditionHides, HideEditConditionToggle))
+	ESplitMeshesMethod SplitMethod = ESplitMeshesMethod::ByMeshTopology;
+
+	UPROPERTY(EditAnywhere, Category = Options, meta = (ClampMin = .0001, UIMax = 1.0, EditCondition = "!bIsInSelectionMode && SplitMethod == ESplitMeshesMethod::ByVertexOverlap", EditConditionHides, HideEditConditionToggle))
+	double ConnectVerticesThreshold = 0.01;
+
 	UPROPERTY(EditAnywhere, Category = Options)
 	bool bTransferMaterials = true;
+
+	UPROPERTY(EditAnywhere, Category = Options)
+	bool bShowPreview = true;
+
+	UPROPERTY()
+	bool bIsInSelectionMode = false;
 };
 
 
@@ -77,4 +101,13 @@ protected:
 	int32 NoSplitCount = 0;
 
 	void UpdateSplitMeshes();
+
+private:
+	// Preview how the meshes are to be split
+	UPROPERTY()
+	TArray<TObjectPtr<UPreviewGeometry>> PerTargetPreviews;
+	UPROPERTY()
+	TObjectPtr<UMaterialInterface> PreviewMaterial = nullptr;
+
+	void UpdatePreviewVisibility(bool bShowPreview);
 };
