Mesh to collision tool improvements around convex decomposition:
+ split out convex decomposition from single convex hulls as a separate shape type (rather than hiding the feature under the 'max hulls' count)
+ add a pre-simplify option to make the decomposition compute faster for very large meshes
+ use an enum rather than bool to select the decomposition algorithm, and only show options relevant to each algorithm
+ use better defaults for the navigable space protection algorithm
+ do not default-enable the 'max count' collision shape option, as it can be confusing/frustrating to wait for a decomposition only to find it has random holes due to that option quietly throwing away shapes

diff --git a/Engine/Plugins/Experimental/MeshModelingToolsetExp/Source/MeshModelingToolsExp/Private/Physics/SetCollisionGeometryTool.cpp b/Engine/Plugins/Experimental/MeshModelingToolsetExp/Source/MeshModelingToolsExp/Private/Physics/SetCollisionGeometryTool.cpp
--- a/Engine/Plugins/Experimental/MeshModelingToolsetExp/Source/MeshModelingToolsExp/Private/Physics/SetCollisionGeometryTool.cpp
+++ b/Engine/Plugins/Experimental/MeshModelingToolsetExp/Source/MeshModelingToolsExp/Private/Physics/SetCollisionGeometryTool.cpp
@@ -132,14 +132,10 @@ public:
 			UseShapeGenerator->Generate_Capsules(NewCollision->Geometry);
 			break;
 		case ECollisionGeometryType::ConvexHulls:
-			if (UseShapeGenerator->ConvexDecompositionMaxPieces > 1)
-			{
-				UseShapeGenerator->Generate_ConvexHullDecompositions(NewCollision->Geometry);
-			}
-			else
-			{
-				UseShapeGenerator->Generate_ConvexHulls(NewCollision->Geometry);
-			}
+			UseShapeGenerator->Generate_ConvexHulls(NewCollision->Geometry);
+			break;
+		case ECollisionGeometryType::ConvexDecompositions:
+			UseShapeGenerator->Generate_ConvexHullDecompositions(NewCollision->Geometry);
 			break;
 		case ECollisionGeometryType::SweptHulls:
 			UseShapeGenerator->Generate_ProjectedHulls(NewCollision->Geometry,
@@ -322,11 +318,14 @@ void USetCollisionGeometryTool::Setup()
 	Settings->WatchProperty(Settings->bDetectCapsules, [this](bool) { InvalidateCompute(); });
 	Settings->WatchProperty(Settings->bSimplifyHulls, [this](bool) { InvalidateCompute(); });
 	Settings->WatchProperty(Settings->HullTargetFaceCount, [this](int32) { InvalidateCompute(); });
-	Settings->WatchProperty(Settings->MaxHullsPerMesh, [this](int32) { InvalidateCompute(); });
+	Settings->WatchProperty(Settings->MaxHullsPerShape, [this](int32) { InvalidateCompute(); });
+	Settings->WatchProperty(Settings->bPreSimplifyToEdgeLength, [this](bool) { InvalidateCompute(); });
+	Settings->WatchProperty(Settings->DecompositionTargetEdgeLength, [this](double) { InvalidateCompute(); });
 	Settings->WatchProperty(Settings->ConvexDecompositionSearchFactor, [this](int32) { InvalidateCompute(); });
 	Settings->WatchProperty(Settings->AddHullsErrorTolerance, [this](int32) { InvalidateCompute(); });
 	Settings->WatchProperty(Settings->MinPartThickness, [this](int32) { InvalidateCompute(); });
-	Settings->WatchProperty(Settings->bUseNegativeSpaceInDecomposition, [this](bool) { InvalidateCompute(); });
+	Settings->WatchProperty(Settings->DecompositionMethod, [this](EConvexDecompositionMethod) { InvalidateCompute(); });
+	Settings->WatchProperty(Settings->bLimitHullsPerShape, [this](bool) { InvalidateCompute(); });
 	Settings->WatchProperty(Settings->NegativeSpaceMinRadius, [this](int32) { InvalidateCompute(); });
 	Settings->WatchProperty(Settings->NegativeSpaceTolerance, [this](int32) { InvalidateCompute(); });
 	Settings->WatchProperty(Settings->bIgnoreInternalNegativeSpace, [this](bool) { InvalidateCompute(); });
@@ -445,11 +444,14 @@ TUniquePtr<UE::Geometry::TGenericDataOperator<FPhysicsDataCollection>> USetColli
 	// SimplifyHulls on the shape generator controls simplification on both swept and convex hull paths, but for Swept Hulls UI we leave simplification always enabled
 	Op->UseShapeGenerator->bSimplifyHulls = Settings->GeometryType == ECollisionGeometryType::SweptHulls || Settings->bSimplifyHulls;
 	Op->UseShapeGenerator->HullTargetFaceCount = Settings->HullTargetFaceCount;
-	Op->UseShapeGenerator->ConvexDecompositionMaxPieces = Settings->MaxHullsPerMesh;
+	Op->UseShapeGenerator->bDecompositionPreSimplifyWithEdgeLength = Settings->bPreSimplifyToEdgeLength;
+	Op->UseShapeGenerator->DecompositionPreSimplifyEdgeLength = Settings->DecompositionTargetEdgeLength;
+	Op->UseShapeGenerator->bUseConvexDecompositionMaxPieces = Settings->bLimitHullsPerShape || Settings->DecompositionMethod != EConvexDecompositionMethod::NavigationDriven;
+	Op->UseShapeGenerator->ConvexDecompositionMaxPieces = Settings->MaxHullsPerShape;
 	Op->UseShapeGenerator->ConvexDecompositionSearchFactor = Settings->ConvexDecompositionSearchFactor;
 	Op->UseShapeGenerator->ConvexDecompositionErrorTolerance = Settings->AddHullsErrorTolerance;
 	Op->UseShapeGenerator->ConvexDecompositionMinPartThickness = Settings->MinPartThickness;
-	Op->UseShapeGenerator->bConvexDecompositionProtectNegativeSpace = Settings->bUseNegativeSpaceInDecomposition;
+	Op->UseShapeGenerator->bConvexDecompositionProtectNegativeSpace = Settings->DecompositionMethod == EConvexDecompositionMethod::NavigationDriven;
 	Op->UseShapeGenerator->NegativeSpaceMinRadius = Settings->NegativeSpaceMinRadius;
 	Op->UseShapeGenerator->NegativeSpaceTolerance = Settings->NegativeSpaceTolerance;
 	Op->UseShapeGenerator->bIgnoreInternalNegativeSpace = Settings->bIgnoreInternalNegativeSpace;
diff --git a/Engine/Plugins/Experimental/MeshModelingToolsetExp/Source/MeshModelingToolsExp/Public/Physics/SetCollisionGeometryTool.h b/Engine/Plugins/Experimental/MeshModelingToolsetExp/Source/MeshModelingToolsExp/Public/Physics/SetCollisionGeometryTool.h
--- a/Engine/Plugins/Experimental/MeshModelingToolsetExp/Source/MeshModelingToolsExp/Public/Physics/SetCollisionGeometryTool.h
+++ b/Engine/Plugins/Experimental/MeshModelingToolsetExp/Source/MeshModelingToolsExp/Public/Physics/SetCollisionGeometryTool.h
@@ -76,10 +76,13 @@ enum class ECollisionGeometryType
 	Capsules = 4,
 	// Fit convex hulls to the inputs
 	ConvexHulls = 5,
+	// Fit multiple convex hulls to each input
+	ConvexDecompositions = 8,
 	// Fit convex hulls to 2D projections of the inputs, and sweep these 2D hulls along the projection dimension
 	SweptHulls = 6,
 	// Fit level sets to the inputs
 	LevelSets = 7,
+	// Note: ConvexDecomposition = 8 is set above, to be next to Convex Hulls in the UI
 	// Fit the boxes, spheres, and capsules to the inputs, and keep the best fitting of these shapes based on volume
 	MinVolume = 10,
 
@@ -106,6 +109,16 @@ enum class EProjectedHullAxis
 	SmallestVolume = 4
 };
 
+// Method to use to compute convex decomposition
+UENUM()
+enum class EConvexDecompositionMethod
+{
+	// Use the Navigable Space Protection's Error Tolerance and Min Radius settings to control the decomposition
+	NavigationDriven,
+	// Minimize volumetric differences to the input
+	VolumetricError
+};
+
 
 UCLASS()
 class MESHMODELINGTOOLSEXP_API USetCollisionGeometryToolProperties : public UInteractiveToolPropertySet
@@ -136,14 +149,14 @@ public:
 
 	/** Whether to discard all but MaxCount collision geometries with the largest volume */
 	UPROPERTY(EditAnywhere, Category = Options)
-	bool bEnableMaxCount = true;
+	bool bEnableMaxCount = false;
 
 	/** The maximum number of collision shapes to generate. If necessary, the shapes with smallest volume will be discarded to meet this count. */
 	UPROPERTY(EditAnywhere, Category = Options, meta = (UIMin = "1", UIMax = "100", ClampMin = "1", ClampMax = "9999999", EditCondition = "bEnableMaxCount"))
 	int32 MaxCount = 50;
 
-	/** Generated collision shapes will be expanded if they are smaller than this in any dimension. Not supported for Level Sets or Convex Decompositions (Convex Hulls with more than one hull per mesh). */
-	UPROPERTY(EditAnywhere, Category = Options, AdvancedDisplay, meta = (ClampMin = "0", UIMax = "10", EditCondition = "GeometryType != ECollisionGeometryType::LevelSets && (GeometryType != ECollisionGeometryType::ConvexHulls || MaxHullsPerMesh == 1)"))
+	/** Generated collision shapes will be expanded if they are smaller than this in any dimension. Not supported for Level Sets or Convex Decompositions. */
+	UPROPERTY(EditAnywhere, Category = Options, AdvancedDisplay, meta = (ClampMin = "0", UIMax = "10", EditCondition = "GeometryType != ECollisionGeometryType::LevelSets && GeometryType != ECollisionGeometryType::ConvexDecompositions"))
 	float MinThickness = 0.01;
 
 	/** Whether to override the requested Geometry Type with a box whenever a box closely fits the input shape */
@@ -170,51 +183,67 @@ public:
 	UPROPERTY(EditAnywhere, Category = MergeCollisionShapes, meta = (EditConditionHides, EditCondition = "bMergeCollisionShapes"))
 	bool bUseNegativeSpaceInMerge = false;
 
-	/** Whether to simplify the convex hull */
-	UPROPERTY(EditAnywhere, Category = ConvexHulls, meta = (EditConditionHides, EditCondition = "GeometryType == ECollisionGeometryType::ConvexHulls"))
+	/** Whether to simplify the convex hull down to at most a target face count. */
+	UPROPERTY(EditAnywhere, DisplayName = "Simplify to Face Count", Category = ConvexHulls, 
+		meta = (EditConditionHides, EditCondition = "GeometryType == ECollisionGeometryType::ConvexHulls || GeometryType == ECollisionGeometryType::ConvexDecompositions"))
 	bool bSimplifyHulls = true;
 
 	/** Target number of faces in the simplified hull */
 	UPROPERTY(EditAnywhere, Category = ConvexHulls, meta = (UIMin = "4", UIMax = "100", ClampMin = "4", ClampMax = "9999999",
-		EditConditionHides, EditCondition = "GeometryType == ECollisionGeometryType::ConvexHulls && bSimplifyHulls"))
+		EditConditionHides, EditCondition = "(GeometryType == ECollisionGeometryType::ConvexHulls  || GeometryType == ECollisionGeometryType::ConvexDecompositions) && bSimplifyHulls"))
 	int32 HullTargetFaceCount = 20;
 
-	/** How many convex hulls can be used to approximate each mesh */
+	/** Whether to simplify the input to this edge length before computing convex decomposition. Can give a decomposition result faster for large meshes. */
+	UPROPERTY(EditAnywhere, DisplayName = "Simplify to Edge Length", Category = ConvexHulls,
+		meta = (EditConditionHides, EditCondition = "GeometryType == ECollisionGeometryType::ConvexDecompositions"))
+	bool bPreSimplifyToEdgeLength = false;
+
+	/** Simplify the input to this edge length before computing convex decomposition. Can give a decomposition result faster for large meshes. */
+	UPROPERTY(EditAnywhere, DisplayName = "Target Edge Length", Category = ConvexHulls, meta = (UIMin = "0", UIMax = "100", ClampMin = "0",
+		EditConditionHides, EditCondition = "bPreSimplifyToEdgeLength && GeometryType == ECollisionGeometryType::ConvexDecompositions"))
+	double DecompositionTargetEdgeLength = 1.0;
+
+	// What algorithm should be used to perform convex decomposition
+	UPROPERTY(EditAnywhere, Category = ConvexHulls, meta = (
+		EditConditionHides, EditCondition = "GeometryType == ECollisionGeometryType::ConvexDecompositions"))
+	EConvexDecompositionMethod DecompositionMethod = EConvexDecompositionMethod::NavigationDriven;
+
+	// Whether to limit the number of convex hulls use in each decomposition
+	UPROPERTY(EditAnywhere, Category = ConvexHulls, meta = (
+		EditConditionHides, EditCondition = "GeometryType == ECollisionGeometryType::ConvexDecompositions && DecompositionMethod == EConvexDecompositionMethod::NavigationDriven"))
+	bool bLimitHullsPerShape = false;
+
+	// Maximum number of convex hulls to use in each decomposition
 	UPROPERTY(EditAnywhere, Category = ConvexHulls, meta = (UIMin = "1", UIMax = "100", ClampMin = "1",
-		EditConditionHides, EditCondition = "GeometryType == ECollisionGeometryType::ConvexHulls"))
-	int32 MaxHullsPerMesh = 1;
+		EditConditionHides, EditCondition = "GeometryType == ECollisionGeometryType::ConvexDecompositions && (bLimitHullsPerShape || DecompositionMethod != EConvexDecompositionMethod::NavigationDriven)"))
+	int32 MaxHullsPerShape = 10;
 
-	/** How much to search the space of possible decompositions beyond MaxHullsPerMesh; for larger values, will do additional work to try to better approximate mesh features (but resulting hulls may overlap more) */
-	UPROPERTY(EditAnywhere, Category = ConvexHulls, meta = (UIMin = "0", UIMax = "2", ClampMin = "0",
-		EditConditionHides, EditCondition = "GeometryType == ECollisionGeometryType::ConvexHulls && MaxHullsPerMesh > 1"))
+	/** How much to search the space of possible decompositions beyond Max Hulls Per Shape; for larger values, will do additional work to try to better approximate mesh features (but resulting hulls may overlap more) */
+	UPROPERTY(EditAnywhere, DisplayName = "Max Hulls Search Factor", Category = ConvexHulls, meta = (UIMin = "0", UIMax = "2", ClampMin = "0",
+		EditConditionHides, EditCondition = "GeometryType == ECollisionGeometryType::ConvexDecompositions && (DecompositionMethod != EConvexDecompositionMethod::NavigationDriven)"))
 	float ConvexDecompositionSearchFactor = .5;
 
 	/** Error tolerance for adding more convex hulls, in cm.  For volumetric errors, the value will be cubed (so a value of 10 indicates a 10x10x10 volume worth of error is acceptable). */
 	UPROPERTY(EditAnywhere, Category = ConvexHulls, meta = (UIMin = "0", UIMax = "1000", ClampMin = "0",
-		EditConditionHides, EditCondition = "GeometryType == ECollisionGeometryType::ConvexHulls && MaxHullsPerMesh > 1"))
+		EditConditionHides, EditCondition = "GeometryType == ECollisionGeometryType::ConvexDecompositions && DecompositionMethod != EConvexDecompositionMethod::NavigationDriven"))
 	float AddHullsErrorTolerance = 0;
 
 	/** Minimum part thickness for convex decomposition, in cm; hulls thinner than this will be merged into adjacent hulls, if possible. */
 	UPROPERTY(EditAnywhere, Category = ConvexHulls, meta = (UIMin = "0", UIMax = "1", ClampMin = "0",
-		EditConditionHides, EditCondition = "GeometryType == ECollisionGeometryType::ConvexHulls && MaxHullsPerMesh > 1"))
+		EditConditionHides, EditCondition = "GeometryType == ECollisionGeometryType::ConvexDecompositions"))
 	float MinPartThickness = 0.1;
 
-	/** Whether to guide the convex decomposition to prioritize not filling negative space of the input shape */
-	UPROPERTY(EditAnywhere, Category = ConvexHulls, meta = (
-		EditConditionHides, EditCondition = "GeometryType == ECollisionGeometryType::ConvexHulls && MaxHullsPerMesh > 1"))
-	bool bUseNegativeSpaceInDecomposition = false;
-
-	/** Negative space closer to the input than this tolerance distance can be filled in */
-	UPROPERTY(EditAnywhere, Category = NegativeSpace, meta = (UIMin = ".001", UIMax = "100", ClampMin = "0",
-		EditConditionHides, EditCondition = "(GeometryType == ECollisionGeometryType::ConvexHulls && MaxHullsPerMesh > 1 && bUseNegativeSpaceInDecomposition) || (bMergeCollisionShapes && bUseNegativeSpaceInMerge)"))
-	double NegativeSpaceTolerance = 3;
-	/** Minimum radius of negative space to protect; tunnels with radius smaller than this could be filled in */
-	UPROPERTY(EditAnywhere, Category = NegativeSpace, meta = (UIMin = ".001", UIMax = "100", ClampMin = "0",
-		EditConditionHides, EditCondition = "(GeometryType == ECollisionGeometryType::ConvexHulls && MaxHullsPerMesh > 1 && bUseNegativeSpaceInDecomposition) || (bMergeCollisionShapes && bUseNegativeSpaceInMerge)"))
-	double NegativeSpaceMinRadius = 10;
-	/** Whether to ignore negative space that is not accessible by traversing from the convex hull (via paths w/ radius of at least Negative Space Tolerance) */
-	UPROPERTY(EditAnywhere, Category = NegativeSpace, meta = (
-		EditConditionHides, EditCondition = "(GeometryType == ECollisionGeometryType::ConvexHulls && MaxHullsPerMesh > 1 && bUseNegativeSpaceInDecomposition) || (bMergeCollisionShapes && bUseNegativeSpaceInMerge)"))
+	/** Navigable space closer to the input than this tolerance distance can be filled in */
+	UPROPERTY(EditAnywhere, DisplayName = "Tolerance Distance", Category = NavigableSpaceProtection, meta = (UIMin = "1", UIMax = "100", ClampMin = "0.1",
+		EditConditionHides, EditCondition = "(GeometryType == ECollisionGeometryType::ConvexDecompositions && DecompositionMethod == EConvexDecompositionMethod::NavigationDriven) || (bMergeCollisionShapes && bUseNegativeSpaceInMerge)"))
+	double NegativeSpaceTolerance = 10;
+	/** Minimum radius of navigable space to protect; tunnels with radius smaller than this could be filled in */
+	UPROPERTY(EditAnywhere, DisplayName = "Min Radius", Category = NavigableSpaceProtection, meta = (UIMin = "1", UIMax = "200", ClampMin = "0.1",
+		EditConditionHides, EditCondition = "(GeometryType == ECollisionGeometryType::ConvexDecompositions && DecompositionMethod == EConvexDecompositionMethod::NavigationDriven) || (bMergeCollisionShapes && bUseNegativeSpaceInMerge)"))
+	double NegativeSpaceMinRadius = 40;
+	/** Whether to ignore navigable space that is not accessible by from outside (e.g., closed-off interiors / air pockets) */
+	UPROPERTY(EditAnywhere, Category = NavigableSpaceProtection, meta = (
+		EditConditionHides, EditCondition = "(GeometryType == ECollisionGeometryType::ConvexDecompositions && DecompositionMethod == EConvexDecompositionMethod::NavigationDriven) || (bMergeCollisionShapes && bUseNegativeSpaceInMerge)"))
 	bool bIgnoreInternalNegativeSpace = true;
 
 	/** If > 0, the polygon used to generate the swept hull will be simplified up to this distance tolerance, in cm */
diff --git a/Engine/Plugins/Runtime/GeometryProcessing/Source/DynamicMesh/Private/ShapeApproximation/MeshSimpleShapeApproximation.cpp b/Engine/Plugins/Runtime/GeometryProcessing/Source/DynamicMesh/Private/ShapeApproximation/MeshSimpleShapeApproximation.cpp
--- a/Engine/Plugins/Runtime/GeometryProcessing/Source/DynamicMesh/Private/ShapeApproximation/MeshSimpleShapeApproximation.cpp
+++ b/Engine/Plugins/Runtime/GeometryProcessing/Source/DynamicMesh/Private/ShapeApproximation/MeshSimpleShapeApproximation.cpp
@@ -14,6 +14,7 @@
 #include "Operations/MeshConvexHull.h"
 #include "Operations/MeshProjectionHull.h"
 #include "Util/ProgressCancel.h"
+#include "MeshSimplification.h"
 
 #define LOCTEXT_NAMESPACE "MeshSimpleShapeApproximation"
 
@@ -387,27 +388,84 @@ void FMeshSimpleShapeApproximation::Generate_ConvexHullDecompositions(FSimpleSha
 		}
 
 		const FDynamicMesh3& SourceMesh = *SourceMeshes[idx];
-		// TODO: if (bSimplifyHulls), also consider simplifying the input?
-		FConvexDecomposition3 Decomposition(SourceMesh);
-		int32 NumAdditionalSplits = FMath::FloorToInt32(float(ConvexDecompositionMaxPieces) * ConvexDecompositionSearchFactor);
+		FConvexDecomposition3::FPreprocessMeshOptions PreprocessOptions;
+		PreprocessOptions.bMergeEdges = true;
+		PreprocessOptions.CustomPreprocess = [this](FDynamicMesh3& ProcessMesh, const FAxisAlignedBox3d& Bounds) -> void
+		{
+			// for solid inputs, flip orientation if the initial volume is negative
+			if (ProcessMesh.IsClosed())
+			{
+				double InitialVolume = TMeshQueries<FDynamicMesh3>::GetVolumeArea(ProcessMesh).X;
+				if (InitialVolume < 0)
+				{
+					ProcessMesh.ReverseOrientation();
+				}
+			}
+			
+			if (bDecompositionPreSimplifyWithEdgeLength)
+			{
+				// Run pre-simplification to the target edge length
+				UE::Geometry::FVolPresMeshSimplification Simplifier(&ProcessMesh);
+				Simplifier.CollapseMode = UE::Geometry::FVolPresMeshSimplification::ESimplificationCollapseModes::MinimalExistingVertexError;
+				Simplifier.SimplifyToEdgeLength(DecompositionPreSimplifyEdgeLength);
+			}
+		};
+		FConvexDecomposition3 Decomposition(SourceMesh, PreprocessOptions);
+		const bool bIsSolid = Decomposition.IsInputSolid();
+		Decomposition.bTreatAsSolid = bIsSolid;
+
 		if (bConvexDecompositionProtectNegativeSpace)
 		{
+			// Use settings tuned for navigation-driven decomposition
+			
 			FNegativeSpaceSampleSettings Settings;
+			Settings.MarchingCubesGridScale = 1.0;
+			Settings.MaxVoxelsPerDim = 1024;
+			Settings.MinSpacing = 0.0;
+			Settings.TargetNumSamples = 0;
+			Settings.VoxelExpandBoundsFactor = UE_DOUBLE_KINDA_SMALL_NUMBER;
 			Settings.bOnlyConnectedToHull = bIgnoreInternalNegativeSpace;
 			Settings.MinRadius = NegativeSpaceMinRadius;
 			Settings.ReduceRadiusMargin = NegativeSpaceTolerance;
 			Settings.MinRadius = FMath::Max(1, (NegativeSpaceMinRadius + NegativeSpaceTolerance) * .5);
-			Settings.SampleMethod = FNegativeSpaceSampleSettings::ESampleMethod::VoxelSearch;
+			Settings.SampleMethod = FNegativeSpaceSampleSettings::ESampleMethod::NavigableVoxelSearch;
 			Settings.bRequireSearchSampleCoverage = true;
-			Settings.TargetNumSamples = 1; // let the sample coverage determine the number of spheres to place
-
+			Settings.TargetNumSamples = 0; // let the sample coverage determine the number of spheres to place
+			Settings.bAllowSamplesInsideMesh = !bIsSolid;
+
+			Decomposition.MaxConvexEdgePlanes = 4;
+			Decomposition.bSplitDisconnectedComponents = false;
+			Decomposition.ConvexEdgeAngleMoreSamplesThreshold = 180;
+			Decomposition.ThickenAfterHullFailure = FMath::Max(FMathd::ZeroTolerance, NegativeSpaceTolerance * .01);
+			
 			Decomposition.InitializeNegativeSpace(Settings);
 
-			// Let negative space decide when to stop merging; target only 1 piece if negative space allows
-			NumAdditionalSplits += ConvexDecompositionMaxPieces;
-			ConvexDecompositionMaxPieces = 1;
+			constexpr int32 MaxAllowedSplits = 1000000; // more parts than any expected / reasonable decomposition
+			int32 TargetNumSplits = bUseConvexDecompositionMaxPieces ? ConvexDecompositionMaxPieces - 1 : MaxAllowedSplits + 1;
+			for (int32 Split = 0; Split < TargetNumSplits; Split++)
+			{
+				int32 NumSplit = Decomposition.SplitWorst(false, -1, true, Settings.ReduceRadiusMargin * .5);
+
+				if (NumSplit == 0)
+				{
+					break;
+				}
+
+				if (!ensureMsgf(Split < MaxAllowedSplits, TEXT("Convex decomposition split the input %d times; likely stuck in a loop"), Split))
+				{
+					break;
+				}
+			}
+
+			Decomposition.FixHullOverlapsInNegativeSpace();
+			int32 TargetNumPieces = bUseConvexDecompositionMaxPieces ? ConvexDecompositionMaxPieces : -1;
+			Decomposition.MergeBest(TargetNumPieces, 0, ConvexDecompositionMinPartThickness, true);
+		}
+		else
+		{
+			int32 NumAdditionalSplits = FMath::FloorToInt32(float(ConvexDecompositionMaxPieces) * ConvexDecompositionSearchFactor);
+			Decomposition.Compute(ConvexDecompositionMaxPieces, NumAdditionalSplits, ConvexDecompositionErrorTolerance, ConvexDecompositionMinPartThickness);
 		}
-		Decomposition.Compute(ConvexDecompositionMaxPieces, NumAdditionalSplits, ConvexDecompositionErrorTolerance, ConvexDecompositionMinPartThickness);
 
 		for (int32 HullIdx = 0; HullIdx < Decomposition.NumHulls(); HullIdx++)
 		{
diff --git a/Engine/Plugins/Runtime/GeometryProcessing/Source/DynamicMesh/Public/ShapeApproximation/MeshSimpleShapeApproximation.h b/Engine/Plugins/Runtime/GeometryProcessing/Source/DynamicMesh/Public/ShapeApproximation/MeshSimpleShapeApproximation.h
--- a/Engine/Plugins/Runtime/GeometryProcessing/Source/DynamicMesh/Public/ShapeApproximation/MeshSimpleShapeApproximation.h
+++ b/Engine/Plugins/Runtime/GeometryProcessing/Source/DynamicMesh/Public/ShapeApproximation/MeshSimpleShapeApproximation.h
@@ -68,15 +68,25 @@ public:
 	/** simplification tolerance when simplifying 2D convex hulls, eg for swept/projected hulls */
 	double HullSimplifyTolerance = 1.0;
 
-	/** How many convex pieces to target per mesh when creating convex decompositions.  If ErrorTolerance is set, can create fewer pieces. */
-	int32 ConvexDecompositionMaxPieces = 20;
-	/** How much additional decomposition decomposition + merging to do, as a fraction of max pieces.  Larger values can help better-cover small features, while smaller values create a cleaner decomposition with less overlap between hulls. */
+	/** Whether to use apply an edge-length based simplification to the input before running convex decompositions. Useful for very dense input meshes where the decomposition will be slow to compute. */
+	bool bDecompositionPreSimplifyWithEdgeLength = false;
+	/** If enabled by the above bool flag, pre-simplify input geometry to this edge length before computing convex decompositions. */
+	double DecompositionPreSimplifyEdgeLength = 1.0;
+
+	/** How many convex pieces to target per mesh when creating convex decompositions. Ignored if < 1. If ErrorTolerance or ProtectNegativeSpace are used, can create fewer pieces. */
+	int32 ConvexDecompositionMaxPieces = 1;
+	/** Whether to use the above Max Pieces to drive the convex decomposition. Otherwise, will allow the error tolerances / negative space protection settings to drive the number of pieces generated. */
+	bool bUseConvexDecompositionMaxPieces = true;
+	/** 
+	 * How much additional decomposition decomposition + merging to do, as a fraction of max pieces.  Larger values can help better-cover small features, while smaller values create a cleaner decomposition with less overlap between hulls. 
+	 * Note: Not used if bConvexDecompositionProtectNegativeSpace is true.
+	 */
 	float ConvexDecompositionSearchFactor = .5;
 	/** Error tolerance to guide convex decomposition (in cm); we stop adding new parts if the volume error is below the threshold.  For volumetric errors, value will be cubed. */
 	double ConvexDecompositionErrorTolerance = 0;
 	/** Minimum part thickness for convex decomposition (in cm); hulls thinner than this will be merged into adjacent hulls, if possible. */
 	double ConvexDecompositionMinPartThickness = .1;
-	/** Whether to guide the convex decomposition to prioritize not filling negative space of the input shape */
+	/** Whether to use 'navigation-driven' convex decomposition -- using NegativeSpaceTolerance and NegativeSpaceMinRadius to define space that the decomposition hulls cannot occupy */
 	bool bConvexDecompositionProtectNegativeSpace = false;
 
 	/** Negative space closer to the input than this tolerance distance can be filled in */
