Skeleton Editing Tool: defer pending action on tick if the viewport does not currently have focus
    the input router collects behaviors to be captured on mouse down and processes them before notifying the previous widget that it has lost focus.
    deferring the behavior function on tick allows other widgets to send their "lost focus" notification before the tool does its job.
    note that this is not specific to this tool.

#rnx

diff --git a/Engine/Plugins/Experimental/Animation/SkeletalMeshModelingTools/Source/SkeletalMeshModelingTools/Private/DetailCustomization/SkeletonEditingToolPropertyCustomizations.cpp b/Engine/Plugins/Experimental/Animation/SkeletalMeshModelingTools/Source/SkeletalMeshModelingTools/Private/DetailCustomization/SkeletonEditingToolPropertyCustomizations.cpp
--- a/Engine/Plugins/Experimental/Animation/SkeletalMeshModelingTools/Source/SkeletalMeshModelingTools/Private/DetailCustomization/SkeletonEditingToolPropertyCustomizations.cpp
+++ b/Engine/Plugins/Experimental/Animation/SkeletalMeshModelingTools/Source/SkeletalMeshModelingTools/Private/DetailCustomization/SkeletonEditingToolPropertyCustomizations.cpp
@@ -120,9 +120,9 @@ void FSkeletonEditingPropertiesDetailCustomization::CustomizeDetails(IDetailLayo
 						}
 						return MultiValues;
 					})
-					.OnTextCommitted_Lambda([this, NodePropHandle](const FText& NewText, ETextCommit::Type) 
+					.OnTextCommitted_Lambda([this, NodePropHandle](const FText& NewText, ETextCommit::Type InCommitType) 
 					{
-						if (NewText.EqualTo(MultiValues))
+						if (InCommitType == ETextCommit::OnCleared || NewText.EqualTo(MultiValues))
 						{
 							return;
 						}
diff --git a/Engine/Plugins/Experimental/MeshModelingToolsetExp/Source/MeshModelingToolsEditorOnlyExp/Private/SkeletalMesh/SkeletonEditingTool.cpp b/Engine/Plugins/Experimental/MeshModelingToolsetExp/Source/MeshModelingToolsEditorOnlyExp/Private/SkeletalMesh/SkeletonEditingTool.cpp
--- a/Engine/Plugins/Experimental/MeshModelingToolsetExp/Source/MeshModelingToolsEditorOnlyExp/Private/SkeletalMesh/SkeletonEditingTool.cpp
+++ b/Engine/Plugins/Experimental/MeshModelingToolsetExp/Source/MeshModelingToolsEditorOnlyExp/Private/SkeletalMesh/SkeletonEditingTool.cpp
@@ -1303,8 +1303,11 @@ void USkeletonEditingTool::OnClickPress(const FInputDeviceRay& InPressPos)
 {
 	if (PendingFunction)
 	{
-		PendingFunction();
-		PendingFunction.Reset();
+		if (!bDeferUntilFocused)
+		{
+			PendingFunction();
+			PendingFunction.Reset();
+		}
 	}
 	else
 	{ // make sure that the PendingFunction handles BeginChange if it needs to
@@ -1314,6 +1317,11 @@ void USkeletonEditingTool::OnClickPress(const FInputDeviceRay& InPressPos)
 
 void USkeletonEditingTool::OnClickDrag(const FInputDeviceRay& InDragPos)
 {
+	if (bDeferUntilFocused)
+	{
+		return;
+	}
+	
 	if (Operation != EEditingOperation::Create)
 	{
 		return;
@@ -1532,6 +1540,9 @@ void USkeletonEditingTool::OnTerminateDragSequence()
 
 void USkeletonEditingTool::OnTick(float DeltaTime)
 {
+	const FViewport* Viewport = GetToolManager()->GetContextQueriesAPI()->GetFocusedViewport();
+	bDeferUntilFocused = Viewport && !Viewport->HasFocus();
+	
 	if (PendingFunction)
 	{
 		PendingFunction();
@@ -1548,9 +1559,11 @@ void USkeletonEditingTool::OnTick(float DeltaTime)
 
 FInputRayHit USkeletonEditingTool::CanBeginClickDragSequence(const FInputDeviceRay& InPressPos)
 {
+	static const FInputRayHit InvalidRayHit;
+	
 	if (Properties->bEnableComponentSelection)
 	{
-		return FInputRayHit();
+		return InvalidRayHit;
 	}
 	
 	PendingFunction.Reset();
@@ -1559,12 +1572,14 @@ FInputRayHit USkeletonEditingTool::CanBeginClickDragSequence(const FInputDeviceR
 	FViewport* Viewport = GetToolManager()->GetContextQueriesAPI()->GetFocusedViewport();
 	if (!Viewport)
 	{
-		return FInputRayHit();
+		return InvalidRayHit;
 	}
+
+	bDeferUntilFocused = !Viewport->HasFocus();
 	
 	if (GizmoWrapper && GizmoWrapper->IsGizmoHit(InPressPos))
 	{
-		return FInputRayHit();
+		return InvalidRayHit;
 	}
 	
 	auto PickBone = [&]() -> int32
@@ -1596,7 +1611,7 @@ FInputRayHit USkeletonEditingTool::CanBeginClickDragSequence(const FInputDeviceR
 		{
 			const FReferenceSkeleton& ReferenceSkeleton = Modifier->GetReferenceSkeleton();
 			ParentBones(ReferenceSkeleton.GetBoneName(BoneIndex));
-			return FInputRayHit();
+			return InvalidRayHit;
 		}
 		
 		// otherwise, update current selection
@@ -1618,7 +1633,7 @@ FInputRayHit USkeletonEditingTool::CanBeginClickDragSequence(const FInputDeviceR
 	{
 		Selection.Empty();
 		Properties->Name = GetCurrentBone();
-		return FInputRayHit();
+		return InvalidRayHit;
 	}
 
 	// if we're in creation mode then create a new bone
@@ -1643,7 +1658,7 @@ FInputRayHit USkeletonEditingTool::CanBeginClickDragSequence(const FInputDeviceR
 		}
 	}
 	
-	return FInputRayHit();
+	return InvalidRayHit;
 }
 
 TWeakObjectPtr<USkeletonModifier> USkeletonEditingTool::GetModifier() const
diff --git a/Engine/Plugins/Experimental/MeshModelingToolsetExp/Source/MeshModelingToolsEditorOnlyExp/Public/SkeletalMesh/SkeletonEditingTool.h b/Engine/Plugins/Experimental/MeshModelingToolsetExp/Source/MeshModelingToolsEditorOnlyExp/Public/SkeletalMesh/SkeletonEditingTool.h
--- a/Engine/Plugins/Experimental/MeshModelingToolsetExp/Source/MeshModelingToolsEditorOnlyExp/Public/SkeletalMesh/SkeletonEditingTool.h
+++ b/Engine/Plugins/Experimental/MeshModelingToolsetExp/Source/MeshModelingToolsEditorOnlyExp/Public/SkeletalMesh/SkeletonEditingTool.h
@@ -289,6 +289,9 @@ private:
 	TFunction<void()> PendingFunction;
 
 	TUniquePtr<UE::Geometry::FTriangleGroupTopology> Topology = nullptr;
+
+	// Defer pending function on tick to allow other external notification to be handled. 
+	bool bDeferUntilFocused = false;
 };
 
 ENUM_CLASS_FLAGS(USkeletonEditingTool::EBoneSelectionMode);
