Multicolor weight display in skin paint tool.

diff --git a/Engine/Plugins/Experimental/MeshModelingToolsetExp/Source/MeshModelingToolsEditorOnlyExp/Private/SkeletalMesh/SkinWeightsPaintTool.cpp b/Engine/Plugins/Experimental/MeshModelingToolsetExp/Source/MeshModelingToolsEditorOnlyExp/Private/SkeletalMesh/SkinWeightsPaintTool.cpp
--- a/Engine/Plugins/Experimental/MeshModelingToolsetExp/Source/MeshModelingToolsEditorOnlyExp/Private/SkeletalMesh/SkinWeightsPaintTool.cpp
+++ b/Engine/Plugins/Experimental/MeshModelingToolsetExp/Source/MeshModelingToolsEditorOnlyExp/Private/SkeletalMesh/SkinWeightsPaintTool.cpp
@@ -26,6 +26,7 @@
 #include "Spatial/PointSetHashTable.h"
 #include "Operations/SmoothBoneWeights.h"
 #include "ContextObjectStore.h"
+#include "SkeletalDebugRendering.h"
 #include "Editor/Persona/Public/IPersonaEditorModeManager.h"
 #include "Editor/Persona/Public/PersonaModule.h"
 
@@ -1149,47 +1150,82 @@ void USkinWeightsPaintTool::CalculateVertexROI(
 	checkNoEntry();
 }
 
-FVector4f USkinWeightsPaintTool::WeightToColor(float Value) const
+FVector4f USkinWeightsPaintTool::GetColorOfVertex(VertexIndex InVertexIndex, BoneIndex InCurrentBoneIndex) const
 {
-	// optional greyscale mode
-	if (WeightToolProperties->ColorMode == EWeightColorMode::MinMax)
+	switch (WeightToolProperties->ColorMode)
 	{
-		return FMath::Lerp(WeightToolProperties->MinColor, WeightToolProperties->MaxColor, Value);
-	}
-	
-	// early out zero weights to min color
-	if (Value <= MinimumWeightThreshold)
-	{
-		return WeightToolProperties->MinColor;
-	}
+	case EWeightColorMode::MinMax:
+		{
+			if (InCurrentBoneIndex == INDEX_NONE)
+			{
+				return WeightToolProperties->MinColor; // with no bone selected, all vertices are drawn black
+			}
+			const float Value = Weights.GetWeightOfBoneOnVertex(InCurrentBoneIndex, InVertexIndex, Weights.CurrentWeights);
+			return FMath::Lerp(WeightToolProperties->MinColor, WeightToolProperties->MaxColor, Value);
+		}
+	case EWeightColorMode::Ramp:
+		{
+			if (InCurrentBoneIndex == INDEX_NONE)
+			{
+				return WeightToolProperties->MinColor; // with no bone selected, all vertices are drawn black
+			}
+			
+			float Value = Weights.GetWeightOfBoneOnVertex(InCurrentBoneIndex, InVertexIndex, Weights.CurrentWeights);
+			Value = FMath::Clamp(Value, 0.0f, 1.0f);
+			
+			// early out zero weights to min color
+			if (Value <= MinimumWeightThreshold)
+			{
+				return WeightToolProperties->MinColor;
+			}
 
-	// early out full weights to max color
-	if (FMath::IsNearlyEqual(Value, 1.0f))
-	{
-		return WeightToolProperties->MaxColor;
-	}
+			// early out full weights to max color
+			if (FMath::IsNearlyEqual(Value, 1.0f))
+			{
+				return WeightToolProperties->MaxColor;
+			}
 
-	// get user-specified color ramp for intermediate colors
-	const TArray<FLinearColor>& Colors = WeightToolProperties->ColorRamp;
+			// get user-specified color ramp for intermediate colors
+			const TArray<FLinearColor>& Colors = WeightToolProperties->ColorRamp;
 	
-	// revert back to simple Lerp(min,max) if user supplied color ramp doesn't have enough colors
-	if (Colors.Num() < 2)
-	{
-		const FLinearColor FinalColor = FMath::Lerp(WeightToolProperties->MinColor, WeightToolProperties->MaxColor, Value);
-		return UE::Geometry::ToVector4<float>(FinalColor);
-	}
+			// revert back to simple Lerp(min,max) if user supplied color ramp doesn't have enough colors
+			if (Colors.Num() < 2)
+			{
+				return UE::Geometry::ToVector4<float>(FMath::Lerp(WeightToolProperties->MinColor, WeightToolProperties->MaxColor, Value));
+			}
 
-	// otherwise, interpolate within two nearest ramp colors
-	Value = FMath::Clamp(Value, 0.0f, 1.0f);
-	const float PerColorRange = 1.0f / (Colors.Num() - 1);
-	const int ColorIndex = static_cast<int>(Value / PerColorRange);
-	const float RangeStart = ColorIndex * PerColorRange;
-	const float RangeEnd = (ColorIndex + 1) * PerColorRange;
-	const float Param = (Value - RangeStart) / (RangeEnd - RangeStart);
-	const FLinearColor& StartColor = Colors[ColorIndex];
-	const FLinearColor& EndColor = Colors[ColorIndex+1];
-	const FLinearColor FinalColor = FMath::Lerp(StartColor, EndColor, Param);
-	return UE::Geometry::ToVector4<float>(FinalColor);
+			// otherwise, interpolate within two nearest ramp colors
+			const float PerColorRange = 1.0f / (Colors.Num() - 1);
+			const int ColorIndex = static_cast<int>(Value / PerColorRange);
+			const float RangeStart = ColorIndex * PerColorRange;
+			const float RangeEnd = (ColorIndex + 1) * PerColorRange;
+			const float Param = (Value - RangeStart) / (RangeEnd - RangeStart);
+			const FLinearColor& StartColor = Colors[ColorIndex];
+			const FLinearColor& EndColor = Colors[ColorIndex+1];
+			return UE::Geometry::ToVector4<float>(FMath::Lerp(StartColor, EndColor, Param));
+		}
+	case EWeightColorMode::MultiColor:
+		{
+			FVector4f Color = FVector4f::Zero();
+			const VertexWeights& VertexWeights = Weights.CurrentWeights[InVertexIndex];
+			for (const FVertexBoneWeight& BoneWeight : VertexWeights)
+			{
+				if (BoneWeight.Weight < KINDA_SMALL_NUMBER)
+				{
+					continue;
+				}
+				
+				const float Value = InCurrentBoneIndex == BoneWeight.BoneIndex ? 1.0f: 0.25f;
+				constexpr float Saturation = 1.f;
+				const FLinearColor BoneColor = SkeletalDebugRendering::GetSemiRandomColorForBone(BoneWeight.BoneIndex, Value, Saturation);
+				Color = FLinearColor::LerpUsingHSV(Color, BoneColor, BoneWeight.Weight);
+			}
+			return Color;
+		}
+	default:
+		checkNoEntry();
+		return FLinearColor::Black;
+	}	
 }
 
 
@@ -1204,19 +1240,11 @@ void USkinWeightsPaintTool::UpdateCurrentBoneVertexColors()
 		UE::Geometry::FDynamicMeshColorOverlay* ColorOverlay = Mesh.Attributes()->PrimaryColors();
 		for (const int32 ElementId : ColorOverlay->ElementIndicesItr())
 		{
-			// with no bone selected, all vertices are drawn black
-			if (CurrentBoneIndex == INDEX_NONE)
-			{
-				ColorOverlay->SetElement(ElementId, FVector4f(FLinearColor::Black));
-				continue;
-			}
-			
 			const int32 VertexID = ColorOverlay->GetParentVertex(ElementId);	
 			const int32 SrcVertexID = NonManifoldMappingSupport.GetOriginalNonManifoldVertexID(VertexID);
-			const float Value = Weights.GetWeightOfBoneOnVertex(CurrentBoneIndex, SrcVertexID, Weights.CurrentWeights);
-			const FVector4f Color(WeightToColor(Value));
-			ColorOverlay->SetElement(ElementId, Color);
+			ColorOverlay->SetElement(ElementId, GetColorOfVertex(SrcVertexID, CurrentBoneIndex));
 		}
+		
 	}, false);
 	PreviewMesh->NotifyDeferredEditCompleted(UPreviewMesh::ERenderUpdateMode::FastUpdate, EMeshRenderAttributeFlags::VertexColors, false);
 }
@@ -1301,13 +1329,12 @@ void USkinWeightsPaintTool::ApplyStamp(const FBrushStampData& Stamp)
 		{
 			TArray<int> ElementIds;
 			UE::Geometry::FDynamicMeshColorOverlay* ColorOverlay = Mesh.Attributes()->PrimaryColors();
-			const int32 BoneIndex = Weights.Deformer.BoneNameToIndexMap[CurrentBone];
+			const int32 CurrentBoneIndex = Weights.Deformer.BoneNameToIndexMap[CurrentBone];
 			const int32 NumVerticesInStamp = VerticesInStamp.Num();
-			for (int32 Index = 0; Index < NumVerticesInStamp; ++Index)
+			for (int32 StampVertexIndex = 0; StampVertexIndex < NumVerticesInStamp; ++StampVertexIndex)
 			{
-				const int32 VertexID = VerticesInStamp[Index];
-				const float Weight = Weights.GetWeightOfBoneOnVertex(BoneIndex, VertexID, Weights.CurrentWeights);
-				FVector4f NewColor(WeightToColor(Weight));
+				const int32 VertexID = VerticesInStamp[StampVertexIndex];
+				FVector4f NewColor(GetColorOfVertex(VertexID, CurrentBoneIndex));
 				ColorOverlay->GetVertexElements(VertexID, ElementIds);
 				for (const int32 ElementId : ElementIds)
 				{
@@ -1315,6 +1342,7 @@ void USkinWeightsPaintTool::ApplyStamp(const FBrushStampData& Stamp)
 				}
 				ElementIds.Reset();
 			}
+			
 		}, false);
 		PreviewMesh->NotifyDeferredEditCompleted(UPreviewMesh::ERenderUpdateMode::FastUpdate, EMeshRenderAttributeFlags::VertexColors, false);
 	}
diff --git a/Engine/Plugins/Experimental/MeshModelingToolsetExp/Source/MeshModelingToolsEditorOnlyExp/Public/SkeletalMesh/SkinWeightsPaintTool.h b/Engine/Plugins/Experimental/MeshModelingToolsetExp/Source/MeshModelingToolsEditorOnlyExp/Public/SkeletalMesh/SkinWeightsPaintTool.h
--- a/Engine/Plugins/Experimental/MeshModelingToolsetExp/Source/MeshModelingToolsEditorOnlyExp/Public/SkeletalMesh/SkinWeightsPaintTool.h
+++ b/Engine/Plugins/Experimental/MeshModelingToolsetExp/Source/MeshModelingToolsEditorOnlyExp/Public/SkeletalMesh/SkinWeightsPaintTool.h
@@ -48,6 +48,7 @@ enum class EWeightColorMode : uint8
 {
 	MinMax,
 	Ramp,
+	MultiColor,
 };
 
 // brush falloff mode
@@ -480,7 +481,7 @@ protected:
 
 	// vertex colors updated when switching current bone or editing weights
 	void UpdateCurrentBoneVertexColors();
-	FVector4f WeightToColor(float Value) const;
+	FVector4f GetColorOfVertex(VertexIndex InVertexIndex, BoneIndex InBoneIndex) const;
 	bool bVisibleWeightsValid = false;
 
 	// which bone are we currently painting?
diff --git a/Engine/Plugins/Experimental/MeshModelingToolsetExp/Source/ModelingEditorUI/Private/SkinWeightDetailCustomization.cpp b/Engine/Plugins/Experimental/MeshModelingToolsetExp/Source/ModelingEditorUI/Private/SkinWeightDetailCustomization.cpp
--- a/Engine/Plugins/Experimental/MeshModelingToolsetExp/Source/ModelingEditorUI/Private/SkinWeightDetailCustomization.cpp
+++ b/Engine/Plugins/Experimental/MeshModelingToolsetExp/Source/ModelingEditorUI/Private/SkinWeightDetailCustomization.cpp
@@ -116,7 +116,9 @@ void FSkinWeightDetailCustomization::CustomizeDetails(IDetailLayoutBuilder& Deta
 				+ SSegmentedControl<EWeightColorMode>::Slot(EWeightColorMode::MinMax)
 				.Text(LOCTEXT("MinMaxMode", "Min / Max"))
 				+SSegmentedControl<EWeightColorMode>::Slot(EWeightColorMode::Ramp)
-				.Text(LOCTEXT("RampMode", "Color Ramp"))
+				.Text(LOCTEXT("RampMode", "Ramp"))
+				+SSegmentedControl<EWeightColorMode>::Slot(EWeightColorMode::MultiColor)
+				.Text(LOCTEXT("MultiColorMode", "Multi Color"))
 			]
 		]
 	];
diff --git a/Engine/Source/Runtime/Engine/Private/SkeletalDebugRendering.cpp b/Engine/Source/Runtime/Engine/Private/SkeletalDebugRendering.cpp
--- a/Engine/Source/Runtime/Engine/Private/SkeletalDebugRendering.cpp
+++ b/Engine/Source/Runtime/Engine/Private/SkeletalDebugRendering.cpp
@@ -483,11 +483,9 @@ void DrawBonesInternal(
 	}
 }
 
-FLinearColor GetSemiRandomColorForBone(const int32 BoneIndex)
+FLinearColor GetSemiRandomColorForBone(const int32 BoneIndex, float Value, float Saturation)
 {
 	// uses deterministic, semi-random desaturated color unique to the bone index
-	constexpr float Saturation = 0.5f;
-	constexpr float Value = 1.0f;
 	constexpr float Rotation = 90.f;
 	return FLinearColor::IntToDistinctColor(BoneIndex, Saturation, Value, Rotation);
 }
diff --git a/Engine/Source/Runtime/Engine/Public/SkeletalDebugRendering.h b/Engine/Source/Runtime/Engine/Public/SkeletalDebugRendering.h
--- a/Engine/Source/Runtime/Engine/Public/SkeletalDebugRendering.h
+++ b/Engine/Source/Runtime/Engine/Public/SkeletalDebugRendering.h
@@ -153,6 +153,6 @@ void DrawBonesInternal(
 	const TArray<TRefCountPtr<HHitProxy>>& HitProxies,
 	const FSkelDebugDrawConfig& DrawConfig);
 
-ENGINE_API FLinearColor GetSemiRandomColorForBone(const int32 BoneIndex);
+ENGINE_API FLinearColor GetSemiRandomColorForBone(const int32 BoneIndex, float Value=1.0f, float Saturation=0.5f);
 ENGINE_API void FillWithMultiColors(TArray<FLinearColor>& BoneColors, const int32 NumBones);
 }
