Fixing engine crash from PreviewProfileController

diff --git a/Engine/Plugins/Experimental/MeshModelingToolsetExp/Source/MeshModelingToolsEditorOnlyExp/Private/SkeletalMesh/SkinWeightsPaintTool.cpp b/Engine/Plugins/Experimental/MeshModelingToolsetExp/Source/MeshModelingToolsEditorOnlyExp/Private/SkeletalMesh/SkinWeightsPaintTool.cpp
--- a/Engine/Plugins/Experimental/MeshModelingToolsetExp/Source/MeshModelingToolsEditorOnlyExp/Private/SkeletalMesh/SkinWeightsPaintTool.cpp
+++ b/Engine/Plugins/Experimental/MeshModelingToolsetExp/Source/MeshModelingToolsEditorOnlyExp/Private/SkeletalMesh/SkinWeightsPaintTool.cpp
@@ -766,6 +766,7 @@ void USkinWeightsPaintTool::Setup()
 	}
 
 	// modify viewport render settings to optimize for painting weights
+	FPreviewProfileController PreviewProfileController;
 	PreviewProfileToRestore = PreviewProfileController.GetActiveProfile();
 	PreviewProfileController.SetActiveProfile(UDefaultEditorProfiles::EditingProfileName.ToString());
 
@@ -1579,6 +1580,7 @@ void USkinWeightsPaintTool::OnShutdown(EToolShutdownType ShutdownType)
 	}
 
 	// restore viewport show flags and preview settings
+	FPreviewProfileController PreviewProfileController;
 	PreviewProfileController.SetActiveProfile(PreviewProfileToRestore);
 
 	if (EditorContext.IsValid())
diff --git a/Engine/Plugins/Experimental/MeshModelingToolsetExp/Source/MeshModelingToolsEditorOnlyExp/Public/SkeletalMesh/SkinWeightsPaintTool.h b/Engine/Plugins/Experimental/MeshModelingToolsetExp/Source/MeshModelingToolsEditorOnlyExp/Public/SkeletalMesh/SkinWeightsPaintTool.h
--- a/Engine/Plugins/Experimental/MeshModelingToolsetExp/Source/MeshModelingToolsEditorOnlyExp/Public/SkeletalMesh/SkinWeightsPaintTool.h
+++ b/Engine/Plugins/Experimental/MeshModelingToolsetExp/Source/MeshModelingToolsEditorOnlyExp/Public/SkeletalMesh/SkinWeightsPaintTool.h
@@ -511,7 +511,6 @@ protected:
 	TWeakObjectPtr<UPersonaEditorModeManagerContext> PersonaModeManagerContext = nullptr;
 
 	// editor state to restore when exiting the paint tool
-	FPreviewProfileController PreviewProfileController;
 	FString PreviewProfileToRestore;
 
 	friend SkinPaintTool::FSkinToolDeformer;
diff --git a/Engine/Source/Editor/AdvancedPreviewScene/Private/PreviewProfileController.cpp b/Engine/Source/Editor/AdvancedPreviewScene/Private/PreviewProfileController.cpp
--- a/Engine/Source/Editor/AdvancedPreviewScene/Private/PreviewProfileController.cpp
+++ b/Engine/Source/Editor/AdvancedPreviewScene/Private/PreviewProfileController.cpp
@@ -36,7 +36,7 @@ FPreviewProfileController::FPreviewProfileController()
 
 FPreviewProfileController::~FPreviewProfileController()
 {
-	if (AssetViewerSettings)
+	if (IsValid(AssetViewerSettings))
 	{
 		AssetViewerSettings->OnAssetViewerProfileAddRemoved().Remove(AssetViewerSettingsProfileAddRemoveHandle);
 		AssetViewerSettings->OnAssetViewerSettingsChanged().Remove(AssetViewerSettingsChangedHandle);
@@ -73,7 +73,7 @@ TArray<FString> FPreviewProfileController::GetPreviewProfiles(int32& OutCurrentP
 
 bool FPreviewProfileController::SetActiveProfile(const FString& ProfileName)
 {
-	if (PerProjectSettings && AssetViewerSettings)
+	if (IsValid(PerProjectSettings) && IsValid(AssetViewerSettings))
 	{
 		EnsureProfilesStateCoherence();
 
