add option to toggle off the cubegrid's grid plane visualization

diff --git a/Engine/Plugins/Experimental/MeshModelingToolsetExp/Source/MeshModelingToolsExp/Private/CubeGridTool.cpp b/Engine/Plugins/Experimental/MeshModelingToolsetExp/Source/MeshModelingToolsExp/Private/CubeGridTool.cpp
--- a/Engine/Plugins/Experimental/MeshModelingToolsetExp/Source/MeshModelingToolsExp/Private/CubeGridTool.cpp
+++ b/Engine/Plugins/Experimental/MeshModelingToolsetExp/Source/MeshModelingToolsExp/Private/CubeGridTool.cpp
@@ -1157,6 +1157,14 @@ void UCubeGridTool::Setup()
 		[this](bool bOn) {
 			UpdateGizmoVisibility(bOn);
 		});
+	Settings->WatchProperty(Settings->bShowGrid,
+		[this](bool bShow)
+		{
+			if (ULineSetComponent* LineSet = LineSets->FindLineSet(GridLineSetID))
+			{
+				LineSet->SetVisibility(bShow);
+			}
+		});
 	Settings->WatchProperty(Settings->bCrosswiseDiagonal,
 		[this](bool bOn) {
 			InvalidatePreview();
@@ -2258,6 +2266,7 @@ void UCubeGridTool::UpdateGridLineSet()
 	LineSet->Clear();
 	DrawGridSection(*LineSet, *CubeGrid, GridBox,
 		GridLineColor, GridLineThickness, GridLineDepthBias);
+	LineSet->SetVisibility(Settings->bShowGrid);
 }
 
 void UCubeGridTool::UpdateCornerModeLineSet()
diff --git a/Engine/Plugins/Experimental/MeshModelingToolsetExp/Source/MeshModelingToolsExp/Public/CubeGridTool.h b/Engine/Plugins/Experimental/MeshModelingToolsetExp/Source/MeshModelingToolsExp/Public/CubeGridTool.h
--- a/Engine/Plugins/Experimental/MeshModelingToolsetExp/Source/MeshModelingToolsExp/Public/CubeGridTool.h
+++ b/Engine/Plugins/Experimental/MeshModelingToolsetExp/Source/MeshModelingToolsExp/Public/CubeGridTool.h
@@ -88,6 +88,9 @@ public:
 		UIMin = -180, UIMax = 180, ClampMin = -180000, ClampMax = 18000))
 	FRotator GridFrameOrientation = FRotator(0, 0, 0);
 
+	UPROPERTY(EditAnywhere, Category = Options)
+	bool bShowGrid = true;
+
 	UPROPERTY(EditAnywhere, Category = Options, meta = (TransientToolProperty))
 	bool bShowGizmo = false;
 
