make edit pivot tool use the mesh element selection (if available) as an initial pivot point

diff --git a/Engine/Plugins/Experimental/MeshModelingToolsetExp/Source/MeshModelingToolsExp/Private/EditPivotTool.cpp b/Engine/Plugins/Experimental/MeshModelingToolsetExp/Source/MeshModelingToolsExp/Private/EditPivotTool.cpp
--- a/Engine/Plugins/Experimental/MeshModelingToolsetExp/Source/MeshModelingToolsExp/Private/EditPivotTool.cpp
+++ b/Engine/Plugins/Experimental/MeshModelingToolsetExp/Source/MeshModelingToolsExp/Private/EditPivotTool.cpp
@@ -22,6 +22,7 @@
 #include "Components/InstancedStaticMeshComponent.h"
 #include "Engine/World.h"
 
+#include "Selection/StoredMeshSelectionUtil.h"
 #include "TargetInterfaces/MeshDescriptionCommitter.h"
 #include "TargetInterfaces/MeshDescriptionProvider.h"
 #include "TargetInterfaces/PrimitiveComponentBackedTarget.h"
@@ -43,6 +44,33 @@ UMultiSelectionMeshEditingTool* UEditPivotToolBuilder::CreateNewTool(const FTool
 	return NewObject<UEditPivotTool>(SceneState.ToolManager);
 }
 
+void UEditPivotToolBuilder::InitializeNewTool(UMultiSelectionMeshEditingTool* NewTool, const FToolBuilderState& SceneState) const
+{
+	const TArray<TObjectPtr<UToolTarget>> Targets = SceneState.TargetManager->BuildAllSelectedTargetable(SceneState, GetTargetRequirements());
+	NewTool->SetTargets(Targets);
+	NewTool->SetWorld(SceneState.World);
+
+	if (Targets.Num() == 1)
+	{
+		if (UEditPivotTool* NewPivotTool = Cast<UEditPivotTool>(NewTool))
+		{
+			// if there is an element selection, use its frame origin as an initial pivot
+			FFrame3d Frame;
+			FAxisAlignedBox3d Bounds;
+			bool bIsElementSelection;
+			if (GetCurrentSelectionWorldFrameBounds(SceneState, Frame, Bounds, bIsElementSelection))
+			{
+				if (bIsElementSelection)
+				{
+					FTransform LocalToWorld = UE::ToolTarget::GetLocalToWorldTransform(Targets[0]);
+					LocalToWorld.SetTranslation(Frame.Origin);
+					NewPivotTool->SetInitialPivot(LocalToWorld);
+				}
+			}
+		}
+	}
+}
+
 
 
 
@@ -91,6 +119,11 @@ void UEditPivotTool::Setup()
 
 	Precompute();
 
+	if (bHasCustomInitialPivot)
+	{
+		ActiveGizmos[0].TransformGizmo->SetNewGizmoTransform(InitialPivot);
+	}
+
 	FText AllTheWarnings = LOCTEXT("EditPivotWarning", "WARNING: This Tool will Modify the selected StaticMesh Assets! If you do not wish to modify the original Assets, please make copies in the Content Browser first!");
 
 	// detect and warn about any meshes in selection that correspond to same source data
diff --git a/Engine/Plugins/Experimental/MeshModelingToolsetExp/Source/MeshModelingToolsExp/Public/EditPivotTool.h b/Engine/Plugins/Experimental/MeshModelingToolsetExp/Source/MeshModelingToolsExp/Public/EditPivotTool.h
--- a/Engine/Plugins/Experimental/MeshModelingToolsetExp/Source/MeshModelingToolsExp/Public/EditPivotTool.h
+++ b/Engine/Plugins/Experimental/MeshModelingToolsetExp/Source/MeshModelingToolsExp/Public/EditPivotTool.h
@@ -27,6 +27,7 @@ class MESHMODELINGTOOLSEXP_API UEditPivotToolBuilder : public UMultiSelectionMes
 	GENERATED_BODY()
 public:
 	virtual UMultiSelectionMeshEditingTool* CreateNewTool(const FToolBuilderState& SceneState) const override;
+	virtual void InitializeNewTool(UMultiSelectionMeshEditingTool* NewTool, const FToolBuilderState& SceneState) const override;
 };
 
 
@@ -194,9 +195,18 @@ public:
 
 	virtual void RequestAction(EEditPivotToolActions ActionType);
 
+	void SetInitialPivot(FTransform3d InInitialPivot) 
+	{
+		bHasCustomInitialPivot = true;
+		InitialPivot = InInitialPivot;
+	}
+
 protected:
 	TArray<int> MapToFirstOccurrences;
 
+	FTransform3d InitialPivot = FTransform3d::Identity;
+	bool bHasCustomInitialPivot = false;
+
 	FTransform3d Transform;
 	UE::Geometry::FAxisAlignedBox3d ObjectBounds;
 	UE::Geometry::FAxisAlignedBox3d WorldBounds;
