Improvements to skin paint relax mode
- Toggle to relax now uses relax mode settings
- Made relax operation iterative to prevent over-relaxing within a single operation.
- Added skin paint brush mode and color modes to config file

diff --git a/Engine/Plugins/Experimental/MeshModelingToolsetExp/Source/MeshModelingToolsEditorOnlyExp/Private/SkeletalMesh/SkinWeightsPaintTool.cpp b/Engine/Plugins/Experimental/MeshModelingToolsetExp/Source/MeshModelingToolsEditorOnlyExp/Private/SkeletalMesh/SkinWeightsPaintTool.cpp
--- a/Engine/Plugins/Experimental/MeshModelingToolsetExp/Source/MeshModelingToolsEditorOnlyExp/Private/SkeletalMesh/SkinWeightsPaintTool.cpp
+++ b/Engine/Plugins/Experimental/MeshModelingToolsetExp/Source/MeshModelingToolsEditorOnlyExp/Private/SkeletalMesh/SkinWeightsPaintTool.cpp
@@ -929,6 +929,33 @@ FBox USkinWeightsPaintTool::GetWorldSpaceFocusBox()
 	return PreviewMesh->GetActor()->GetComponentsBoundingBox();
 }
 
+void USkinWeightsPaintTool::OnUpdateModifierState(int ModifierID, bool bIsOn)
+{
+	// toggle Relax mode on while shift key is held, then swap back to prior mode on release
+	if (ModifierID == ShiftModifier)
+	{
+		if (bIsOn)
+		{
+			// when shift key is pressed
+			if (!bShiftToggle)
+			{
+				WeightToolProperties->PriorBrushMode = WeightToolProperties->BrushMode;
+				WeightToolProperties->SetBrushMode(EWeightEditOperation::Relax);
+			}
+		}
+		else
+		{
+			// when shift key is released
+			if (bShiftToggle)
+			{
+				WeightToolProperties->SetBrushMode(WeightToolProperties->PriorBrushMode);
+			}
+		}
+	}
+
+	Super::OnUpdateModifierState(ModifierID, bIsOn);
+}
+
 void USkinWeightsPaintTool::OnTick(float DeltaTime)
 {
 	if (bStampPending)
@@ -1011,7 +1038,6 @@ void USkinWeightsPaintTool::OnBeginDrag(const FRay& WorldRay)
 	if (IsInBrushStroke())
 	{
 		bInvertStroke = GetCtrlToggle();
-		bSmoothStroke = GetShiftToggle();
 		BeginChange();
 		StartStamp = UBaseBrushTool::LastBrushStamp;
 		LastStamp = StartStamp;
@@ -1035,7 +1061,6 @@ void USkinWeightsPaintTool::OnEndDrag(const FRay& Ray)
 	UDynamicMeshBrushTool::OnEndDrag(Ray);
 
 	bInvertStroke = false;
-	bSmoothStroke = false;
 	bStampPending = false;
 
 	if (ActiveChange)
@@ -1216,8 +1241,8 @@ FVector4f USkinWeightsPaintTool::GetColorOfVertex(VertexIndex InVertexIndex, Bon
 					continue;
 				}
 				
-				const float Value = InCurrentBoneIndex == BoneWeight.BoneIndex ? 1.0f: 0.25f;
-				constexpr float Saturation = 1.f;
+				const float Value = InCurrentBoneIndex == BoneWeight.BoneIndex ? 1.0f: 0.6f;
+				constexpr float Saturation = 0.75f;
 				const FLinearColor BoneColor = SkeletalDebugRendering::GetSemiRandomColorForBone(BoneWeight.BoneIndex, Value, Saturation);
 				Color = FLinearColor::LerpUsingHSV(Color, BoneColor, BoneWeight.Weight);
 			}
@@ -1287,7 +1312,7 @@ void USkinWeightsPaintTool::ApplyStamp(const FBrushStampData& Stamp)
 	{
 		TRACE_CPUPROFILER_EVENT_SCOPE(SkinTool::EditWeightOfVerticesInStamp);
 		// generate a weight edit from this stamp (includes modifications caused by normalization)
-		if (bSmoothStroke || WeightToolProperties->BrushMode == EWeightEditOperation::Relax)
+		if (WeightToolProperties->BrushMode == EWeightEditOperation::Relax)
 		{
 			// use mesh topology to iteratively smooth weights across neighboring vertices
 			const float UseStrength = CalculateBrushStrengthToUse(EWeightEditOperation::Relax);
@@ -1446,24 +1471,30 @@ void USkinWeightsPaintTool::RelaxWeightOnVertices(
 		return;
 	}
 
-	for (int32 Index = 0; Index < VerticesToEdit.Num(); ++Index)
+	constexpr int32 NumRelaxIterations = 3;
+	constexpr float PercentPerIter = 1.0f / static_cast<float>(NumRelaxIterations);
+	for (int32 Iteration=0; Iteration<NumRelaxIterations; ++Iteration)
 	{
-		const int32 VertexID = VerticesToEdit[Index];
-		const float UseFalloff = VertexFalloffs.IsValidIndex(Index) ? VertexFalloffs[Index] * UseStrength : UseStrength;
+		for (int32 VertexIndex = 0; VertexIndex < VerticesToEdit.Num(); ++VertexIndex)
+		{
+			const int32 VertexID = VerticesToEdit[VertexIndex];
+			float UseFalloff = VertexFalloffs.IsValidIndex(VertexIndex) ? VertexFalloffs[VertexIndex] * UseStrength : UseStrength;
+			UseFalloff *= PercentPerIter;
 
-		TMap<int32, float> FinalWeights;
-		const bool bSmoothSuccess = SmoothWeightsOp->SmoothWeightsAtVertex(VertexID, UseFalloff, FinalWeights);
+			TMap<int32, float> FinalWeights;
+			const bool bSmoothSuccess = SmoothWeightsOp->SmoothWeightsAtVertex(VertexID, UseFalloff, FinalWeights);
 
-		if (ensure(bSmoothSuccess))
-		{
-			// apply weight edits
-			for (const TTuple<BoneIndex, float>& FinalWeight : FinalWeights)
+			if (ensure(bSmoothSuccess))
 			{
-				// record an edit for this vertex, for this bone
-				const int32 BoneIndex = FinalWeight.Key;
-				const float NewWeight = FinalWeight.Value;
-				const float OldWeight = Weights.GetWeightOfBoneOnVertex(BoneIndex, VertexID, Weights.PreChangeWeights);
-				InOutWeightEdits.MergeSingleEdit(BoneIndex, VertexID, OldWeight, NewWeight);
+				// apply weight edits
+				for (const TTuple<BoneIndex, float>& FinalWeight : FinalWeights)
+				{
+					// record an edit for this vertex, for this bone
+					const int32 BoneIndex = FinalWeight.Key;
+					const float NewWeight = FinalWeight.Value;
+					const float OldWeight = Weights.GetWeightOfBoneOnVertex(BoneIndex, VertexID, Weights.PreChangeWeights);
+					InOutWeightEdits.MergeSingleEdit(BoneIndex, VertexID, OldWeight, NewWeight);
+				}
 			}
 		}
 	}
diff --git a/Engine/Plugins/Experimental/MeshModelingToolsetExp/Source/MeshModelingToolsEditorOnlyExp/Public/SkeletalMesh/SkinWeightsPaintTool.h b/Engine/Plugins/Experimental/MeshModelingToolsetExp/Source/MeshModelingToolsEditorOnlyExp/Public/SkeletalMesh/SkinWeightsPaintTool.h
--- a/Engine/Plugins/Experimental/MeshModelingToolsetExp/Source/MeshModelingToolsEditorOnlyExp/Public/SkeletalMesh/SkinWeightsPaintTool.h
+++ b/Engine/Plugins/Experimental/MeshModelingToolsetExp/Source/MeshModelingToolsEditorOnlyExp/Public/SkeletalMesh/SkinWeightsPaintTool.h
@@ -280,21 +280,22 @@ class MESHMODELINGTOOLSEDITORONLYEXP_API USkinWeightsPaintToolProperties : publi
 public:
 
 	// brush vs selection modes
-	UPROPERTY()
+	UPROPERTY(Config)
 	EWeightEditMode EditingMode;
 
 	// custom brush modes and falloff types
-	UPROPERTY()
+	UPROPERTY(Config)
 	EWeightEditOperation BrushMode;
+	EWeightEditOperation PriorBrushMode; // when toggling with modifier key
 
 	// weight color properties
-	UPROPERTY(EditAnywhere, Category = WeightColors)
+	UPROPERTY(EditAnywhere, Config, Category = WeightColors)
 	EWeightColorMode ColorMode;
-	UPROPERTY(EditAnywhere, Category = WeightColors)
+	UPROPERTY(EditAnywhere, Config, Category = WeightColors)
 	TArray<FLinearColor> ColorRamp;
-	UPROPERTY(EditAnywhere, Category = WeightColors)
+	UPROPERTY(EditAnywhere, Config, Category = WeightColors)
 	FLinearColor MinColor;
-	UPROPERTY(EditAnywhere, Category = WeightColors)
+	UPROPERTY(EditAnywhere, Config, Category = WeightColors)
 	FLinearColor MaxColor;
 	bool bColorModeChanged = false;
 
@@ -322,6 +323,17 @@ public:
 
 	// pointer back to paint tool
 	TObjectPtr<USkinWeightsPaintTool> WeightTool;
+
+	void SetBrushMode(EWeightEditOperation InBrushMode)
+	{
+		BrushMode = InBrushMode;
+
+		// sync base tool settings with the mode specific saved values
+		// these are the source of truth for the base class viewport rendering of brush
+		BrushRadius = GetBrushConfig().Radius;
+		BrushStrength = GetBrushConfig().Strength;
+		BrushFalloffAmount = GetBrushConfig().Falloff;
+	}
 };
 
 // An interactive tool for painting and editing skin weights.
@@ -353,6 +365,9 @@ public:
 	virtual bool SupportsWorldSpaceFocusBox() override { return true; }
 	virtual FBox GetWorldSpaceFocusBox() override;
 
+	// IClickDragBehaviorTarget implementation
+	virtual void OnUpdateModifierState(int ModifierID, bool bIsOn) override;
+
 	// using when ToolChange is applied via Undo/Redo
 	void ExternalUpdateWeights(const int32 BoneIndex, const TMap<int32, float>& IndexValues);
 
@@ -429,7 +444,6 @@ protected:
 		TArray<float>& VertexFalloffs);
 	float CalculateBrushStrengthToUse(EWeightEditOperation EditMode) const;
 	bool bInvertStroke = false;
-	bool bSmoothStroke = false;
 	FBrushStampData StartStamp;
 	FBrushStampData LastStamp;
 	bool bStampPending;
diff --git a/Engine/Plugins/Experimental/MeshModelingToolsetExp/Source/ModelingEditorUI/Private/SkinWeightDetailCustomization.cpp b/Engine/Plugins/Experimental/MeshModelingToolsetExp/Source/ModelingEditorUI/Private/SkinWeightDetailCustomization.cpp
--- a/Engine/Plugins/Experimental/MeshModelingToolsetExp/Source/ModelingEditorUI/Private/SkinWeightDetailCustomization.cpp
+++ b/Engine/Plugins/Experimental/MeshModelingToolsetExp/Source/ModelingEditorUI/Private/SkinWeightDetailCustomization.cpp
@@ -169,13 +169,7 @@ void FSkinWeightDetailCustomization::AddBrushUI(IDetailLayoutBuilder& DetailBuil
 			})
 			.OnValueChanged_Lambda([this](EWeightEditOperation Mode)
 			{
-				SkinToolSettings->BrushMode = Mode;
-
-				// sync base tool settings with the mode specific saved values
-				// these are the source of truth for the base class viewport rendering of brush
-				SkinToolSettings->BrushRadius = SkinToolSettings->GetBrushConfig().Radius;
-				SkinToolSettings->BrushStrength = SkinToolSettings->GetBrushConfig().Strength;
-				SkinToolSettings->BrushFalloffAmount = SkinToolSettings->GetBrushConfig().Falloff;
+				SkinToolSettings->SetBrushMode(Mode);
 			})
 			+SSegmentedControl<EWeightEditOperation>::Slot(EWeightEditOperation::Add)
 			.Text(LOCTEXT("BrushAddMode", "Add"))
diff --git a/Engine/Source/Runtime/InteractiveToolsFramework/Private/BaseTools/MeshSurfacePointTool.cpp b/Engine/Source/Runtime/InteractiveToolsFramework/Private/BaseTools/MeshSurfacePointTool.cpp
--- a/Engine/Source/Runtime/InteractiveToolsFramework/Private/BaseTools/MeshSurfacePointTool.cpp
+++ b/Engine/Source/Runtime/InteractiveToolsFramework/Private/BaseTools/MeshSurfacePointTool.cpp
@@ -53,11 +53,6 @@ void UMeshSurfacePointToolBuilder::InitializeNewTool(UMeshSurfacePointTool* NewT
 	NewTool->SetWorld(SceneState.World);
 }
 
-
-static const int UMeshSurfacePointTool_ShiftModifier = 1;
-static const int UMeshSurfacePointTool_CtrlModifier = 2;
-
-
 /*
  * Tool
  */
@@ -70,14 +65,14 @@ void UMeshSurfacePointTool::Setup()
 
 	// add input behaviors
 	UClickDragInputBehavior* DragBehavior = NewObject<UClickDragInputBehavior>();
-	DragBehavior->Modifiers.RegisterModifier(UMeshSurfacePointTool_ShiftModifier, FInputDeviceState::IsShiftKeyDown);
-	DragBehavior->Modifiers.RegisterModifier(UMeshSurfacePointTool_CtrlModifier, FInputDeviceState::IsCtrlKeyDown);
+	DragBehavior->Modifiers.RegisterModifier(ShiftModifier, FInputDeviceState::IsShiftKeyDown);
+	DragBehavior->Modifiers.RegisterModifier(CtrlModifier, FInputDeviceState::IsCtrlKeyDown);
 	DragBehavior->Initialize(this);
 	AddInputBehavior(DragBehavior);
 
 	UMouseHoverBehavior* HoverBehavior = NewObject<UMouseHoverBehavior>();
-	HoverBehavior->Modifiers.RegisterModifier(UMeshSurfacePointTool_ShiftModifier, FInputDeviceState::IsShiftKeyDown);
-	HoverBehavior->Modifiers.RegisterModifier(UMeshSurfacePointTool_CtrlModifier, FInputDeviceState::IsCtrlKeyDown);
+	HoverBehavior->Modifiers.RegisterModifier(ShiftModifier, FInputDeviceState::IsShiftKeyDown);
+	HoverBehavior->Modifiers.RegisterModifier(CtrlModifier, FInputDeviceState::IsCtrlKeyDown);
 	HoverBehavior->Initialize(this);
 	AddInputBehavior(HoverBehavior);
 }
@@ -144,11 +139,11 @@ void UMeshSurfacePointTool::SetCtrlToggle(bool bCtrlDown)
 
 void UMeshSurfacePointTool::OnUpdateModifierState(int ModifierID, bool bIsOn)
 {
-	if (ModifierID == UMeshSurfacePointTool_ShiftModifier)
+	if (ModifierID == ShiftModifier)
 	{
 		bShiftToggle = bIsOn;
 	}
-	else if (ModifierID == UMeshSurfacePointTool_CtrlModifier)
+	else if (ModifierID == CtrlModifier)
 	{
 		bCtrlToggle = bIsOn;
 	}
diff --git a/Engine/Source/Runtime/InteractiveToolsFramework/Public/BaseTools/MeshSurfacePointTool.h b/Engine/Source/Runtime/InteractiveToolsFramework/Public/BaseTools/MeshSurfacePointTool.h
--- a/Engine/Source/Runtime/InteractiveToolsFramework/Public/BaseTools/MeshSurfacePointTool.h
+++ b/Engine/Source/Runtime/InteractiveToolsFramework/Public/BaseTools/MeshSurfacePointTool.h
@@ -153,5 +153,8 @@ protected:
 
 	UPROPERTY()
 	TWeakObjectPtr<UWorld> TargetWorld = nullptr;
+
+	static constexpr int ShiftModifier = 1;
+	static constexpr int CtrlModifier = 2;
 };
 
