Skeleton Editing Tool: handle bone size & colors

#rnx

diff --git a/Engine/Plugins/Experimental/MeshModelingToolsetExp/Source/MeshModelingToolsEditorOnlyExp/Private/SkeletalMesh/SkeletonEditingTool.cpp b/Engine/Plugins/Experimental/MeshModelingToolsetExp/Source/MeshModelingToolsEditorOnlyExp/Private/SkeletalMesh/SkeletonEditingTool.cpp
--- a/Engine/Plugins/Experimental/MeshModelingToolsetExp/Source/MeshModelingToolsEditorOnlyExp/Private/SkeletalMesh/SkeletonEditingTool.cpp
+++ b/Engine/Plugins/Experimental/MeshModelingToolsetExp/Source/MeshModelingToolsEditorOnlyExp/Private/SkeletalMesh/SkeletonEditingTool.cpp
@@ -112,6 +112,8 @@ void USkeletonEditingTool::Setup()
 		return;
 	}
 
+	WeakMesh = SkeletalMesh;
+
 	SetupModifier(SkeletalMesh);
 	SetupPreviewMesh();
 	SetupProperties();
@@ -1728,7 +1730,7 @@ void USkeletonEditingTool::DrawHUD(FCanvas* Canvas, IToolsContextRenderAPI* Rend
 
 void USkeletonEditingTool::Render(IToolsContextRenderAPI* RenderAPI)
 {
-	// FIXME many things could be caches here and updated lazilly
+	// NOTE many things could be cached here and updated lazily
 	if (!Target)
 	{
 		return;
@@ -1737,7 +1739,11 @@ void USkeletonEditingTool::Render(IToolsContextRenderAPI* RenderAPI)
 	const UPersonaOptions* PersonaOptions = GetDefault<UPersonaOptions>();
 	static FSkelDebugDrawConfig DrawConfig;
 		DrawConfig.BoneDrawMode = EBoneDrawMode::Type::All;
+#if WITH_EDITORONLY_DATA
+		DrawConfig.BoneDrawSize = WeakMesh.IsValid() ? WeakMesh->BoneDrawSize : 1.f;
+#else
 		DrawConfig.BoneDrawSize = 1.f;
+#endif
 		DrawConfig.bAddHitProxy = true;
 		DrawConfig.bForceDraw = false;
 		DrawConfig.DefaultBoneColor = PersonaOptions->DefaultBoneColor;
@@ -1759,13 +1765,15 @@ void USkeletonEditingTool::Render(IToolsContextRenderAPI* RenderAPI)
 	TArray<FBoneIndexType> RequiredBones; RequiredBones.AddUninitialized(NumBones);
 	TArray<FTransform> WorldTransforms; WorldTransforms.AddUninitialized(NumBones);
 	TArray<FLinearColor> BoneColors; BoneColors.AddUninitialized(NumBones);
+
+	const bool bUseBoneColors = GetDefault<UPersonaOptions>()->bShowBoneColors;
 	
 	for (int32 Index = 0; Index < NumBones; ++Index)
 	{
 		const FTransform& BoneTransform = Modifier->GetTransform(Index, true);
 		WorldTransforms[Index] = BoneTransform;
 		RequiredBones[Index] = Index;
-		BoneColors[Index] = SkeletalDebugRendering::GetSemiRandomColorForBone(Index);
+		BoneColors[Index] = bUseBoneColors ? SkeletalDebugRendering::GetSemiRandomColorForBone(Index) : DrawConfig.DefaultBoneColor;
 		HitProxies.Add(new HBoneHitProxy(Index, RefSkeleton.GetBoneName(Index)));
 	}
 
diff --git a/Engine/Plugins/Experimental/MeshModelingToolsetExp/Source/MeshModelingToolsEditorOnlyExp/Public/SkeletalMesh/SkeletonEditingTool.h b/Engine/Plugins/Experimental/MeshModelingToolsetExp/Source/MeshModelingToolsEditorOnlyExp/Public/SkeletalMesh/SkeletonEditingTool.h
--- a/Engine/Plugins/Experimental/MeshModelingToolsetExp/Source/MeshModelingToolsEditorOnlyExp/Public/SkeletalMesh/SkeletonEditingTool.h
+++ b/Engine/Plugins/Experimental/MeshModelingToolsetExp/Source/MeshModelingToolsEditorOnlyExp/Public/SkeletalMesh/SkeletonEditingTool.h
@@ -292,6 +292,8 @@ private:
 
 	// Defer pending function on tick to allow other external notification to be handled. 
 	bool bDeferUntilFocused = false;
+
+	TWeakObjectPtr<USkeletalMesh> WeakMesh = nullptr;
 };
 
 ENUM_CLASS_FLAGS(USkeletonEditingTool::EBoneSelectionMode);
