Added ability to toggle weight display off in skin paint tool
Now auto-switches to use colored bones when entering tool

diff --git a/Engine/Plugins/Experimental/MeshModelingToolsetExp/Source/MeshModelingToolsEditorOnlyExp/Private/SkeletalMesh/SkinWeightsPaintTool.cpp b/Engine/Plugins/Experimental/MeshModelingToolsetExp/Source/MeshModelingToolsEditorOnlyExp/Private/SkeletalMesh/SkinWeightsPaintTool.cpp
--- a/Engine/Plugins/Experimental/MeshModelingToolsetExp/Source/MeshModelingToolsEditorOnlyExp/Private/SkeletalMesh/SkinWeightsPaintTool.cpp
+++ b/Engine/Plugins/Experimental/MeshModelingToolsetExp/Source/MeshModelingToolsEditorOnlyExp/Private/SkeletalMesh/SkinWeightsPaintTool.cpp
@@ -29,6 +29,7 @@
 #include "SkeletalDebugRendering.h"
 #include "Editor/Persona/Public/IPersonaEditorModeManager.h"
 #include "Editor/Persona/Public/PersonaModule.h"
+#include "Preferences/PersonaOptions.h"
 
 #include UE_INLINE_GENERATED_CPP_BY_NAME(SkinWeightsPaintTool)
 
@@ -759,16 +760,17 @@ void USkinWeightsPaintTool::Setup()
 		// Create an overlay that has no split elements, init with zero value.
 		Mesh.Attributes()->PrimaryColors()->CreateFromPredicate([](int ParentVID, int TriIDA, int TriIDB){return true;}, 0.f);
 	});
-	UMaterialInterface* VtxColorMaterial = GetToolManager()->GetContextQueriesAPI()->GetStandardMaterial(EStandardToolContextMaterials::VertexColorMaterial);
-	if (VtxColorMaterial != nullptr)
-	{
-		PreviewMesh->SetOverrideRenderMaterial(VtxColorMaterial);
-	}
+
+	// start by displaying vertex color material
+	SetDisplayVertexColors(true);
 
 	// modify viewport render settings to optimize for painting weights
 	FPreviewProfileController PreviewProfileController;
 	PreviewProfileToRestore = PreviewProfileController.GetActiveProfile();
 	PreviewProfileController.SetActiveProfile(UDefaultEditorProfiles::EditingProfileName.ToString());
+	// turn on bone colors
+	bBoneColorsToRestore = GetDefault<UPersonaOptions>()->bShowBoneColors;
+	GetMutableDefault<UPersonaOptions>()->bShowBoneColors = true;
 
 	// build octree for vertices
 	VerticesOctree.Initialize(PreviewMesh->GetMesh(), true);
@@ -981,6 +983,32 @@ void USkinWeightsPaintTool::OnTick(float DeltaTime)
 	Weights.Deformer.UpdateVertexDeformation(this);
 }
 
+void USkinWeightsPaintToolProperties::SetFalloffMode(EWeightBrushFalloffMode InFalloffMode)
+{
+	bColorModeChanged = true;
+	GetBrushConfig().FalloffMode = InFalloffMode;
+	SaveConfig();
+}
+
+void USkinWeightsPaintToolProperties::SetColorMode(EWeightColorMode InColorMode)
+{
+	ColorMode = InColorMode;
+	bColorModeChanged = true;
+
+	WeightTool->SetDisplayVertexColors(ColorMode != EWeightColorMode::FullMaterial);
+}
+
+void USkinWeightsPaintToolProperties::SetBrushMode(EWeightEditOperation InBrushMode)
+{
+	BrushMode = InBrushMode;
+
+	// sync base tool settings with the mode specific saved values
+	// these are the source of truth for the base class viewport rendering of brush
+	BrushRadius = GetBrushConfig().Radius;
+	BrushStrength = GetBrushConfig().Strength;
+	BrushFalloffAmount = GetBrushConfig().Falloff;
+}
+
 bool USkinWeightsPaintTool::HitTest(const FRay& Ray, FHitResult& OutHit)
 {
 	// do not query the triangle octree until all async ops are finished
@@ -1180,14 +1208,14 @@ FVector4f USkinWeightsPaintTool::GetColorOfVertex(VertexIndex InVertexIndex, Bon
 {
 	switch (WeightToolProperties->ColorMode)
 	{
-	case EWeightColorMode::MinMax:
+	case EWeightColorMode::Greyscale:
 		{
 			if (InCurrentBoneIndex == INDEX_NONE)
 			{
 				return WeightToolProperties->MinColor; // with no bone selected, all vertices are drawn black
 			}
 			const float Value = Weights.GetWeightOfBoneOnVertex(InCurrentBoneIndex, InVertexIndex, Weights.CurrentWeights);
-			return FMath::Lerp(WeightToolProperties->MinColor, WeightToolProperties->MaxColor, Value);
+			return FMath::Lerp(FLinearColor::Black, FLinearColor::White, Value);
 		}
 	case EWeightColorMode::Ramp:
 		{
@@ -1230,7 +1258,7 @@ FVector4f USkinWeightsPaintTool::GetColorOfVertex(VertexIndex InVertexIndex, Bon
 			const FLinearColor& EndColor = Colors[ColorIndex+1];
 			return UE::Geometry::ToVector4<float>(FMath::Lerp(StartColor, EndColor, Param));
 		}
-	case EWeightColorMode::MultiColor:
+	case EWeightColorMode::BoneColors:
 		{
 			FVector4f Color = FVector4f::Zero();
 			const VertexWeights& VertexWeights = Weights.CurrentWeights[InVertexIndex];
@@ -1248,6 +1276,8 @@ FVector4f USkinWeightsPaintTool::GetColorOfVertex(VertexIndex InVertexIndex, Bon
 			}
 			return Color;
 		}
+	case EWeightColorMode::FullMaterial:
+		return FLinearColor::White;
 	default:
 		checkNoEntry();
 		return FLinearColor::Black;
@@ -1613,6 +1643,7 @@ void USkinWeightsPaintTool::OnShutdown(EToolShutdownType ShutdownType)
 	// restore viewport show flags and preview settings
 	FPreviewProfileController PreviewProfileController;
 	PreviewProfileController.SetActiveProfile(PreviewProfileToRestore);
+	GetMutableDefault<UPersonaOptions>()->bShowBoneColors = bBoneColorsToRestore;
 
 	if (EditorContext.IsValid())
 	{
@@ -2187,6 +2218,19 @@ FName USkinWeightsPaintTool::GetBoneNameFromIndex(BoneIndex InIndex) const
 	return NAME_None;
 }
 
+void USkinWeightsPaintTool::SetDisplayVertexColors(bool bShowVertexColors) const
+{
+	if (bShowVertexColors)
+	{
+		UMaterialInterface* VtxColorMaterial = GetToolManager()->GetContextQueriesAPI()->GetStandardMaterial(EStandardToolContextMaterials::VertexColorMaterial);
+		PreviewMesh->SetOverrideRenderMaterial(VtxColorMaterial);
+	}
+	else
+	{
+		PreviewMesh->ClearOverrideRenderMaterial();
+	}
+}
+
 void USkinWeightsPaintTool::OnPropertyModified(UObject* ModifiedObject, FProperty* ModifiedProperty)
 {
 	Super::OnPropertyModified(ModifiedObject, ModifiedProperty);
diff --git a/Engine/Plugins/Experimental/MeshModelingToolsetExp/Source/MeshModelingToolsEditorOnlyExp/Public/SkeletalMesh/SkinWeightsPaintTool.h b/Engine/Plugins/Experimental/MeshModelingToolsetExp/Source/MeshModelingToolsEditorOnlyExp/Public/SkeletalMesh/SkinWeightsPaintTool.h
--- a/Engine/Plugins/Experimental/MeshModelingToolsetExp/Source/MeshModelingToolsEditorOnlyExp/Public/SkeletalMesh/SkinWeightsPaintTool.h
+++ b/Engine/Plugins/Experimental/MeshModelingToolsetExp/Source/MeshModelingToolsEditorOnlyExp/Public/SkeletalMesh/SkinWeightsPaintTool.h
@@ -46,9 +46,10 @@ enum class EWeightEditMode : uint8
 UENUM()
 enum class EWeightColorMode : uint8
 {
-	MinMax,
+	Greyscale,
 	Ramp,
-	MultiColor,
+	BoneColors,
+	FullMaterial
 };
 
 // brush falloff mode
@@ -289,13 +290,13 @@ public:
 	EWeightEditOperation PriorBrushMode; // when toggling with modifier key
 
 	// weight color properties
-	UPROPERTY(EditAnywhere, Config, Category = WeightColors)
+	UPROPERTY(EditAnywhere, Config, Category = MeshDisplay)
 	EWeightColorMode ColorMode;
-	UPROPERTY(EditAnywhere, Config, Category = WeightColors)
+	UPROPERTY(EditAnywhere, Config, Category = MeshDisplay)
 	TArray<FLinearColor> ColorRamp;
-	UPROPERTY(EditAnywhere, Config, Category = WeightColors)
+	UPROPERTY(EditAnywhere, Config, Category = MeshDisplay)
 	FLinearColor MinColor;
-	UPROPERTY(EditAnywhere, Config, Category = WeightColors)
+	UPROPERTY(EditAnywhere, Config, Category = MeshDisplay)
 	FLinearColor MaxColor;
 	bool bColorModeChanged = false;
 
@@ -324,16 +325,9 @@ public:
 	// pointer back to paint tool
 	TObjectPtr<USkinWeightsPaintTool> WeightTool;
 
-	void SetBrushMode(EWeightEditOperation InBrushMode)
-	{
-		BrushMode = InBrushMode;
-
-		// sync base tool settings with the mode specific saved values
-		// these are the source of truth for the base class viewport rendering of brush
-		BrushRadius = GetBrushConfig().Radius;
-		BrushStrength = GetBrushConfig().Strength;
-		BrushFalloffAmount = GetBrushConfig().Falloff;
-	}
+	void SetFalloffMode(EWeightBrushFalloffMode InFalloffMode);
+	void SetColorMode(EWeightColorMode InColorMode);
+	void SetBrushMode(EWeightEditOperation InBrushMode);
 };
 
 // An interactive tool for painting and editing skin weights.
@@ -403,6 +397,9 @@ public:
 	// convert an index to a name
 	FName GetBoneNameFromIndex(BoneIndex InIndex) const;
 
+	// toggle the display of weights on the preview mesh (if false, uses the normal skeletal mesh material)
+	void SetDisplayVertexColors(bool bShowVertexColors=true) const;
+
 	// HOW TO EDIT WEIGHTS WITH UNDO/REDO:
 	//
 	// Live Edits:
@@ -523,9 +520,10 @@ protected:
 	TWeakObjectPtr<USkeletalMeshEditorContextObjectBase> EditorContext = nullptr;
 	UPROPERTY()
 	TWeakObjectPtr<UPersonaEditorModeManagerContext> PersonaModeManagerContext = nullptr;
-
+	
 	// editor state to restore when exiting the paint tool
 	FString PreviewProfileToRestore;
+	bool bBoneColorsToRestore;
 
 	friend SkinPaintTool::FSkinToolDeformer;
 };
diff --git a/Engine/Plugins/Experimental/MeshModelingToolsetExp/Source/ModelingEditorUI/Private/SkinWeightDetailCustomization.cpp b/Engine/Plugins/Experimental/MeshModelingToolsetExp/Source/ModelingEditorUI/Private/SkinWeightDetailCustomization.cpp
--- a/Engine/Plugins/Experimental/MeshModelingToolsetExp/Source/ModelingEditorUI/Private/SkinWeightDetailCustomization.cpp
+++ b/Engine/Plugins/Experimental/MeshModelingToolsetExp/Source/ModelingEditorUI/Private/SkinWeightDetailCustomization.cpp
@@ -84,42 +84,38 @@ void FSkinWeightDetailCustomization::CustomizeDetails(IDetailLayoutBuilder& Deta
 	}
 	
 	// COLOR MODE category
-	IDetailCategoryBuilder& WeightColorsCategory = DetailBuilder.EditCategory("WeightColors", FText::GetEmpty(), ECategoryPriority::Important);
-	WeightColorsCategory.InitiallyCollapsed(true);
-	WeightColorsCategory.AddCustomRow(LOCTEXT("ColorModeCategory", "Color Mode"), false)
-	.NameContent()
-	[
-		SNew(STextBlock)
-		.Text(LOCTEXT("ColorModeLabel", "Color Mode"))
-		.Font(FAppStyle::GetFontStyle(TEXT("PropertyWindow.NormalFont")))
-		.ToolTipText(LOCTEXT("ColorModeTooltip", "Determines how the weight colors are displayed."))
-	]
-	.ValueContent()
+	IDetailCategoryBuilder& MeshDisplayCategory = DetailBuilder.EditCategory("MeshDisplay", FText::GetEmpty(), ECategoryPriority::Important);
+	MeshDisplayCategory.InitiallyCollapsed(false);
+	MeshDisplayCategory.AddCustomRow(LOCTEXT("ColorModeCategory", "Color Mode"), false)
+	.WholeRowContent()
 	[
-		SNew(SHorizontalBox)
-		+ SHorizontalBox::Slot()
+		SNew(SBox)
+		.Padding(2.0f)
+		.HAlign(HAlign_Center)
 		[
-			SNew(SBox)
-			.Padding(2.0f)
-			.HAlign(HAlign_Center)
-			[
-				SNew(SSegmentedControl<EWeightColorMode>)
-				.Value_Lambda([this]()
-				{
-					return SkinToolSettings->ColorMode;
-				})
-				.OnValueChanged_Lambda([this](EWeightColorMode Mode)
-				{
-					SkinToolSettings->ColorMode = Mode;
-					SkinToolSettings->bColorModeChanged = true;
-				})
-				+ SSegmentedControl<EWeightColorMode>::Slot(EWeightColorMode::MinMax)
-				.Text(LOCTEXT("MinMaxMode", "Min / Max"))
-				+SSegmentedControl<EWeightColorMode>::Slot(EWeightColorMode::Ramp)
-				.Text(LOCTEXT("RampMode", "Ramp"))
-				+SSegmentedControl<EWeightColorMode>::Slot(EWeightColorMode::MultiColor)
-				.Text(LOCTEXT("MultiColorMode", "Multi Color"))
-			]
+			SNew(SSegmentedControl<EWeightColorMode>)
+			.ToolTipText(LOCTEXT("WeightColorTooltip",
+					"Adjust the weight display in the viewport.\n\n"
+					"Greyscale: Displays weights on the current bone by blending from black (0) to white (1).\n"
+					"Ramp: Displays weights on the current bone. Weights at 0 and 1 use the min and max colors. Weights inbetween 0 and 1 use the ramp colors.\n"
+					"Multi Color: Displays weights on ALL bones using the color of the bones.\n"
+					"Full Material: Shows the normal fully-lit mesh material.\n"))
+			.Value_Lambda([this]()
+			{
+				return SkinToolSettings->ColorMode;
+			})
+			.OnValueChanged_Lambda([this](EWeightColorMode Mode)
+			{
+				SkinToolSettings->SetColorMode(Mode);
+			})
+			+ SSegmentedControl<EWeightColorMode>::Slot(EWeightColorMode::Greyscale)
+			.Text(LOCTEXT("GreyscaleMode", "Greyscale"))
+			+SSegmentedControl<EWeightColorMode>::Slot(EWeightColorMode::Ramp)
+			.Text(LOCTEXT("RampMode", "Ramp"))
+			+SSegmentedControl<EWeightColorMode>::Slot(EWeightColorMode::BoneColors)
+			.Text(LOCTEXT("BoneColorsMode", "Bone Colors"))
+			+SSegmentedControl<EWeightColorMode>::Slot(EWeightColorMode::FullMaterial)
+			.Text(LOCTEXT("MaterialMode", "Full Material"))
 		]
 	];
 
@@ -199,9 +195,7 @@ void FSkinWeightDetailCustomization::AddBrushUI(IDetailLayoutBuilder& DetailBuil
 			})
 			.OnValueChanged_Lambda([this](EWeightBrushFalloffMode Mode)
 			{
-				SkinToolSettings->bColorModeChanged = true;
-				SkinToolSettings->GetBrushConfig().FalloffMode = Mode;
-				SkinToolSettings->SaveConfig();
+				SkinToolSettings->SetFalloffMode(Mode);
 			})
 			+SSegmentedControl<EWeightBrushFalloffMode>::Slot(EWeightBrushFalloffMode::Surface)
 			.Text(LOCTEXT("SurfaceMode", "Surface"))
