Adding ability to set viewport focus from within tools in the skeletal mesh editor context.

diff --git a/Engine/Plugins/Experimental/MeshModelingToolsetExp/Source/MeshModelingToolsEditorOnlyExp/MeshModelingToolsEditorOnlyExp.Build.cs b/Engine/Plugins/Experimental/MeshModelingToolsetExp/Source/MeshModelingToolsEditorOnlyExp/MeshModelingToolsEditorOnlyExp.Build.cs
--- a/Engine/Plugins/Experimental/MeshModelingToolsetExp/Source/MeshModelingToolsEditorOnlyExp/MeshModelingToolsEditorOnlyExp.Build.cs
+++ b/Engine/Plugins/Experimental/MeshModelingToolsetExp/Source/MeshModelingToolsEditorOnlyExp/MeshModelingToolsEditorOnlyExp.Build.cs
@@ -56,7 +56,8 @@ public class MeshModelingToolsEditorOnlyExp : ModuleRules
 				"InputCore",
 				"UnrealEd",
 				"Persona",
-				"AdvancedPreviewScene"
+				"AdvancedPreviewScene",
+				"Slate"
 				// ... add private dependencies that you statically link with here ...	
 			}
 		);
diff --git a/Engine/Plugins/Experimental/MeshModelingToolsetExp/Source/MeshModelingToolsEditorOnlyExp/Private/SkeletalMesh/SkinWeightsPaintTool.cpp b/Engine/Plugins/Experimental/MeshModelingToolsetExp/Source/MeshModelingToolsEditorOnlyExp/Private/SkeletalMesh/SkinWeightsPaintTool.cpp
--- a/Engine/Plugins/Experimental/MeshModelingToolsetExp/Source/MeshModelingToolsEditorOnlyExp/Private/SkeletalMesh/SkinWeightsPaintTool.cpp
+++ b/Engine/Plugins/Experimental/MeshModelingToolsetExp/Source/MeshModelingToolsEditorOnlyExp/Private/SkeletalMesh/SkinWeightsPaintTool.cpp
@@ -760,9 +760,8 @@ void USkinWeightsPaintTool::Setup()
 		// Create an overlay that has no split elements, init with zero value.
 		Mesh.Attributes()->PrimaryColors()->CreateFromPredicate([](int ParentVID, int TriIDA, int TriIDB){return true;}, 0.f);
 	});
-
-	// start by displaying vertex color material
-	SetDisplayVertexColors(true);
+	// optionally display vertex color material
+	SetDisplayVertexColors(WeightToolProperties->ColorMode != EWeightColorMode::FullMaterial);
 
 	// modify viewport render settings to optimize for painting weights
 	FPreviewProfileController PreviewProfileController;
@@ -836,6 +835,8 @@ void USkinWeightsPaintTool::Setup()
 	// trigger last used mode
 	ToggleEditingMode();
 	
+	SetFocusInViewport();
+	
 	// inform user of tool keys
 	// TODO talk with UX team about viewport overlay to show hotkeys
 	GetToolManager()->DisplayMessage(
@@ -988,6 +989,8 @@ void USkinWeightsPaintToolProperties::SetFalloffMode(EWeightBrushFalloffMode InF
 	bColorModeChanged = true;
 	GetBrushConfig().FalloffMode = InFalloffMode;
 	SaveConfig();
+
+	WeightTool->SetFocusInViewport();
 }
 
 void USkinWeightsPaintToolProperties::SetColorMode(EWeightColorMode InColorMode)
@@ -995,7 +998,8 @@ void USkinWeightsPaintToolProperties::SetColorMode(EWeightColorMode InColorMode)
 	ColorMode = InColorMode;
 	bColorModeChanged = true;
 
-	WeightTool->SetDisplayVertexColors(ColorMode != EWeightColorMode::FullMaterial);
+	WeightTool->SetDisplayVertexColors(ColorMode!=EWeightColorMode::FullMaterial);
+	WeightTool->SetFocusInViewport();
 }
 
 void USkinWeightsPaintToolProperties::SetBrushMode(EWeightEditOperation InBrushMode)
@@ -1007,6 +1011,8 @@ void USkinWeightsPaintToolProperties::SetBrushMode(EWeightEditOperation InBrushM
 	BrushRadius = GetBrushConfig().Radius;
 	BrushStrength = GetBrushConfig().Strength;
 	BrushFalloffAmount = GetBrushConfig().Falloff;
+
+	WeightTool->SetFocusInViewport();
 }
 
 bool USkinWeightsPaintTool::HitTest(const FRay& Ray, FHitResult& OutHit)
@@ -1106,6 +1112,20 @@ bool USkinWeightsPaintTool::OnUpdateHover(const FInputDeviceRay& DevicePos)
 	return true;
 }
 
+double USkinWeightsPaintTool::EstimateMaximumTargetDimension()
+{
+	if (const IPrimitiveComponentBackedTarget* TargetComponent = Cast<IPrimitiveComponentBackedTarget>(Target))
+	{
+		const USkeletalMeshComponent* Component = Cast<USkeletalMeshComponent>(TargetComponent->GetOwnerComponent());
+		if (USkeletalMesh* SkeletalMesh = Component->GetSkeletalMeshAsset())
+		{
+			return SkeletalMesh->GetBounds().SphereRadius * 2.0f;
+		}
+	}
+	
+	return Super::EstimateMaximumTargetDimension();
+}
+
 void USkinWeightsPaintTool::CalculateVertexROI(
 	const FBrushStampData& Stamp,
 	TArray<VertexIndex>& VertexIDs,
@@ -1616,6 +1636,14 @@ BoneIndex USkinWeightsPaintTool::GetBoneIndexFromName(const FName BoneName) cons
 	return Found ? *Found : INDEX_NONE;
 }
 
+void USkinWeightsPaintTool::SetFocusInViewport() const
+{
+	if (PersonaModeManagerContext.IsValid())
+	{
+		PersonaModeManagerContext->SetFocusInViewport();	
+	}
+}
+
 void USkinWeightsPaintTool::OnShutdown(EToolShutdownType ShutdownType)
 {
 	// save tool properties
@@ -2119,6 +2147,8 @@ void USkinWeightsPaintTool::HandleSkeletalMeshModified(const TArray<FName>& InBo
 	default:
 		checkNoEntry();
 	}
+
+	SetFocusInViewport();
 }
 
 void USkinWeightsPaintTool::ToggleEditingMode()
@@ -2144,6 +2174,8 @@ void USkinWeightsPaintTool::ToggleEditingMode()
 			PersonaModeManagerContext->GetPersonaEditorModeManager()->DeactivateMode(FPersonaEditModes::SkeletonSelection);	
 		}
 	}
+
+	SetFocusInViewport();
 }
 
 TObjectPtr<UPolygonSelectionMechanic> USkinWeightsPaintTool::GetSelectionMechanic()
@@ -2272,6 +2304,8 @@ void USkinWeightsPaintTool::OnPropertyModified(UObject* ModifiedObject, FPropert
 			Color.A = 1.f;
 		}
 	}
+	
+	SetFocusInViewport();
 }
 
 
diff --git a/Engine/Plugins/Experimental/MeshModelingToolsetExp/Source/MeshModelingToolsEditorOnlyExp/Public/SkeletalMesh/SkinWeightsPaintTool.h b/Engine/Plugins/Experimental/MeshModelingToolsetExp/Source/MeshModelingToolsEditorOnlyExp/Public/SkeletalMesh/SkinWeightsPaintTool.h
--- a/Engine/Plugins/Experimental/MeshModelingToolsetExp/Source/MeshModelingToolsEditorOnlyExp/Public/SkeletalMesh/SkinWeightsPaintTool.h
+++ b/Engine/Plugins/Experimental/MeshModelingToolsetExp/Source/MeshModelingToolsEditorOnlyExp/Public/SkeletalMesh/SkinWeightsPaintTool.h
@@ -49,7 +49,7 @@ enum class EWeightColorMode : uint8
 	Greyscale,
 	Ramp,
 	BoneColors,
-	FullMaterial
+	FullMaterial,
 };
 
 // brush falloff mode
@@ -347,6 +347,7 @@ public:
 	virtual void OnUpdateDrag(const FRay& Ray) override;
 	virtual void OnEndDrag(const FRay& Ray) override;
 	virtual bool OnUpdateHover(const FInputDeviceRay& DevicePos) override;
+	virtual double EstimateMaximumTargetDimension() override;
 
 	void Init(const FToolBuilderState& InSceneState);
 	
@@ -399,6 +400,8 @@ public:
 
 	// toggle the display of weights on the preview mesh (if false, uses the normal skeletal mesh material)
 	void SetDisplayVertexColors(bool bShowVertexColors=true) const;
+	// set focus back to viewport so that hotkeys are immediately detected while hovering
+	void SetFocusInViewport() const;
 
 	// HOW TO EDIT WEIGHTS WITH UNDO/REDO:
 	//
diff --git a/Engine/Plugins/Experimental/MeshModelingToolsetExp/Source/ModelingEditorUI/Private/SkinWeightDetailCustomization.cpp b/Engine/Plugins/Experimental/MeshModelingToolsetExp/Source/ModelingEditorUI/Private/SkinWeightDetailCustomization.cpp
--- a/Engine/Plugins/Experimental/MeshModelingToolsetExp/Source/ModelingEditorUI/Private/SkinWeightDetailCustomization.cpp
+++ b/Engine/Plugins/Experimental/MeshModelingToolsetExp/Source/ModelingEditorUI/Private/SkinWeightDetailCustomization.cpp
@@ -99,7 +99,7 @@ void FSkinWeightDetailCustomization::CustomizeDetails(IDetailLayoutBuilder& Deta
 					"Greyscale: Displays weights on the current bone by blending from black (0) to white (1).\n"
 					"Ramp: Displays weights on the current bone. Weights at 0 and 1 use the min and max colors. Weights inbetween 0 and 1 use the ramp colors.\n"
 					"Multi Color: Displays weights on ALL bones using the color of the bones.\n"
-					"Full Material: Shows the normal fully-lit mesh material.\n"))
+					"Full Material: Displays normal mesh materials with textures.\n"))
 			.Value_Lambda([this]()
 			{
 				return SkinToolSettings->ColorMode;
diff --git a/Engine/Source/Editor/Persona/Private/IPersonaEditorModeManager.cpp b/Engine/Source/Editor/Persona/Private/IPersonaEditorModeManager.cpp
--- a/Engine/Source/Editor/Persona/Private/IPersonaEditorModeManager.cpp
+++ b/Engine/Source/Editor/Persona/Private/IPersonaEditorModeManager.cpp
@@ -4,6 +4,10 @@
 
 #include "ContextObjectStore.h"
 #include "EdModeInteractiveToolsContext.h"
+#include "SEditorViewport.h"
+#include "Slate/SceneViewport.h"
+#include "Framework/Application/SlateApplication.h"
+#include "Framework/Application/SlateUser.h"
 
 IPersonaEditorModeManager* UPersonaEditorModeManagerContext::GetPersonaEditorModeManager() const
 {
@@ -22,6 +26,38 @@ void UPersonaEditorModeManagerContext::GetOnScreenDebugInfo(TArray<FText>& OutDe
 	return ModeManager->GetOnScreenDebugInfo(OutDebugText);
 }
 
+void UPersonaEditorModeManagerContext::SetFocusInViewport()
+{
+	const IPersonaEditorModeManager* PersonaEditorModeManager = GetPersonaEditorModeManager();
+	if (!PersonaEditorModeManager)
+	{
+		return;
+	}
+
+	const FEditorViewportClient* ViewportClient = PersonaEditorModeManager->GetHoveredViewportClient();
+	if (!ViewportClient)
+	{
+		ViewportClient = PersonaEditorModeManager->GetFocusedViewportClient();
+	}
+	if (!ViewportClient)
+	{
+		return;
+	}
+
+	const TWeakPtr<SViewport> ViewportWidget = ViewportClient->GetEditorViewportWidget()->GetSceneViewport()->GetViewportWidget();
+	if (!ViewportWidget.IsValid())
+	{
+		return;
+	}
+	
+	// set focus back to viewport so that hotkeys are immediately detected
+	TSharedPtr<SWidget> ViewportContents = ViewportWidget.Pin()->GetContent();
+	FSlateApplication::Get().ForEachUser([&ViewportContents](FSlateUser& User) 
+	{
+		User.SetFocus(ViewportContents.ToSharedRef());
+	});
+}
+
 IPersonaEditorModeManager::IPersonaEditorModeManager() : FAssetEditorModeManager()
 {
 	UContextObjectStore* ContextObjectStore = GetInteractiveToolsContext()->ContextObjectStore;
diff --git a/Engine/Source/Editor/Persona/Public/IPersonaEditorModeManager.h b/Engine/Source/Editor/Persona/Public/IPersonaEditorModeManager.h
--- a/Engine/Source/Editor/Persona/Public/IPersonaEditorModeManager.h
+++ b/Engine/Source/Editor/Persona/Public/IPersonaEditorModeManager.h
@@ -43,6 +43,7 @@ public:
 	IPersonaEditorModeManager* GetPersonaEditorModeManager() const;
 	virtual bool GetCameraTarget(FSphere& OutTarget) const override;
 	virtual void GetOnScreenDebugInfo(TArray<FText>& OutDebugText) const override;
+	virtual void SetFocusInViewport();
 private:
 	IPersonaEditorModeManager* ModeManager = nullptr;
 	static UPersonaEditorModeManagerContext* CreateFor(IPersonaEditorModeManager* InModeManager)
