Let the weld tool optionally split bowtie vertices, which can enable welds that otherwise would not be allowed

diff --git a/Engine/Plugins/Experimental/MeshModelingToolsetExp/Source/MeshModelingToolsExp/Private/WeldMeshEdgesTool.cpp b/Engine/Plugins/Experimental/MeshModelingToolsetExp/Source/MeshModelingToolsExp/Private/WeldMeshEdgesTool.cpp
--- a/Engine/Plugins/Experimental/MeshModelingToolsetExp/Source/MeshModelingToolsExp/Private/WeldMeshEdgesTool.cpp
+++ b/Engine/Plugins/Experimental/MeshModelingToolsetExp/Source/MeshModelingToolsExp/Private/WeldMeshEdgesTool.cpp
@@ -10,6 +10,7 @@
 
 #include "DynamicMesh/DynamicMesh3.h"
 #include "DynamicMesh/Operations/MergeCoincidentMeshEdges.h"
+#include "DynamicMeshEditor.h"
 #include "Operations/MeshResolveTJunctions.h"
 
 #include "Math/UnrealMathUtility.h"
@@ -58,6 +59,7 @@ public:
 	float SplitTangentsThreshold;
 	float SplitUVThreshold;
 	float SplitColorThreshold; 
+	bool bSplitBowties;
 
 	int32 InitialNumBoundaryEdges = 0;
 	int32 FinalNumBoundaryEdges = 0;
@@ -90,6 +92,28 @@ public:
 
 		ResultMesh->Copy(*SourceMesh, true, true, true, true);
 
+		if (bSplitBowties)
+		{
+			FDynamicMeshEditor Editor(ResultMesh.Get());
+			FDynamicMeshEditResult EditResult;
+			if (SelectedEdges.IsEmpty())
+			{
+				Editor.SplitBowties(EditResult);
+			}
+			else
+			{
+				// Note split bowties should never introduce new edges, just reconnects triangles to a new vertex, so edge IDs should be stable here
+				for (int32 EID : SelectedEdges)
+				{
+					FIndex2i EdgeV = ResultMesh->GetEdgeV(EID);
+					for (int32 SubIdx = 0; SubIdx < 2; ++SubIdx)
+					{
+						Editor.SplitBowties(EdgeV[SubIdx], EditResult);
+					}
+				}
+			}
+		}
+
 		// If we had bowties or broken normals on input, we will probably still have them on output...
 		bool bWasCleanMesh = ResultMesh->CheckValidity(FDynamicMesh3::FValidityOptions(), EValidityCheckFailMode::ReturnOnly);
 
@@ -230,6 +254,7 @@ void UWeldMeshEdgesTool::Setup()
 	Settings->WatchProperty(Settings->Tolerance, [this](float) { PreviewCompute->InvalidateResult(); });
 	Settings->WatchProperty(Settings->bOnlyUnique, [this](bool) { PreviewCompute->InvalidateResult(); });
 	Settings->WatchProperty(Settings->bResolveTJunctions, [this](bool) { PreviewCompute->InvalidateResult(); });
+	Settings->WatchProperty(Settings->bSplitBowties, [this](bool) { PreviewCompute->InvalidateResult(); });
 	Settings->WatchProperty(Settings->AttrWeldingMode, [this](EWeldMeshEdgesAttributeUIMode) { PreviewCompute->InvalidateResult(); });
 	Settings->WatchProperty(Settings->SplitUVThreshold, [this](float) { PreviewCompute->InvalidateResult(); });
 	Settings->WatchProperty(Settings->SplitColorThreshold, [this](float) { PreviewCompute->InvalidateResult(); });
@@ -337,6 +362,7 @@ void UWeldMeshEdgesTool::UpdateOpParameters(FWeldMeshEdgesOp& Op) const
 
 	Op.bOnlyUnique = Settings->bOnlyUnique;
 	Op.bResolveTJunctions = Settings->bResolveTJunctions;
+	Op.bSplitBowties = Settings->bSplitBowties;
 	Op.Tolerance = Settings->Tolerance;
 
 	Op.WeldAttributeMode = GetAttrMergeMode(Settings->AttrWeldingMode);
diff --git a/Engine/Plugins/Experimental/MeshModelingToolsetExp/Source/MeshModelingToolsExp/Public/WeldMeshEdgesTool.h b/Engine/Plugins/Experimental/MeshModelingToolsetExp/Source/MeshModelingToolsExp/Public/WeldMeshEdgesTool.h
--- a/Engine/Plugins/Experimental/MeshModelingToolsetExp/Source/MeshModelingToolsExp/Public/WeldMeshEdgesTool.h
+++ b/Engine/Plugins/Experimental/MeshModelingToolsetExp/Source/MeshModelingToolsExp/Public/WeldMeshEdgesTool.h
@@ -62,6 +62,10 @@ public:
 	UPROPERTY(EditAnywhere, Category = Options)
 	bool bResolveTJunctions = false;
 
+	/** If enabled, will split bowtie vertices before welding. This can in some cases enable more edges to be successfully welded */
+	UPROPERTY(EditAnywhere, Category = Options)
+	bool bSplitBowties = true;
+
 	/** Initial number of open boundary edges */
 	UPROPERTY(VisibleAnywhere, Category = Statistics)
 	int32 InitialEdges;
