merge settings structs in generate polygroups tool so 'output' category properties can live in the same ui section

diff --git a/Engine/Plugins/Experimental/MeshModelingToolsetExp/Source/MeshModelingToolsExp/Private/ConvertToPolygonsTool.cpp b/Engine/Plugins/Experimental/MeshModelingToolsetExp/Source/MeshModelingToolsExp/Private/ConvertToPolygonsTool.cpp
--- a/Engine/Plugins/Experimental/MeshModelingToolsetExp/Source/MeshModelingToolsExp/Private/ConvertToPolygonsTool.cpp
+++ b/Engine/Plugins/Experimental/MeshModelingToolsetExp/Source/MeshModelingToolsExp/Private/ConvertToPolygonsTool.cpp
@@ -347,20 +347,18 @@ void UConvertToPolygonsTool::Setup()
 	SetToolPropertySourceEnabled(CopyFromLayerProperties, Settings->ConversionMode == EConvertToPolygonsMode::CopyFromLayer);
 
 	// add picker for output group layer
-	OutputProperties = NewObject<UOutputPolygroupLayerProperties>(this);
-	AddToolPropertySource(OutputProperties);
-	OutputProperties->OptionsList.Reset();
-	OutputProperties->OptionsList.Add(TEXT("Default"));		// always have standard group
+	Settings->OptionsList.Reset();
+	Settings->OptionsList.Add(TEXT("Default"));		// always have standard group
 	if (OriginalDynamicMesh->Attributes())
 	{
 		for (int32 k = 0; k < OriginalDynamicMesh->Attributes()->NumPolygroupLayers(); k++)
 		{
 			FName Name = OriginalDynamicMesh->Attributes()->GetPolygroupLayer(k)->GetName();
-			OutputProperties->OptionsList.Add(Name.ToString());
+			Settings->OptionsList.Add(Name.ToString());
 		}
 	}
-	OutputProperties->OptionsList.Add(TEXT("Create New..."));
-	OutputProperties->WatchProperty(OutputProperties->GroupLayer, [&](FName NewName) { OutputProperties->bShowNewLayerName = (NewName == TEXT("Create New...")); });
+	Settings->OptionsList.Add(TEXT("Create New..."));
+	Settings->WatchProperty(Settings->GroupLayer, [&](FName NewName) { Settings->bShowNewLayerName = (NewName == TEXT("Create New...")); });
 
 	if (bUsingSelection)
 	{
@@ -460,7 +458,7 @@ void UConvertToPolygonsTool::OnShutdown(EToolShutdownType ShutdownType)
 			FDynamicMesh3* DynamicMeshResult = Result.Mesh.Get();
 			if (ensure(DynamicMeshResult != nullptr))
 			{
-				if (OutputProperties->GroupLayer != TEXT("Default"))
+				if (Settings->GroupLayer != TEXT("Default"))
 				{
 					FDynamicMesh3 UseResultMesh = *OriginalDynamicMesh;
 					if (UseResultMesh.HasAttributes() == false)
@@ -470,18 +468,18 @@ void UConvertToPolygonsTool::OnShutdown(EToolShutdownType ShutdownType)
 
 					// if we want to write to any layer other than default, we have to find or create it
 					FDynamicMeshPolygroupAttribute* UseAttribLayer = nullptr;
-					if (OutputProperties->GroupLayer == TEXT("Create New..."))
+					if (Settings->GroupLayer == TEXT("Create New..."))
 					{
 						// append new group layer and set it's name
 						int32 TargetLayerIdx = UseResultMesh.Attributes()->NumPolygroupLayers();
 						UseResultMesh.Attributes()->SetNumPolygroupLayers(TargetLayerIdx + 1);
 						UseAttribLayer = UseResultMesh.Attributes()->GetPolygroupLayer(TargetLayerIdx);
-						FString UseUniqueName = UE::Geometry::MakeUniqueGroupLayerName(UseResultMesh, OutputProperties->NewLayerName);
+						FString UseUniqueName = UE::Geometry::MakeUniqueGroupLayerName(UseResultMesh, Settings->NewLayerName);
 						UseAttribLayer->SetName( FName(UseUniqueName) );
 					}
 					else
 					{
-						UseAttribLayer = UE::Geometry::FindPolygroupLayerByName(UseResultMesh, OutputProperties->GroupLayer);
+						UseAttribLayer = UE::Geometry::FindPolygroupLayerByName(UseResultMesh, Settings->GroupLayer);
 					}
 
 					if (UseAttribLayer)
diff --git a/Engine/Plugins/Experimental/MeshModelingToolsetExp/Source/MeshModelingToolsExp/Public/ConvertToPolygonsTool.h b/Engine/Plugins/Experimental/MeshModelingToolsetExp/Source/MeshModelingToolsExp/Public/ConvertToPolygonsTool.h
--- a/Engine/Plugins/Experimental/MeshModelingToolsetExp/Source/MeshModelingToolsExp/Public/ConvertToPolygonsTool.h
+++ b/Engine/Plugins/Experimental/MeshModelingToolsetExp/Source/MeshModelingToolsExp/Public/ConvertToPolygonsTool.h
@@ -117,31 +117,17 @@ public:
 	/** group filtering */
 	UPROPERTY(EditAnywhere, Category = Filtering, meta = (UIMin = "1", UIMax = "100", ClampMin = "1", ClampMax = "10000", EditCondition = "ConversionMode != EConvertToPolygonsMode::CopyFromLayer"))
 	int32 MinGroupSize = 2;
-
-
-	/** If true, normals are recomputed per-group, with hard edges at group boundaries */
-	UPROPERTY(EditAnywhere, Category = Output, meta=(EditCondition = "ConversionMode != EConvertToPolygonsMode::CopyFromLayer") )
-	bool bCalculateNormals = false;
 	
 	/** Display each group with a different auto-generated color */
 	UPROPERTY(EditAnywhere, Category = Display)
 	bool bShowGroupColors = true;
-};
-
-
-
-
-
 
-
-UCLASS()
-class MESHMODELINGTOOLSEXP_API UOutputPolygroupLayerProperties : public UInteractiveToolPropertySet
-{
-	GENERATED_BODY()
-public:
+	/** If true, normals are recomputed per-group, with hard edges at group boundaries */
+	UPROPERTY(EditAnywhere, Category = Output, meta = (EditCondition = "ConversionMode != EConvertToPolygonsMode::CopyFromLayer"))
+	bool bCalculateNormals = false;
 
 	/** Select PolyGroup layer to use. */
-	UPROPERTY(EditAnywhere, Category = "Output", meta = (DisplayName = "Output Layer", GetOptions = GetGroupOptionsList, NoResetToDefault))
+	UPROPERTY(EditAnywhere, Category = Output, meta = (DisplayName = "Output Layer", GetOptions = GetGroupOptionsList, NoResetToDefault))
 	FName GroupLayer = "Default";
 
 	// Provides set of available group layers
@@ -156,7 +142,7 @@ public:
 	bool bShowNewLayerName = false;
 
 	/** Name of the new Group Layer */
-	UPROPERTY(EditAnywhere, Category = "Output", meta = (TransientToolProperty, DisplayName = "New Layer Name",
+	UPROPERTY(EditAnywhere, Category = Output, meta = (TransientToolProperty, DisplayName = "New Layer Name",
 		EditCondition = "bShowNewLayerName", HideEditConditionToggle, NoResetToDefault))
 	FString NewLayerName = TEXT("polygroups");
 };
@@ -211,9 +197,6 @@ protected:
 	UPROPERTY()
 	TObjectPtr<UPolygroupLayersProperties> CopyFromLayerProperties = nullptr;
 
-	UPROPERTY()
-	TObjectPtr<UOutputPolygroupLayerProperties> OutputProperties = nullptr;
-
 
 	UPROPERTY()
 	TObjectPtr<UMeshOpPreviewWithBackgroundCompute> PreviewCompute = nullptr;
