Skin Paint Tool: improve brush & vertex focus

diff --git a/Engine/Plugins/Experimental/MeshModelingToolsetExp/Source/MeshModelingToolsEditorOnlyExp/Private/SkeletalMesh/SkinWeightsPaintTool.cpp b/Engine/Plugins/Experimental/MeshModelingToolsetExp/Source/MeshModelingToolsEditorOnlyExp/Private/SkeletalMesh/SkinWeightsPaintTool.cpp
--- a/Engine/Plugins/Experimental/MeshModelingToolsetExp/Source/MeshModelingToolsEditorOnlyExp/Private/SkeletalMesh/SkinWeightsPaintTool.cpp
+++ b/Engine/Plugins/Experimental/MeshModelingToolsetExp/Source/MeshModelingToolsEditorOnlyExp/Private/SkeletalMesh/SkinWeightsPaintTool.cpp
@@ -868,25 +868,40 @@ void USkinWeightsPaintTool::Render(IToolsContextRenderAPI* RenderAPI)
 
 FBox USkinWeightsPaintTool::GetWorldSpaceFocusBox()
 {
-	// 1. Prioritize framing vertex selection if vertices are selected
-	if (PolygonSelectionMechanic)
+	// 1. Prioritize Brush & Vertex modes
+	switch (WeightToolProperties->EditingMode)
 	{
-		const FGroupTopologySelection& Selection = PolygonSelectionMechanic->GetActiveSelection();
-		if (!Selection.SelectedCornerIDs.IsEmpty())
-		{
-			TArray<VertexIndex> SelectedVertexIndices = Selection.SelectedCornerIDs.Array();
-			const FDynamicMesh3* Mesh = PreviewMesh->GetMesh();
-			FTransform3d Transform(PreviewMesh->GetTransform());
-			FAxisAlignedBox3d Bounds = FAxisAlignedBox3d::Empty();
-			for (const int32 VertexID : SelectedVertexIndices)
+	case EWeightEditMode::Brush:
 			{
-				Bounds.Contain(Transform.TransformPosition(Mesh->GetVertex(VertexID)));
+				const FVector Radius(CurrentBrushRadius);
+				return FBox(LastBrushStamp.WorldPosition - Radius, LastBrushStamp.WorldPosition + Radius);
 			}
-			if (Bounds.MaxDim() > FMathf::ZeroTolerance)
+			break;
+	case EWeightEditMode::Vertices:
+		{
+			if (PolygonSelectionMechanic)
 			{
-				return static_cast<FBox>(Bounds);
+				const FGroupTopologySelection& Selection = PolygonSelectionMechanic->GetActiveSelection();
+				if (!Selection.SelectedCornerIDs.IsEmpty())
+				{
+					const FDynamicMesh3* Mesh = PreviewMesh->GetMesh();
+					const FTransform3d Transform(PreviewMesh->GetTransform());
+					FAxisAlignedBox3d Bounds = FAxisAlignedBox3d::Empty();
+					for (const int32 VertexID : Selection.SelectedCornerIDs)
+					{
+						Bounds.Contain(Transform.TransformPosition(Mesh->GetVertex(VertexID)));
+					}
+					if (Bounds.MaxDim() > FMathf::ZeroTolerance)
+					{
+						return static_cast<FBox>(Bounds);
+					}
+				}
 			}
 		}
+		break;
+	case EWeightEditMode::Bones:
+	default:
+		break;
 	}
 
 	// 2. Fallback on framing selected bones (if there are any)
