Fixed speed of relax brush in skin paint tool (trivial)

diff --git a/Engine/Plugins/Experimental/MeshModelingToolsetExp/Source/MeshModelingToolsEditorOnlyExp/Private/SkeletalMesh/SkinWeightsPaintTool.cpp b/Engine/Plugins/Experimental/MeshModelingToolsetExp/Source/MeshModelingToolsEditorOnlyExp/Private/SkeletalMesh/SkinWeightsPaintTool.cpp
--- a/Engine/Plugins/Experimental/MeshModelingToolsetExp/Source/MeshModelingToolsEditorOnlyExp/Private/SkeletalMesh/SkinWeightsPaintTool.cpp
+++ b/Engine/Plugins/Experimental/MeshModelingToolsetExp/Source/MeshModelingToolsEditorOnlyExp/Private/SkeletalMesh/SkinWeightsPaintTool.cpp
@@ -1537,14 +1537,13 @@ void USkinWeightsPaintTool::RelaxWeightOnVertices(
 	}
 
 	constexpr int32 NumRelaxIterations = 3;
-	constexpr float PercentPerIter = 1.0f / static_cast<float>(NumRelaxIterations);
+	constexpr float PercentPerIteration = 0.9f;
 	for (int32 Iteration=0; Iteration<NumRelaxIterations; ++Iteration)
 	{
 		for (int32 VertexIndex = 0; VertexIndex < VerticesToEdit.Num(); ++VertexIndex)
 		{
 			const int32 VertexID = VerticesToEdit[VertexIndex];
-			float UseFalloff = VertexFalloffs.IsValidIndex(VertexIndex) ? VertexFalloffs[VertexIndex] * UseStrength : UseStrength;
-			UseFalloff *= PercentPerIter;
+			const float UseFalloff = (VertexFalloffs.IsValidIndex(VertexIndex) ? VertexFalloffs[VertexIndex] * UseStrength : UseStrength) * PercentPerIteration;
 
 			TMap<int32, float> FinalWeights;
 			const bool bSmoothSuccess = SmoothWeightsOp->SmoothWeightsAtVertex(VertexID, UseFalloff, FinalWeights);
