ModelingTools: Fix tooltip error for the DisplacementMap property in the Displace Mesh tool.
Rename Texture2D options category from "Options" to "Texture 2D Options".

diff --git a/Engine/Plugins/Experimental/MeshModelingToolsetExp/Source/MeshModelingToolsExp/Public/DisplaceMeshTool.h b/Engine/Plugins/Experimental/MeshModelingToolsetExp/Source/MeshModelingToolsExp/Public/DisplaceMeshTool.h
--- a/Engine/Plugins/Experimental/MeshModelingToolsetExp/Source/MeshModelingToolsExp/Public/DisplaceMeshTool.h
+++ b/Engine/Plugins/Experimental/MeshModelingToolsetExp/Source/MeshModelingToolsExp/Public/DisplaceMeshTool.h
@@ -158,43 +158,43 @@ class MESHMODELINGTOOLSEXP_API UDisplaceMeshTextureMapProperties : public UInter
 	GENERATED_BODY()
 			
 public:
-	/** Displacement map. Only the first channel is used. */
-	UPROPERTY(EditAnywhere, Category = Options)
+	/** The texture to sample for displacement. */
+	UPROPERTY(EditAnywhere, Category = TextureMapOptions)
 	TObjectPtr<UTexture2D> DisplacementMap = nullptr;
 
 	/** Channel in the displacement map to use. */
-	UPROPERTY(EditAnywhere, Category = Options)
+	UPROPERTY(EditAnywhere, Category = TextureMapOptions)
 	EDisplaceMeshToolChannelType Channel = EDisplaceMeshToolChannelType::Red;
 
 	/** The value in the texture map that corresponds to no displacement. For instance, if set to 0, then all
 	 displacement will be positive. If set to 0.5, displacement below 0.5 will be negative, and above will be
 	 positive. Default is for 128/255 to be no displacement. */
-	UPROPERTY(EditAnywhere, Category = Options, meta = (UIMin = "0", UIMax = "1", ClampMin = "0", ClampMax = "1"))
+	UPROPERTY(EditAnywhere, Category = TextureMapOptions, meta = (UIMin = "0", UIMax = "1", ClampMin = "0", ClampMax = "1"))
 	float DisplacementMapBaseValue = 128.0/255;
 
 	/** When sampling from the texture map, how to scale the mesh UV's in the x and y directions. For a mesh with
 	UV's in the range 0 to 1, setting a scale above 1 will result in tiling the texture map, and scaling below
 	1 will result in using only part of the texture map.*/
-	UPROPERTY(EditAnywhere, Category = Options, meta = (Delta = 0.01, LinearDeltaSensitivity = 1))
+	UPROPERTY(EditAnywhere, Category = TextureMapOptions, meta = (Delta = 0.01, LinearDeltaSensitivity = 1))
 	FVector2D UVScale = FVector2D(1,1);
 
 	/** When sampling from the texture map, how to offset the mesh UV's. This will result in offsetting the
 	tiling of the texture map across the mesh. */
-	UPROPERTY(EditAnywhere, Category = Options, meta = (UIMin = "-1", UIMax = "1"))
+	UPROPERTY(EditAnywhere, Category = TextureMapOptions, meta = (UIMin = "-1", UIMax = "1"))
 	FVector2D UVOffset = FVector2D(0, 0);
 
 	/** When true, applies a function to remap the values in the displacement map, which can be used 
 	 for contrast adjustment. The texture map values are converted to the range [0,1] before applying
 	 the remapping. */
-	UPROPERTY(EditAnywhere, Category = Options)
+	UPROPERTY(EditAnywhere, Category = TextureMapOptions)
 	bool bApplyAdjustmentCurve = false;
 
 	/** This curve is queried in the range [0,1] to adjust contrast of the displacement map. */
-	UPROPERTY(EditAnywhere, Category = Options, meta = (EditCondition = "bApplyAdjustmentCurve", EditConditionHides))
+	UPROPERTY(EditAnywhere, Category = TextureMapOptions, meta = (EditCondition = "bApplyAdjustmentCurve", EditConditionHides))
 	TObjectPtr<UCurveFloat> AdjustmentCurve = nullptr;
 
 	/** Recalculate normals from displaced mesh. Disable this if you are applying Displacements that are paired with an existing Normal Map in your Material. */
-	UPROPERTY(EditAnywhere, Category = Options)
+	UPROPERTY(EditAnywhere, Category = TextureMapOptions)
 	bool bRecalcNormals = true;
 };
 
