ModelingTools: Fixed tooltip error in TwistDegrees property of the Warp tool.

diff --git a/Engine/Plugins/Experimental/MeshModelingToolsetExp/Source/MeshModelingToolsExp/Public/MeshSpaceDeformerTool.h b/Engine/Plugins/Experimental/MeshModelingToolsetExp/Source/MeshModelingToolsExp/Public/MeshSpaceDeformerTool.h
--- a/Engine/Plugins/Experimental/MeshModelingToolsetExp/Source/MeshModelingToolsExp/Public/MeshSpaceDeformerTool.h
+++ b/Engine/Plugins/Experimental/MeshModelingToolsetExp/Source/MeshModelingToolsExp/Public/MeshSpaceDeformerTool.h
@@ -100,7 +100,7 @@ public:
 		EditCondition = "SelectedOperationType == ENonlinearOperationType::Bend", EditConditionHides))
 	float BendDegrees = 90;
 
-	/** Degrees of twist to from the lower bound to the upper bound along the gizmo Z axis. */
+	/** Degrees of twist from the lower bound to the upper bound along the gizmo Z axis. */
 	UPROPERTY(EditAnywhere, Category = Options, meta = (UIMin = "-360", UIMax = "360", ClampMin = "-3600", ClampMax = "3600", 
 		EditCondition = "SelectedOperationType == ENonlinearOperationType::Twist", EditConditionHides))
 	float TwistDegrees = 180;
