ModelingTools: Fixed group colors not rendering correctly in Simplify and Remesh tools.

diff --git a/Engine/Plugins/Experimental/MeshModelingToolsetExp/Source/MeshModelingToolsEditorOnlyExp/Private/SimplifyMeshTool.cpp b/Engine/Plugins/Experimental/MeshModelingToolsetExp/Source/MeshModelingToolsEditorOnlyExp/Private/SimplifyMeshTool.cpp
--- a/Engine/Plugins/Experimental/MeshModelingToolsetExp/Source/MeshModelingToolsEditorOnlyExp/Private/SimplifyMeshTool.cpp
+++ b/Engine/Plugins/Experimental/MeshModelingToolsetExp/Source/MeshModelingToolsEditorOnlyExp/Private/SimplifyMeshTool.cpp
@@ -258,7 +258,7 @@ void USimplifyMeshTool::UpdateVisualization()
 {
 	if (SimplifyProperties->bShowGroupColors)
 	{
-		Preview->OverrideMaterial = ToolSetupUtil::GetSelectionMaterial(GetToolManager());
+		Preview->OverrideMaterial = ToolSetupUtil::GetVertexColorMaterial(GetToolManager());
 		Preview->PreviewMesh->SetTriangleColorFunction([this](const FDynamicMesh3* Mesh, int TriangleID)
 		{
 			return LinearColors::SelectFColor(Mesh->GetTriangleGroup(TriangleID));
diff --git a/Engine/Plugins/Experimental/MeshModelingToolsetExp/Source/MeshModelingToolsExp/Private/RemeshMeshTool.cpp b/Engine/Plugins/Experimental/MeshModelingToolsetExp/Source/MeshModelingToolsExp/Private/RemeshMeshTool.cpp
--- a/Engine/Plugins/Experimental/MeshModelingToolsetExp/Source/MeshModelingToolsExp/Private/RemeshMeshTool.cpp
+++ b/Engine/Plugins/Experimental/MeshModelingToolsetExp/Source/MeshModelingToolsExp/Private/RemeshMeshTool.cpp
@@ -265,7 +265,7 @@ void URemeshMeshTool::UpdateVisualization()
 {
 	if (BasicProperties->bShowGroupColors)
 	{
-		Preview->OverrideMaterial = ToolSetupUtil::GetSelectionMaterial(GetToolManager());
+		Preview->OverrideMaterial = ToolSetupUtil::GetVertexColorMaterial(GetToolManager());
 		Preview->PreviewMesh->SetTriangleColorFunction([this](const FDynamicMesh3* Mesh, int TriangleID)
 		{
 			return LinearColors::SelectFColor(Mesh->GetTriangleGroup(TriangleID));
