Add comments to split meshes tool settings + make the source object invisible while the preview is visible

diff --git a/Engine/Plugins/Experimental/MeshModelingToolsetExp/Source/MeshModelingToolsExp/Private/SplitMeshesTool.cpp b/Engine/Plugins/Experimental/MeshModelingToolsetExp/Source/MeshModelingToolsExp/Private/SplitMeshesTool.cpp
--- a/Engine/Plugins/Experimental/MeshModelingToolsetExp/Source/MeshModelingToolsExp/Private/SplitMeshesTool.cpp
+++ b/Engine/Plugins/Experimental/MeshModelingToolsetExp/Source/MeshModelingToolsExp/Private/SplitMeshesTool.cpp
@@ -89,8 +89,6 @@ void USplitMeshesTool::Setup()
 
 	UpdateSplitMeshes();
 
-	UpdatePreviewVisibility(BasicProperties->bShowPreview);
-
 	SetToolDisplayName(LOCTEXT("ToolName", "Split"));
 	GetToolManager()->DisplayMessage(
 		LOCTEXT("OnStartTool", "Split Meshes into parts"),
@@ -99,9 +97,11 @@ void USplitMeshesTool::Setup()
 
 void USplitMeshesTool::UpdatePreviewVisibility(bool bShowPreview)
 {
+	checkSlow(Targets.Num() == SplitMeshes.Num());
 	for (int32 PreviewIdx = 0; PreviewIdx < PerTargetPreviews.Num(); ++PreviewIdx)
 	{
 		PerTargetPreviews[PreviewIdx]->SetAllVisible(bShowPreview);
+		UE::ToolTarget::SetSourceObjectVisible(Targets[PreviewIdx], !bShowPreview || SplitMeshes[PreviewIdx].bNoComponents);
 	}
 }
 
@@ -117,6 +117,12 @@ void USplitMeshesTool::OnShutdown(EToolShutdownType ShutdownType)
 		PreviewGeom->Disconnect();
 	}
 
+	// make sure source objects are visible
+	for (int32 ComponentIdx = 0; ComponentIdx < Targets.Num(); ComponentIdx++)
+	{
+		UE::ToolTarget::ShowSourceObject(Targets[ComponentIdx]);
+	}
+
 	OutputTypeProperties->SaveProperties(this, TEXT("OutputTypeFromInputTool"));
 	BasicProperties->SaveProperties(this);
 
@@ -385,6 +391,8 @@ void USplitMeshesTool::UpdateSplitMeshes()
 		GetToolManager()->DisplayMessage(FText(), EToolMessageLevel::UserWarning);
 	}
 
+	UpdatePreviewVisibility(BasicProperties->bShowPreview);
+
 }
 
 
diff --git a/Engine/Plugins/Experimental/MeshModelingToolsetExp/Source/MeshModelingToolsExp/Public/SplitMeshesTool.h b/Engine/Plugins/Experimental/MeshModelingToolsetExp/Source/MeshModelingToolsExp/Public/SplitMeshesTool.h
--- a/Engine/Plugins/Experimental/MeshModelingToolsetExp/Source/MeshModelingToolsExp/Public/SplitMeshesTool.h
+++ b/Engine/Plugins/Experimental/MeshModelingToolsetExp/Source/MeshModelingToolsExp/Public/SplitMeshesTool.h
@@ -43,15 +43,20 @@ class MESHMODELINGTOOLSEXP_API USplitMeshesToolProperties : public UInteractiveT
 {
 	GENERATED_BODY()
 public:
+
+	// Method to use to split the input(s) into output meshes
 	UPROPERTY(EditAnywhere, Category = Options, meta = (EditCondition = "!bIsInSelectionMode", EditConditionHides, HideEditConditionToggle))
 	ESplitMeshesMethod SplitMethod = ESplitMeshesMethod::ByMeshTopology;
 
+	// Vertices as close as this distance will be treated as overlapping, and kept in the same output mesh
 	UPROPERTY(EditAnywhere, Category = Options, meta = (ClampMin = .0001, UIMax = 1.0, EditCondition = "!bIsInSelectionMode && SplitMethod == ESplitMeshesMethod::ByVertexOverlap", EditConditionHides, HideEditConditionToggle))
 	double ConnectVerticesThreshold = 0.01;
 
+	// Whether to transfer materials to the output meshes
 	UPROPERTY(EditAnywhere, Category = Options)
 	bool bTransferMaterials = true;
 
+	// Whether to color mesh faces based on how they will be split into output meshes
 	UPROPERTY(EditAnywhere, Category = Options)
 	bool bShowPreview = true;
 
