Modeling Tools: Adds Selection support for the lattice tool

diff --git a/Engine/Plugins/Experimental/MeshModelingToolsetExp/Source/MeshModelingToolsExp/Private/LatticeDeformerTool.cpp b/Engine/Plugins/Experimental/MeshModelingToolsetExp/Source/MeshModelingToolsExp/Private/LatticeDeformerTool.cpp
--- a/Engine/Plugins/Experimental/MeshModelingToolsetExp/Source/MeshModelingToolsExp/Private/LatticeDeformerTool.cpp
+++ b/Engine/Plugins/Experimental/MeshModelingToolsetExp/Source/MeshModelingToolsExp/Private/LatticeDeformerTool.cpp
@@ -15,6 +15,7 @@
 #include "ToolBuilderUtil.h"
 #include "ToolSetupUtil.h"
 #include "DynamicMeshToMeshDescription.h"
+#include "DynamicSubmesh3.h"
 #include "DynamicMesh/MeshTransforms.h"
 #include "Algo/ForEach.h"
 #include "Operations/FFDLattice.h"
@@ -24,6 +25,8 @@
 #include "TargetInterfaces/MeshDescriptionProvider.h"
 #include "TargetInterfaces/PrimitiveComponentBackedTarget.h"
 #include "ModelingToolTargetUtil.h"
+#include "ToolTargetManager.h"
+#include "Selection/StoredMeshSelectionUtil.h"
 
 #include UE_INLINE_GENERATED_CPP_BY_NAME(LatticeDeformerTool)
 
@@ -105,11 +108,22 @@ void ULatticeDeformerToolProperties::PostAction(ELatticeDeformerToolAction Actio
 
 // Tool builder
 
-USingleSelectionMeshEditingTool* ULatticeDeformerToolBuilder::CreateNewTool(const FToolBuilderState& SceneState) const
+UMultiTargetWithSelectionTool* ULatticeDeformerToolBuilder::CreateNewTool(const FToolBuilderState& SceneState) const
 {
 	return NewObject<ULatticeDeformerTool>(SceneState.ToolManager);
 }
 
+bool ULatticeDeformerToolBuilder::CanBuildTool(const FToolBuilderState& SceneState) const
+{
+	if (RequiresInputSelection() && UE::Geometry::HaveAvailableGeometrySelection(SceneState) == false )
+	{
+		return false;
+	}
+
+	// disable multi-selection for now
+	return SceneState.TargetManager->CountSelectedAndTargetable(SceneState, GetTargetRequirements()) == 1;
+}
+
 
 // Operator factory
 
@@ -120,12 +134,28 @@ TUniquePtr<FDynamicMeshOperator> ULatticeDeformerOperatorFactory::MakeNewOperato
 		ELatticeInterpolation::Cubic :
 		ELatticeInterpolation::Linear;
 
-	TUniquePtr<FLatticeDeformerOp> LatticeDeformOp = MakeUnique<FLatticeDeformerOp>(
+	TUniquePtr<FLatticeDeformerOp> LatticeDeformOp = nullptr;
+
+	if (!LatticeDeformerTool->bHasSelection)
+	{
+		LatticeDeformOp = MakeUnique<FLatticeDeformerOp>(
+		LatticeDeformerTool->OriginalMesh,
+		LatticeDeformerTool->Lattice,
+		LatticeDeformerTool->ControlPointsMechanic->GetControlPoints(),
+		OpInterpolationType,
+		LatticeDeformerTool->Settings->bDeformNormals);
+	}
+	else
+	{
+		LatticeDeformOp = MakeUnique<FLatticeDeformerOp>(
 		LatticeDeformerTool->OriginalMesh,
+		LatticeDeformerTool->Submesh,
+		LatticeDeformerTool->WorldTransform,
 		LatticeDeformerTool->Lattice,
 		LatticeDeformerTool->ControlPointsMechanic->GetControlPoints(),
 		OpInterpolationType,
 		LatticeDeformerTool->Settings->bDeformNormals);
+	}
 
 	return LatticeDeformOp;
 }
@@ -150,12 +180,23 @@ bool ULatticeDeformerTool::CanAccept() const
 
 void ULatticeDeformerTool::InitializeLattice(TArray<FVector3d>& OutLatticePoints, TArray<FVector2i>& OutLatticeEdges)
 {
-	Lattice = MakeShared<FFFDLattice, ESPMode::ThreadSafe>(GetLatticeResolution(), *OriginalMesh, Settings->Padding);
+	TUniquePtr<UE::Geometry::FDynamicMesh3> MeshToDeform = nullptr;
+	
+	if (bHasSelection && Submesh)
+	{
+		MeshToDeform = MakeUnique<FDynamicMesh3>(Submesh->GetSubmesh());
+	}
+	else
+	{
+		MeshToDeform = MakeUnique<FDynamicMesh3>(*OriginalMesh.Get());
+	}
+	Lattice = MakeShared<FFFDLattice, ESPMode::ThreadSafe>(GetLatticeResolution(), *MeshToDeform.Get(), Settings->Padding);
+
 
 	Lattice->GenerateInitialLatticePositions(OutLatticePoints);
 
 	// Put the lattice in world space
-	FTransform3d LocalToWorld(Cast<IPrimitiveComponentBackedTarget>(Target)->GetWorldTransform());
+	FTransform3d LocalToWorld(Cast<IPrimitiveComponentBackedTarget>(Targets[0])->GetWorldTransform());
 	Algo::ForEach(OutLatticePoints, [&LocalToWorld](FVector3d& Point) {
 		Point = LocalToWorld.TransformPosition(Point);
 	});
@@ -171,11 +212,24 @@ void ULatticeDeformerTool::Setup()
 	GetToolManager()->DisplayMessage(LOCTEXT("LatticeDeformerToolMessage", 
 		"Drag the lattice control points to deform the mesh"), EToolMessageLevel::UserNotification);
 
+	// for now only supports one target
+	// TODO: include support for multiple targets
 	OriginalMesh = MakeShared<FDynamicMesh3, ESPMode::ThreadSafe>();
-	*OriginalMesh = UE::ToolTarget::GetDynamicMeshCopy(Target);
+	*OriginalMesh = UE::ToolTarget::GetDynamicMeshCopy(Targets[0]);
+	
+	bHasSelection = HasGeometrySelection(0);
+	if (bHasSelection)
+	{
+		TSet<int32> SelectionTriangleROI;
+		const FGeometrySelection& InputSelection = GetGeometrySelection(0);
+		EnumerateSelectionTriangles(InputSelection, *OriginalMesh,
+			[&](int32 TriangleID) { SelectionTriangleROI.Add(TriangleID);});
+
+		Submesh = MakeShared<FDynamicSubmesh3, ESPMode::ThreadSafe>(OriginalMesh.Get(), SelectionTriangleROI.Array());
+	}
 
 	// Note: Mesh will be implicitly transformed to world space by transforming the lattice; we account for whether that would invert the mesh here
-	MeshTransforms::ReverseOrientationIfNeeded(*OriginalMesh, (Cast<IPrimitiveComponentBackedTarget>(Target)->GetWorldTransform()));
+	MeshTransforms::ReverseOrientationIfNeeded(*OriginalMesh, (Cast<IPrimitiveComponentBackedTarget>(Targets[0])->GetWorldTransform()));
 
 	Settings = NewObject<ULatticeDeformerToolProperties>(this, TEXT("Lattice Deformer Tool Settings"));
 	Settings->Initialize(this);
@@ -220,7 +274,8 @@ void ULatticeDeformerTool::Setup()
 	ControlPointsMechanic = NewObject<ULatticeControlPointsMechanic>(this);
 	ControlPointsMechanic->Setup(this);
 	ControlPointsMechanic->SetWorld(GetTargetWorld());
-	FTransform3d LocalToWorld(Cast<IPrimitiveComponentBackedTarget>(Target)->GetWorldTransform());
+	FTransform3d LocalToWorld(Cast<IPrimitiveComponentBackedTarget>(Targets[0])->GetWorldTransform());
+	WorldTransform = LocalToWorld;
 	ControlPointsMechanic->Initialize(LatticePoints, LatticeEdges, LocalToWorld);
 
 	auto OnPointsChangedLambda = [this]()
@@ -359,7 +414,7 @@ void ULatticeDeformerTool::OnShutdown(EToolShutdownType ShutdownType)
 	Settings->SaveProperties(this);
 	ControlPointsMechanic->Shutdown();
 
-	IPrimitiveComponentBackedTarget* TargetComponent = Cast<IPrimitiveComponentBackedTarget>(Target);
+	IPrimitiveComponentBackedTarget* TargetComponent = Cast<IPrimitiveComponentBackedTarget>(Targets[0]);
 	TargetComponent->SetOwnerVisibility(true);
 
 	if (Preview)
@@ -378,7 +433,7 @@ void ULatticeDeformerTool::OnShutdown(EToolShutdownType ShutdownType)
 			FTransform3d LocalToWorld(TargetComponent->GetWorldTransform());
 			MeshTransforms::ApplyTransformInverse(*DynamicMeshResult, LocalToWorld, true);
 
-			UE::ToolTarget::CommitMeshDescriptionUpdateViaDynamicMesh(Target, *DynamicMeshResult, true);
+			UE::ToolTarget::CommitMeshDescriptionUpdateViaDynamicMesh(Targets[0], *DynamicMeshResult, true);
 
 			GetToolManager()->EndUndoTransaction();
 		}
@@ -393,12 +448,12 @@ void ULatticeDeformerTool::StartPreview()
 
 	Preview = NewObject<UMeshOpPreviewWithBackgroundCompute>(LatticeDeformOpCreator);
 	Preview->Setup(GetTargetWorld(), LatticeDeformOpCreator);
-	ToolSetupUtil::ApplyRenderingConfigurationToPreview(Preview->PreviewMesh, Target);
+	ToolSetupUtil::ApplyRenderingConfigurationToPreview(Preview->PreviewMesh, Targets[0]);
 
 	Preview->SetIsMeshTopologyConstant(true, EMeshRenderAttributeFlags::Positions | EMeshRenderAttributeFlags::VertexNormals);
 
 	FComponentMaterialSet MaterialSet;
-	Cast<IMaterialProvider>(Target)->GetMaterialSet(MaterialSet);
+	Cast<IMaterialProvider>(Targets[0])->GetMaterialSet(MaterialSet);
 	Preview->ConfigureMaterials(MaterialSet.Materials,
 								ToolSetupUtil::GetDefaultWorkingMaterial(GetToolManager())
 	);
@@ -414,7 +469,7 @@ void ULatticeDeformerTool::StartPreview()
 	Preview->SetVisibility(true);
 	Preview->InvalidateResult();
 
-	Cast<IPrimitiveComponentBackedTarget>(Target)->SetOwnerVisibility(false);
+	Cast<IPrimitiveComponentBackedTarget>(Targets[0])->SetOwnerVisibility(false);
 }
 
 
@@ -450,7 +505,7 @@ void ULatticeDeformerTool::OnTick(float DeltaTime)
 			TArray<FVector3d> LatticePoints;
 			TArray<FVector2i> LatticeEdges;
 			InitializeLattice(LatticePoints, LatticeEdges);
-			FTransform3d LocalToWorld(Cast<IPrimitiveComponentBackedTarget>(Target)->GetWorldTransform());
+			FTransform3d LocalToWorld(Cast<IPrimitiveComponentBackedTarget>(Targets[0])->GetWorldTransform());
 			ControlPointsMechanic->Initialize(LatticePoints, LatticeEdges, LocalToWorld);
 			Preview->InvalidateResult();
 			bShouldRebuild = false;
diff --git a/Engine/Plugins/Experimental/MeshModelingToolsetExp/Source/MeshModelingToolsExp/Public/LatticeDeformerTool.h b/Engine/Plugins/Experimental/MeshModelingToolsetExp/Source/MeshModelingToolsExp/Public/LatticeDeformerTool.h
--- a/Engine/Plugins/Experimental/MeshModelingToolsetExp/Source/MeshModelingToolsExp/Public/LatticeDeformerTool.h
+++ b/Engine/Plugins/Experimental/MeshModelingToolsetExp/Source/MeshModelingToolsExp/Public/LatticeDeformerTool.h
@@ -6,23 +6,31 @@
 #include "InteractiveTool.h" //UInteractiveToolPropertySet
 #include "InteractiveToolBuilder.h" //UInteractiveToolBuilder
 #include "InteractiveToolChange.h" //FToolCommandChange
-#include "SingleSelectionTool.h" //USingleSelectionTool
+#include "BaseTools/MultiTargetWithSelectionTool.h" // UMultiTargetWithSelectionTool
+#include "DynamicMesh/MeshSharingUtil.h"
 #include "Operations/FFDLattice.h"
-#include "BaseTools/SingleSelectionMeshEditingTool.h"
 #include "Solvers/ConstrainedMeshDeformer.h"
 
 #include "LatticeDeformerTool.generated.h"
 
+namespace UE::Geometry
+{
+	struct FDynamicSubmesh3;
+}
+
 class ULatticeControlPointsMechanic;
 class UMeshOpPreviewWithBackgroundCompute;
 
 UCLASS()
-class MESHMODELINGTOOLSEXP_API ULatticeDeformerToolBuilder : public USingleSelectionMeshEditingToolBuilder
+class MESHMODELINGTOOLSEXP_API ULatticeDeformerToolBuilder : public UMultiTargetWithSelectionToolBuilder
 {
 	GENERATED_BODY()
 
 public:
-	virtual USingleSelectionMeshEditingTool* CreateNewTool(const FToolBuilderState& SceneState) const override;
+	virtual UMultiTargetWithSelectionTool* CreateNewTool(const FToolBuilderState& SceneState) const override;
+	virtual bool CanBuildTool(const FToolBuilderState& SceneState) const override;
+
+	virtual bool RequiresInputSelection() const override { return false; }
 };
 
 
@@ -127,7 +135,7 @@ public:
 
 /** Deform a mesh using a regular hexahedral lattice */
 UCLASS()
-class MESHMODELINGTOOLSEXP_API ULatticeDeformerTool : public USingleSelectionMeshEditingTool, public IInteractiveToolManageGeometrySelectionAPI
+class MESHMODELINGTOOLSEXP_API ULatticeDeformerTool : public UMultiTargetWithSelectionTool, public IInteractiveToolManageGeometrySelectionAPI
 {
 	GENERATED_BODY()
 
@@ -160,6 +168,10 @@ protected:
 	// Input mesh
 	TSharedPtr<UE::Geometry::FDynamicMesh3, ESPMode::ThreadSafe> OriginalMesh;
 
+	TSharedPtr<UE::Geometry::FDynamicSubmesh3, ESPMode::ThreadSafe> Submesh = nullptr;
+
+	FTransform3d WorldTransform;
+
 	TSharedPtr<UE::Geometry::FFFDLattice, ESPMode::ThreadSafe> Lattice;
 
 	UPROPERTY()
@@ -176,6 +188,8 @@ protected:
 
 	bool bShouldRebuild = false;
 
+	bool bHasSelection = false;
+
 	// Create and store an FFFDLattice. Pass out the lattice's positions and edges.
 	void InitializeLattice(TArray<FVector3d>& OutLatticePoints, TArray<UE::Geometry::FVector2i>& OutLatticeEdges);
 
diff --git a/Engine/Plugins/Runtime/MeshModelingToolset/Source/ModelingOperators/Private/DeformationOps/LatticeDeformerOp.cpp b/Engine/Plugins/Runtime/MeshModelingToolset/Source/ModelingOperators/Private/DeformationOps/LatticeDeformerOp.cpp
--- a/Engine/Plugins/Runtime/MeshModelingToolset/Source/ModelingOperators/Private/DeformationOps/LatticeDeformerOp.cpp
+++ b/Engine/Plugins/Runtime/MeshModelingToolset/Source/ModelingOperators/Private/DeformationOps/LatticeDeformerOp.cpp
@@ -1,6 +1,7 @@
 // Copyright Epic Games, Inc. All Rights Reserved.
 
 #include "DeformationOps/LatticeDeformerOp.h"
+#include "DynamicSubmesh3.h"
 #include "Operations/FFDLattice.h"
 
 using namespace UE::Geometry;
@@ -22,17 +23,51 @@ void FLatticeDeformerOp::CalculateResult(FProgressCancel* Progress)
 	TArray<FVector3d> DeformedPositions;
 	FLatticeExecutionInfo ExecutionInfo = FLatticeExecutionInfo();
 	ExecutionInfo.bParallel = true;
+	
+	const bool bUsingSubmesh = Submesh != nullptr;
+	
+	// retrieves the deformed vertex positions of the applicable mesh - either the whole mesh or a submesh, depending on selection
 	Lattice->GetDeformedMeshVertexPositions(LatticeControlPoints, DeformedPositions, InterpolationType, ExecutionInfo, Progress);
-	check(ResultMesh->VertexCount() == DeformedPositions.Num());
+	
+	if (bUsingSubmesh)
+	{
+		check(Submesh->GetSubmesh().VertexCount() == DeformedPositions.Num())
+	}
+	else
+	{
+		check(ResultMesh->VertexCount() == DeformedPositions.Num());
+	}
 
 	if (Progress && Progress->Cancelled())
 	{
 		return;
 	}
 
-	for (int vid : ResultMesh->VertexIndicesItr())
+	if (bUsingSubmesh)
 	{
-		ResultMesh->SetVertex(vid, DeformedPositions[vid]);
+		for (const int SubVID : Submesh->GetSubmesh().VertexIndicesItr())
+		{
+			// retrieve the index in the Base Mesh which maps to the (deformed) submesh
+			const int BaseMeshVertexID = Submesh->MapVertexToBaseMesh(SubVID);
+			// sets the deformed position of the submesh in the result mesh
+			ResultMesh->SetVertex(BaseMeshVertexID, DeformedPositions[SubVID]);
+		}
+		for (const int BaseVID : ResultMesh->VertexIndicesItr())
+		{
+			// for all vertices in the base mesh that are NOT in the submesh, they keep their original (transformed) position
+			if (Submesh->MapVertexToSubmesh(BaseVID) == INDEX_NONE)
+			{
+				ResultMesh->SetVertex(BaseVID, WorldTransform.TransformPosition(ResultMesh->GetVertex(BaseVID)));
+			}
+		}
+	}
+	else
+	{
+		// if deforming entire mesh, all the mesh's new positions will be in DeformedPositions
+		for (const int VID : ResultMesh->VertexIndicesItr())
+		{
+			ResultMesh->SetVertex(VID, DeformedPositions[VID]);
+		}
 	}
 
 	if (bDeformNormals)
@@ -101,3 +136,19 @@ FLatticeDeformerOp::FLatticeDeformerOp(TSharedPtr<FDynamicMesh3, ESPMode::Thread
 	InterpolationType(InInterpolationType),
 	bDeformNormals(bInDeformNormals)
 {}
+
+FLatticeDeformerOp::FLatticeDeformerOp(TSharedPtr<FDynamicMesh3, ESPMode::ThreadSafe> InOriginalMesh,
+									   TSharedPtr<FDynamicSubmesh3, ESPMode::ThreadSafe> InSubmesh,
+									   FTransform3d InTransform,
+									   TSharedPtr<FFFDLattice, ESPMode::ThreadSafe> InLattice,
+									   const TArray<FVector3d>& InLatticeControlPoints,
+									   ELatticeInterpolation InInterpolationType,
+									   bool bInDeformNormals) :
+	Lattice(InLattice),
+	OriginalMesh(InOriginalMesh),
+	Submesh(InSubmesh),
+	WorldTransform(InTransform),
+	LatticeControlPoints(InLatticeControlPoints),
+	InterpolationType(InInterpolationType),
+	bDeformNormals(bInDeformNormals)
+{}
diff --git a/Engine/Plugins/Runtime/MeshModelingToolset/Source/ModelingOperators/Public/DeformationOps/LatticeDeformerOp.h b/Engine/Plugins/Runtime/MeshModelingToolset/Source/ModelingOperators/Public/DeformationOps/LatticeDeformerOp.h
--- a/Engine/Plugins/Runtime/MeshModelingToolset/Source/ModelingOperators/Public/DeformationOps/LatticeDeformerOp.h
+++ b/Engine/Plugins/Runtime/MeshModelingToolset/Source/ModelingOperators/Public/DeformationOps/LatticeDeformerOp.h
@@ -22,6 +22,14 @@ public:
 					   UE::Geometry::ELatticeInterpolation InInterpolationType,
 					   bool bInDeformNormals);
 
+	FLatticeDeformerOp(TSharedPtr<FDynamicMesh3, ESPMode::ThreadSafe> InOriginalMesh,
+					   TSharedPtr<FDynamicSubmesh3, ESPMode::ThreadSafe> InSubmesh,
+					   FTransform3d InWorldTransform,
+					   TSharedPtr<UE::Geometry::FFFDLattice, ESPMode::ThreadSafe> InLattice,
+					   const TArray<FVector3d>& InLatticeControlPoints,
+					   UE::Geometry::ELatticeInterpolation InInterpolationType,
+					   bool bInDeformNormals);
+
 	// FDynamicMeshOperator implementation
 	void CalculateResult(FProgressCancel* Progress) override;
 
@@ -30,6 +38,8 @@ protected:
 	// Inputs
 	const TSharedPtr<const UE::Geometry::FFFDLattice, ESPMode::ThreadSafe> Lattice;
 	const TSharedPtr<const FDynamicMesh3, ESPMode::ThreadSafe> OriginalMesh;
+	const TSharedPtr<const FDynamicSubmesh3, ESPMode::ThreadSafe> Submesh;
+	const FTransform3d WorldTransform;
 	const TArray<FVector3d> LatticeControlPoints;
 	UE::Geometry::ELatticeInterpolation InterpolationType;
 	bool bDeformNormals = false;
