[GeometryProcessing] Support inputting empty geometry.

Previously we would generate a cube for merged parts with empty geometry. We should simply ingore those parts.
This was preventing us to merge certain geometry only on specific LODs.

diff --git a/Engine/Plugins/Experimental/MeshModelingToolsetExp/Source/GeometryProcessingAdapters/Private/CombineMeshInstancesImpl.cpp b/Engine/Plugins/Experimental/MeshModelingToolsetExp/Source/GeometryProcessingAdapters/Private/CombineMeshInstancesImpl.cpp
--- a/Engine/Plugins/Experimental/MeshModelingToolsetExp/Source/GeometryProcessingAdapters/Private/CombineMeshInstancesImpl.cpp
+++ b/Engine/Plugins/Experimental/MeshModelingToolsetExp/Source/GeometryProcessingAdapters/Private/CombineMeshInstancesImpl.cpp
@@ -421,13 +421,6 @@ void InitializeAssemblySourceMeshesFromLOD(
 			ExtractSourceMeshLOD(*Part, 0, Target.SourceMeshLODs[0]);
 		}
 
-		// now if first LOD is missing, just fall back to a box
-		if (Target.SourceMeshLODs[0].TriangleCount() == 0)
-		{
-			FGridBoxMeshGenerator BoxGen;
-			Target.SourceMeshLODs[0].Copy(&BoxGen.Generate());
-		}
-
 		// now make sure every one of our Source LODs has a mesh by copying from N-1
 		for (int32 k = 1; k < NumSourceLODs; ++k)
 		{
@@ -1000,6 +993,12 @@ static void ComputeSimplePartApproximation(
 	EApproximatePartMethod ApproxMethod)
 {
 
+	if (SourcePartMesh.TriangleCount() == 0)
+	{
+		// Nothing to approximate.
+		return;
+	}
+
 	if (ApproxMethod == EApproximatePartMethod::AxisAlignedBox)
 	{
 		ComputeBoxApproximation(SourcePartMesh, DestMesh, true);
@@ -1348,12 +1347,13 @@ void ComputeMeshApproximations(
 	using namespace UE::Geometry;
 	const double AngleThresholdDeg = CombineOptions.HardNormalAngleDeg;
 
-	int32 NumParts = Assembly.Parts.Num();
+	const int32 NumParts = Assembly.Parts.Num();
 	Assembly.OptimizedMeshGeometry.SetNum(NumParts);
 
-	int32 NumSimplifiedLODs = CombineOptions.NumSimplifiedLODs;
-	int32 NumApproxLODs = FMath::Max(1, 
+	const int32 NumSimplifiedLODs = CombineOptions.NumSimplifiedLODs;
+	const int32 NumApproxLODs = FMath::Max(0, 
 		CombineOptions.NumLODs - CombineOptions.NumCopiedLODs - CombineOptions.NumSimplifiedLODs);
+	const bool bNeedsApproximateDecorativePartLODs = CombineOptions.NumLODs >= CombineOptions.FilterDecorativePartsLODLevel - CombineOptions.ApproximateDecorativePartLODs;
 
 	bool bVerbose = CVarGeometryCombineMeshInstancesVerbose.GetValueOnAnyThread();
 	ParallelFor(NumParts, [&](int32 Index)
@@ -1398,8 +1398,8 @@ void ComputeMeshApproximations(
 		// Note that ExtraLODs is a hack here - we are computing more than necessary
 		// so that the cost approximation strategy below has additional simplified approximations available. 
 		// This could be smarter, but this dumb method works OK for now...
-		int32 ExtraLODs = 10;
-		int32 UseNumApproxLODs = NumApproxLODs + ExtraLODs;
+		const int32 ExtraLODs = 10;
+		const int32 UseNumApproxLODs = NumApproxLODs > 0 || bNeedsApproximateDecorativePartLODs ? NumApproxLODs + ExtraLODs : 0;
 		ApproxGeo.ApproximateMeshLODs.SetNum(UseNumApproxLODs);
 		double InitialTriCost = CombineOptions.OptimizeBaseTriCost;
 		TArray<EApproximatePartMethod> SelectedMethodID; SelectedMethodID.SetNum(UseNumApproxLODs);		// useful for debugging
@@ -3463,31 +3463,27 @@ void BuildCombinedMesh(
 
 		for (int32 LODLevel = 0; LODLevel < NumLODs; ++LODLevel)
 		{
-			const FDynamicMesh3* SourceAppendMesh = nullptr;
 			const FDynamicMesh3* ApproximateAppendMesh = nullptr;
 			const FDynamicMesh3* UseAppendMesh = nullptr;
 
 			// default approximate mesh to lowest-quality approximation (box), need to do this
 			// so that we always have something to swap to for Decorative parts
-			ApproximateAppendMesh = &OptimizedGeometry.ApproximateMeshLODs.Last();
+			ApproximateAppendMesh = OptimizedGeometry.ApproximateMeshLODs.Num() > 0 ? &OptimizedGeometry.ApproximateMeshLODs.Last() : nullptr;
 
 			ECombinedLODType LevelLODType = LODTypes[LODLevel];
 			if (LevelLODType == ECombinedLODType::Copied)
 			{
-				SourceAppendMesh = (LODLevel < SourceGeometry.SourceMeshLODs.Num()) ? 
+				UseAppendMesh = (LODLevel < SourceGeometry.SourceMeshLODs.Num()) ?
 					&SourceGeometry.SourceMeshLODs[LODLevel] : &SourceGeometry.SourceMeshLODs.Last();
-				UseAppendMesh = SourceAppendMesh;
 			}
 			else if (LevelLODType == ECombinedLODType::Simplified)
 			{
 				int32 SimplifiedLODIndex = LODLevel - CombineOptions.NumCopiedLODs;
-				SourceAppendMesh = &OptimizedGeometry.SimplifiedMeshLODs[SimplifiedLODIndex];
-				UseAppendMesh = SourceAppendMesh;
+				UseAppendMesh = &OptimizedGeometry.SimplifiedMeshLODs[SimplifiedLODIndex];;
 			}
 			else if (LevelLODType == ECombinedLODType::VoxWrapped)
 			{
-				SourceAppendMesh = &SourceGeometry.SourceMeshLODs.Last();
-				UseAppendMesh = SourceAppendMesh;
+				UseAppendMesh = &SourceGeometry.SourceMeshLODs.Last();;
 			}
 			else // ECombinedLODType::Approximated
 			{
@@ -3518,6 +3514,7 @@ void BuildCombinedMesh(
 					// at last detail part LOD, switch to approximate mesh
 					if (LODLevel >= (CombineOptions.FilterDecorativePartsLODLevel - CombineOptions.ApproximateDecorativePartLODs) )
 					{
+						check(ApproximateAppendMesh)
 						InstanceAppendMesh = ApproximateAppendMesh;
 					}
 				}
