ModelingTools: Fixed crash in Paint Maps tool when using the brush adjuster behavior.

diff --git a/Engine/Plugins/Experimental/MeshModelingToolsetExp/Source/MeshModelingToolsExp/Private/MeshAttributePaintTool.cpp b/Engine/Plugins/Experimental/MeshModelingToolsetExp/Source/MeshModelingToolsExp/Private/MeshAttributePaintTool.cpp
--- a/Engine/Plugins/Experimental/MeshModelingToolsetExp/Source/MeshModelingToolsExp/Private/MeshAttributePaintTool.cpp
+++ b/Engine/Plugins/Experimental/MeshModelingToolsetExp/Source/MeshModelingToolsExp/Private/MeshAttributePaintTool.cpp
@@ -750,19 +750,22 @@ TUniquePtr<FMeshAttributePaintChange> UMeshAttributePaintTool::EndChange()
 	}
 
 	TUniquePtr<FMeshAttributePaintChange> Result = ActiveChangeBuilder->ExtractResult();
-	
-	Result->ApplyFunction = [](UObject* Object, const int32& AttribIndex, const TArray<int32>& Indices, const TArray<float>& Values)
-	{
-		UMeshAttributePaintTool* Tool = CastChecked<UMeshAttributePaintTool>(Object);
-		Tool->ExternalUpdateValues(AttribIndex, Indices, Values);
-	};
-	Result->RevertFunction = [](UObject* Object, const int32& AttribIndex, const TArray<int32>& Indices, const TArray<float>& Values)
+	if (Result)
 	{
-		UMeshAttributePaintTool* Tool = CastChecked<UMeshAttributePaintTool>(Object);
-		Tool->ExternalUpdateValues(AttribIndex, Indices, Values);
-	};
-
-	return MoveTemp(Result);
+		Result->ApplyFunction = [](UObject* Object, const int32& AttribIndex, const TArray<int32>& Indices, const TArray<float>& Values)
+		{
+			UMeshAttributePaintTool* Tool = CastChecked<UMeshAttributePaintTool>(Object);
+			Tool->ExternalUpdateValues(AttribIndex, Indices, Values);
+		};
+		Result->RevertFunction = [](UObject* Object, const int32& AttribIndex, const TArray<int32>& Indices, const TArray<float>& Values)
+		{
+			UMeshAttributePaintTool* Tool = CastChecked<UMeshAttributePaintTool>(Object);
+			Tool->ExternalUpdateValues(AttribIndex, Indices, Values);
+		};
+		return MoveTemp(Result);
+	}
+	
+	return nullptr;
 }
 
 
