Skin Weights Paint Tool: removed viewport focus on skeletal mesh changes notification
    - navigating in the skeleton tree (for example) using keyboard resulted in switchong the focus back to the viewport widget

#rnx

diff --git a/Engine/Plugins/Experimental/MeshModelingToolsetExp/Source/MeshModelingToolsEditorOnlyExp/Private/SkeletalMesh/SkinWeightsPaintTool.cpp b/Engine/Plugins/Experimental/MeshModelingToolsetExp/Source/MeshModelingToolsEditorOnlyExp/Private/SkeletalMesh/SkinWeightsPaintTool.cpp
--- a/Engine/Plugins/Experimental/MeshModelingToolsetExp/Source/MeshModelingToolsEditorOnlyExp/Private/SkeletalMesh/SkinWeightsPaintTool.cpp
+++ b/Engine/Plugins/Experimental/MeshModelingToolsetExp/Source/MeshModelingToolsEditorOnlyExp/Private/SkeletalMesh/SkinWeightsPaintTool.cpp
@@ -2161,8 +2161,6 @@ void USkinWeightsPaintTool::HandleSkeletalMeshModified(const TArray<FName>& InBo
 	default:
 		checkNoEntry();
 	}
-
-	SetFocusInViewport();
 }
 
 void USkinWeightsPaintTool::ToggleEditingMode()
