ModelingTools: Fixed crash in EditMaterials tool when using the BrushAdjuster behavior.

diff --git a/Engine/Plugins/Experimental/MeshModelingToolsetExp/Source/MeshModelingToolsExp/Private/MeshAttributePaintTool.cpp b/Engine/Plugins/Experimental/MeshModelingToolsetExp/Source/MeshModelingToolsExp/Private/MeshAttributePaintTool.cpp
--- a/Engine/Plugins/Experimental/MeshModelingToolsetExp/Source/MeshModelingToolsExp/Private/MeshAttributePaintTool.cpp
+++ b/Engine/Plugins/Experimental/MeshModelingToolsetExp/Source/MeshModelingToolsExp/Private/MeshAttributePaintTool.cpp
@@ -355,17 +355,23 @@ void UMeshAttributePaintTool::OnUpdateDrag(const FRay& WorldRay)
 
 void UMeshAttributePaintTool::OnEndDrag(const FRay& Ray)
 {
+	// Capture brush stroke state prior to invoking Super::OnEndDrag
+	const bool bWasInBrushStroke = IsInBrushStroke();
+	
 	UDynamicMeshBrushTool::OnEndDrag(Ray);
 
 	bInRemoveStroke = bInSmoothStroke = false;
 	bStampPending = false;
 
-	// close change record
-	TUniquePtr<FMeshAttributePaintChange> Change = EndChange();
-	if (Change)
+	if (bWasInBrushStroke)
 	{
-		GetToolManager()->EmitObjectChange(this, MoveTemp(Change), LOCTEXT("AttributeValuesChange", "Paint"));
-		LongTransactions.Close(GetToolManager());
+		// close change record
+		TUniquePtr<FMeshAttributePaintChange> Change = EndChange();
+		if (Change)
+		{
+			GetToolManager()->EmitObjectChange(this, MoveTemp(Change), LOCTEXT("AttributeValuesChange", "Paint"));
+			LongTransactions.Close(GetToolManager());
+		}
 	}
 }
 
diff --git a/Engine/Plugins/Experimental/MeshModelingToolsetExp/Source/MeshModelingToolsExp/Private/MeshSelectionTool.cpp b/Engine/Plugins/Experimental/MeshModelingToolsetExp/Source/MeshModelingToolsExp/Private/MeshSelectionTool.cpp
--- a/Engine/Plugins/Experimental/MeshModelingToolsetExp/Source/MeshModelingToolsExp/Private/MeshSelectionTool.cpp
+++ b/Engine/Plugins/Experimental/MeshModelingToolsetExp/Source/MeshModelingToolsExp/Private/MeshSelectionTool.cpp
@@ -752,15 +752,24 @@ void UMeshSelectionTool::UpdateFaceSelection(const FBrushStampData& Stamp, const
 
 void UMeshSelectionTool::OnEndDrag(const FRay& Ray)
 {
+	// Capture brush stroke state prior to invoking Super::OnEndDrag
+	const bool bWasInBrushStroke = IsInBrushStroke();
+	
 	UDynamicMeshBrushTool::OnEndDrag(Ray);
 
 	bInRemoveStroke = false;
 	bStampPending = false;
 
 	// close change record
-	TUniquePtr<FToolCommandChange> Change = EndChange();
-	GetToolManager()->EmitObjectChange(Selection, MoveTemp(Change), LOCTEXT("MeshSelectionChange", "Mesh Selection"));
-	LongTransactions.Close(GetToolManager());
+	if (bWasInBrushStroke)
+	{
+		TUniquePtr<FToolCommandChange> Change = EndChange();
+		if (Change)
+		{
+			GetToolManager()->EmitObjectChange(Selection, MoveTemp(Change), LOCTEXT("MeshSelectionChange", "Mesh Selection"));
+			LongTransactions.Close(GetToolManager());
+		}
+	}
 }
 
 
@@ -1039,7 +1048,6 @@ void UMeshSelectionTool::CancelChange()
 
 TUniquePtr<FToolCommandChange> UMeshSelectionTool::EndChange()
 {
-	check(ActiveSelectionChange);
 	if (ActiveSelectionChange != nullptr)
 	{
 		TUniquePtr<FMeshSelectionChange> Result = MoveTemp(ActiveSelectionChange->Change);
diff --git a/Engine/Plugins/Experimental/MeshModelingToolsetExp/Source/MeshModelingToolsExp/Private/SeamSculptTool.cpp b/Engine/Plugins/Experimental/MeshModelingToolsetExp/Source/MeshModelingToolsExp/Private/SeamSculptTool.cpp
--- a/Engine/Plugins/Experimental/MeshModelingToolsetExp/Source/MeshModelingToolsExp/Private/SeamSculptTool.cpp
+++ b/Engine/Plugins/Experimental/MeshModelingToolsetExp/Source/MeshModelingToolsExp/Private/SeamSculptTool.cpp
@@ -151,7 +151,7 @@ void USeamSculptTool::OnBeginDrag(const FRay& Ray)
 {
 	UBaseBrushTool::OnBeginDrag(Ray);
 
-	if (CurrentSnapVertex >= 0)
+	if (IsInBrushStroke() && CurrentSnapVertex >= 0)
 	{
 		CaptureState = EActiveCaptureState::DrawNewPath;
 		DrawPathStartVertex = CurrentSnapVertex;
@@ -188,9 +188,12 @@ void USeamSculptTool::OnUpdateDrag(const FRay& Ray)
 
 void USeamSculptTool::OnEndDrag(const FRay& Ray)
 {
+	// Capture brush stroke state prior to invoking Super::OnEndDrag
+	const bool bWasInBrushStroke = IsInBrushStroke();
+	
 	UBaseBrushTool::OnEndDrag(Ray);
 
-	if (CaptureState == EActiveCaptureState::DrawNewPath)
+	if (bWasInBrushStroke && CaptureState == EActiveCaptureState::DrawNewPath)
 	{
 		CreateSeamAlongPath();
 		CurDrawPath.Reset();
