ModelingTools: Allow Inspector tool work on skeletal meshes.

diff --git a/Engine/Plugins/Experimental/MeshModelingToolsetExp/Source/MeshModelingToolsExp/Private/MeshInspectorTool.cpp b/Engine/Plugins/Experimental/MeshModelingToolsetExp/Source/MeshModelingToolsExp/Private/MeshInspectorTool.cpp
--- a/Engine/Plugins/Experimental/MeshModelingToolsetExp/Source/MeshModelingToolsExp/Private/MeshInspectorTool.cpp
+++ b/Engine/Plugins/Experimental/MeshModelingToolsetExp/Source/MeshModelingToolsExp/Private/MeshInspectorTool.cpp
@@ -26,6 +26,9 @@
 #include "ToolSetupUtil.h"
 #include "AssetUtils/MeshDescriptionUtil.h"
 #include "ModelingToolTargetUtil.h"
+#include "TargetInterfaces/MaterialProvider.h"
+#include "TargetInterfaces/MeshDescriptionProvider.h"
+#include "TargetInterfaces/PrimitiveComponentBackedTarget.h"
 
 #include "CanvasTypes.h"
 #include "CanvasItem.h"
@@ -163,6 +166,15 @@ UMaterialInterface* UMeshInspectorMaterialProperties::GetActiveOverrideMaterial(
  * ToolBuilder
  */
 
+const FToolTargetTypeRequirements& UMeshInspectorToolBuilder::GetTargetRequirements() const
+{
+	static FToolTargetTypeRequirements TypeRequirements({
+		UMeshDescriptionProvider::StaticClass(),
+		UPrimitiveComponentBackedTarget::StaticClass(),
+		UMaterialProvider::StaticClass()
+		});
+	return TypeRequirements;
+}
 
 USingleSelectionMeshEditingTool* UMeshInspectorToolBuilder::CreateNewTool(const FToolBuilderState& SceneState) const
 {
diff --git a/Engine/Plugins/Experimental/MeshModelingToolsetExp/Source/MeshModelingToolsExp/Public/MeshInspectorTool.h b/Engine/Plugins/Experimental/MeshModelingToolsetExp/Source/MeshModelingToolsExp/Public/MeshInspectorTool.h
--- a/Engine/Plugins/Experimental/MeshModelingToolsetExp/Source/MeshModelingToolsExp/Public/MeshInspectorTool.h
+++ b/Engine/Plugins/Experimental/MeshModelingToolsetExp/Source/MeshModelingToolsExp/Public/MeshInspectorTool.h
@@ -32,6 +32,9 @@ class MESHMODELINGTOOLSEXP_API UMeshInspectorToolBuilder : public USingleSelecti
 
 public:
 	virtual USingleSelectionMeshEditingTool* CreateNewTool(const FToolBuilderState& SceneState) const override;
+
+protected:
+	virtual const FToolTargetTypeRequirements& GetTargetRequirements() const override;
 };
 
 
