Skin Weights Paint Tool:  support for LODs & Skin Weight Profiles

diff --git a/Engine/Plugins/Experimental/Animation/SkeletalMeshModelingTools/Source/SkeletalMeshModelingTools/Private/DetailCustomization/SkinWeightDetailCustomization.cpp b/Engine/Plugins/Experimental/Animation/SkeletalMeshModelingTools/Source/SkeletalMeshModelingTools/Private/DetailCustomization/SkinWeightDetailCustomization.cpp
--- a/Engine/Plugins/Experimental/Animation/SkeletalMeshModelingTools/Source/SkeletalMeshModelingTools/Private/DetailCustomization/SkinWeightDetailCustomization.cpp
+++ b/Engine/Plugins/Experimental/Animation/SkeletalMeshModelingTools/Source/SkeletalMeshModelingTools/Private/DetailCustomization/SkinWeightDetailCustomization.cpp
@@ -36,6 +36,10 @@ void FSkinWeightDetailCustomization::CustomizeDetails(IDetailLayoutBuilder& Deta
 	ensure(DetailObjects.Num()==1);
 	SkinToolSettings = Cast<USkinWeightsPaintToolProperties>(DetailObjects[0]);
 
+	// Edit SkinWeightLayer category first 
+	IDetailCategoryBuilder& SkinWeightLayerCategory = DetailBuilder.EditCategory("SkinWeightLayer", FText::GetEmpty(), ECategoryPriority::Important);
+	SkinWeightLayerCategory.InitiallyCollapsed(true);
+	
 	// custom display of falloff mode as segmented toggle buttons
 	IDetailCategoryBuilder& EditModeCategory = DetailBuilder.EditCategory("Weight Editing Mode", FText::GetEmpty(), ECategoryPriority::Important);
 
diff --git a/Engine/Plugins/Experimental/MeshModelingToolsetExp/Source/MeshModelingToolsEditorOnlyExp/Private/SkeletalMesh/SkinWeightsPaintTool.cpp b/Engine/Plugins/Experimental/MeshModelingToolsetExp/Source/MeshModelingToolsEditorOnlyExp/Private/SkeletalMesh/SkinWeightsPaintTool.cpp
--- a/Engine/Plugins/Experimental/MeshModelingToolsetExp/Source/MeshModelingToolsEditorOnlyExp/Private/SkeletalMesh/SkinWeightsPaintTool.cpp
+++ b/Engine/Plugins/Experimental/MeshModelingToolsetExp/Source/MeshModelingToolsEditorOnlyExp/Private/SkeletalMesh/SkinWeightsPaintTool.cpp
@@ -32,6 +32,8 @@
 #include "Editor/Persona/Public/IPersonaEditorModeManager.h"
 #include "Editor/Persona/Public/PersonaModule.h"
 #include "Preferences/PersonaOptions.h"
+#include "PreviewProfileController.h"
+#include "Animation/SkinWeightProfile.h"
 
 #include UE_INLINE_GENERATED_CPP_BY_NAME(SkinWeightsPaintTool)
 
@@ -39,6 +41,73 @@
 
 using namespace SkinPaintTool;
 
+namespace SkinPaintTool
+{
+	
+EMeshLODIdentifier GetLODId(const FName InLODName)
+{
+	static const TMap<FName, EMeshLODIdentifier> LODs({
+		{"LOD0", EMeshLODIdentifier::LOD0},
+		{"LOD1", EMeshLODIdentifier::LOD1},
+		{"LOD2", EMeshLODIdentifier::LOD2},
+		{"LOD3", EMeshLODIdentifier::LOD3},
+		{"LOD4", EMeshLODIdentifier::LOD4},
+		{"LOD5", EMeshLODIdentifier::LOD5},
+		{"LOD6", EMeshLODIdentifier::LOD6},
+		{"LOD7", EMeshLODIdentifier::LOD7},
+		{"HiResSource", EMeshLODIdentifier::HiResSource},
+		{"Default", EMeshLODIdentifier::Default},
+		{"MaxQuality", EMeshLODIdentifier::MaxQuality}
+	});
+
+	const EMeshLODIdentifier* LODIdFound = LODs.Find(InLODName);
+	return LODIdFound ? *LODIdFound : EMeshLODIdentifier::Default;
+}
+
+FName GetLODName(const EMeshLODIdentifier InLOD)
+{
+	static const TMap<EMeshLODIdentifier, FName> LODs({
+		{EMeshLODIdentifier::LOD0, "LOD0"},
+		{EMeshLODIdentifier::LOD1, "LOD1"},
+		{EMeshLODIdentifier::LOD2, "LOD2"},
+		{EMeshLODIdentifier::LOD3, "LOD3"},
+		{EMeshLODIdentifier::LOD4, "LOD4"},
+		{EMeshLODIdentifier::LOD5, "LOD5"},
+		{EMeshLODIdentifier::LOD6, "LOD6"},
+		{EMeshLODIdentifier::LOD7, "LOD7"},
+		{EMeshLODIdentifier::HiResSource, "HiResSource"},
+		{EMeshLODIdentifier::Default, "Default"},
+		{EMeshLODIdentifier::MaxQuality, "MaxQuality"}
+	});
+
+	const FName* Name = LODs.Find(InLOD);
+	return Name ? *Name : NAME_None;
+}
+
+USkeletalMeshComponent* GetSkeletalMeshComponent(const UToolTarget* InTarget)
+{
+	if (!ensure(InTarget))
+	{
+		return nullptr;
+	}
+	
+	const IPrimitiveComponentBackedTarget* TargetComponent = Cast<IPrimitiveComponentBackedTarget>(InTarget);
+	if (!ensure(TargetComponent))
+	{
+		return nullptr;
+	}
+	
+	USkeletalMeshComponent* Component = Cast<USkeletalMeshComponent>(TargetComponent->GetOwnerComponent());
+	if (!ensure(Component && Component->GetSkeletalMeshAsset()))
+	{
+		return nullptr;
+	}
+
+	return Component;
+}
+
+}
+
 // thread pool to use for async operations
 static EAsyncExecution SkinPaintToolAsyncExecTarget = EAsyncExecution::ThreadPool;
 
@@ -99,6 +168,57 @@ USkinWeightsPaintToolProperties::USkinWeightsPaintToolProperties()
 	LoadConfig();
 }
 
+TArray<FName> USkinWeightsPaintToolProperties::GetLODsFunc() const
+{
+	static TArray<FName> Dummy;
+
+	if (!ensure(WeightTool && WeightTool->GetTarget()))
+	{
+		return Dummy;
+	}
+
+	bool bSupportsLODs = false;
+	const TArray<EMeshLODIdentifier> LODIDSs = UE::ToolTarget::GetMeshDescriptionLODs(WeightTool->GetTarget(), bSupportsLODs);
+	if (!ensure(bSupportsLODs))
+	{
+		return Dummy;
+	}
+
+	TArray<FName> LODs;
+	LODs.Reserve(LODIDSs.Num());
+	for (const EMeshLODIdentifier LODId: LODIDSs)
+	{
+		const FName LODName = GetLODName(LODId);
+		if (LODName != NAME_None)
+		{
+			LODs.Add(LODName);
+		}
+	}
+	ensure(!LODs.IsEmpty());
+	
+	return LODs;
+}
+
+TArray<FName> USkinWeightsPaintToolProperties::GetSkinWeightProfilesFunc() const
+{
+	TArray Profiles = {FSkeletalMeshAttributesShared::DefaultSkinWeightProfileName};
+
+	if (const USkeletalMeshComponent* SkeletalMeshComponent = GetSkeletalMeshComponent(WeightTool->GetTarget()))
+	{
+		const USkeletalMesh* SkeletalMesh = SkeletalMeshComponent->GetSkeletalMeshAsset();
+		const int32 LOD = static_cast<int32>(GetLODId(ActiveLOD));
+		for (const FSkinWeightProfileInfo& Profile : SkeletalMesh->GetSkinWeightProfiles())
+		{
+			if (Profile.PerLODSourceFiles.Contains(LOD)) 
+			{
+				Profiles.Add(Profile.Name);
+			}
+		}		
+	}
+	
+	return Profiles;
+}
+
 FSkinWeightBrushConfig& USkinWeightsPaintToolProperties::GetBrushConfig()
 {
 	return *BrushConfigs[BrushMode];
@@ -259,8 +379,8 @@ void FSkinToolDeformer::UpdateVertexDeformation(USkinWeightsPaintTool* Tool)
 		// update vertex acceleration structure
 		{
 			TRACE_CPUPROFILER_EVENT_SCOPE(SkinTool::UpdateVertexOctree);
-			Tool->VerticesOctree.RemoveVertices(VerticesWithModifiedWeights);
-			Tool->VerticesOctree.InsertVertices(VerticesWithModifiedWeights);
+			Tool->VerticesOctree->RemoveVertices(VerticesWithModifiedWeights);
+			Tool->VerticesOctree->InsertVertices(VerticesWithModifiedWeights);
 		}
 		
 		// update triangle acceleration structure
@@ -293,7 +413,7 @@ void FSkinToolDeformer::UpdateVertexDeformation(USkinWeightsPaintTool* Tool)
 			// asynchronously update the octree, this normally finishes well before the next update
 			// but in the unlikely event that it does not, it would result in a frame where the paint brush
 			// is not perfectly aligned with the mesh; not a deal breaker.
-			UE::Geometry::FDynamicMeshOctree3& OctreeToUpdate = Tool->TrianglesOctree;
+			UE::Geometry::FDynamicMeshOctree3& OctreeToUpdate = *Tool->TrianglesOctree;
 			Tool->TriangleOctreeFuture = Async(SkinPaintToolAsyncExecTarget, [&OctreeToUpdate, &AffectedTriangles]()
 			{
 				TRACE_CPUPROFILER_EVENT_SCOPE(SkinTool::TriangleOctreeReinsert);	
@@ -332,7 +452,7 @@ void FSkinToolWeights::InitializeSkinWeights(
 
 	// initialize current weights (using compact format: num_verts * max_influences)
 	const FSkeletalMeshConstAttributes MeshAttribs(*Mesh);
-	const FSkinWeightsVertexAttributesConstRef VertexSkinWeights = MeshAttribs.GetVertexSkinWeights();
+	const FSkinWeightsVertexAttributesConstRef VertexSkinWeights = MeshAttribs.GetVertexSkinWeights(Profile);
 	const int32 NumVertices = Mesh->Vertices().Num();
 	CurrentWeights.SetNum(NumVertices);
 	for (int32 VertexIndex = 0; VertexIndex < NumVertices; VertexIndex++)
@@ -503,7 +623,7 @@ void FSkinToolWeights::EditVertexWeightAndNormalize(
 void FSkinToolWeights::ApplyCurrentWeightsToMeshDescription(FMeshDescription* EditedMesh)
 {
 	FSkeletalMeshAttributes MeshAttribs(*EditedMesh);
-	FSkinWeightsVertexAttributesRef VertexSkinWeights = MeshAttribs.GetVertexSkinWeights();
+	FSkinWeightsVertexAttributesRef VertexSkinWeights = MeshAttribs.GetVertexSkinWeights(Profile);
 	
 	UE::AnimationCore::FBoneWeightsSettings Settings;
 	Settings.SetNormalizeType(UE::AnimationCore::EBoneWeightNormalizeType::None);
@@ -674,7 +794,8 @@ void FSkinToolWeights::UpdateIsBoneWeighted(BoneIndex BoneToUpdate)
 void FMeshSkinWeightsChange::Apply(UObject* Object)
 {
 	USkinWeightsPaintTool* Tool = CastChecked<USkinWeightsPaintTool>(Object);
-	
+
+	Tool->ExternalUpdateSkinWeightLayer(LOD, SkinWeightProfile);
 	for (TTuple<int32, FSingleBoneWeightEdits>& Pair : AllWeightEdits.PerBoneWeightEdits)
 	{
 		Tool->ExternalUpdateWeights(Pair.Key, Pair.Value.NewWeights);
@@ -685,6 +806,7 @@ void FMeshSkinWeightsChange::Revert(UObject* Object)
 {
 	USkinWeightsPaintTool* Tool = CastChecked<USkinWeightsPaintTool>(Object);
 
+	Tool->ExternalUpdateSkinWeightLayer(LOD, SkinWeightProfile);
 	for (TTuple<int32, FSingleBoneWeightEdits>& Pair : AllWeightEdits.PerBoneWeightEdits)
 	{
 		Tool->ExternalUpdateWeights(Pair.Key, Pair.Value.OldWeights);
@@ -730,7 +852,11 @@ void USkinWeightsPaintTool::Setup()
 	check(Component && Component->GetSkeletalMeshAsset())
 
 	// create a mesh description for editing (this must be done before calling UpdateBonePositionInfos) 
-	EditedMesh = MakeUnique<FMeshDescription>();
+	bool bSupportsLODs = false;
+	const EMeshLODIdentifier DefaultLOD = UE::ToolTarget::GetTargetMeshDescriptionLOD(Target, bSupportsLODs);
+	ensure(bSupportsLODs);
+
+	EditedMesh = &EditedMeshes.Emplace(DefaultLOD);
 	*EditedMesh = *UE::ToolTarget::GetMeshDescription(Target);
 
 	// create a custom set of properties inheriting from the base tool properties
@@ -738,6 +864,14 @@ void USkinWeightsPaintTool::Setup()
 	WeightToolProperties->RestoreProperties(this);
 	WeightToolProperties->WeightTool = this;
 	WeightToolProperties->bSpecifyRadius = true;
+	// watch for skin weight layer changes
+	WeightToolProperties->ActiveLOD = "LOD0";
+	int32 WatcherIndex = WeightToolProperties->WatchProperty(WeightToolProperties->ActiveLOD, [this](FName) { OnActiveLODChanged(); });
+	WeightToolProperties->SilentUpdateWatcherAtIndex(WatcherIndex);
+	WeightToolProperties->ActiveSkinWeightProfile = FSkeletalMeshAttributesShared::DefaultSkinWeightProfileName;
+	WatcherIndex = WeightToolProperties->WatchProperty(WeightToolProperties->ActiveSkinWeightProfile, [this](FName) { OnActiveSkinWeightProfileChanged(); });
+	WeightToolProperties->SilentUpdateWatcherAtIndex(WatcherIndex);
+	
 	// replace the base brush properties
 	ReplaceToolPropertySource(BrushProperties, WeightToolProperties);
 	BrushProperties = WeightToolProperties;
@@ -773,22 +907,11 @@ void USkinWeightsPaintTool::Setup()
 	bBoneColorsToRestore = GetDefault<UPersonaOptions>()->bShowBoneColors;
 	GetMutableDefault<UPersonaOptions>()->bShowBoneColors = true;
 
-	// build octree for vertices
-	VerticesOctree.Initialize(PreviewMesh->GetMesh(), true);
-
-	// build octree for triangles
-	{
-		TRACE_CPUPROFILER_EVENT_SCOPE(SkinTool::InitTriangleOctree);
-		
-		TriangleOctreeFuture = Async(SkinPaintToolAsyncExecTarget, [&]()
-		{
-			TRACE_CPUPROFILER_EVENT_SCOPE(SkinTool::InitTriangleOctreeRun);
-			TrianglesOctree.Initialize(PreviewMesh->GetMesh());
-		});
-	}
+	// initialize vertices & triangle octrees (this must be done after PreviewMesh has been initialized)
+	InitializeOctrees();
 
 	// initialize weight maps and deformation data
-	Weights.InitializeSkinWeights(Component, EditedMesh.Get());
+	Weights.InitializeSkinWeights(Component, EditedMesh);
 	bVisibleWeightsValid = false;
 
 	RecalculateBrushRadius();
@@ -832,19 +955,10 @@ void USkinWeightsPaintTool::Setup()
 	PolygonSelectionMechanic->PolyEdgesRenderer.PointSize = 5.0f;
 	PolygonSelectionMechanic->PolyEdgesRenderer.DepthBias = 2.0f;
 	PolygonSelectionMechanic->PolyEdgesRenderer.LineThickness = 1.0f;
+
+	// initialize the polygon selection mechanic (this must be done after PreviewMesh has been initialized)
+	InitializeSelectionMechanic();
 	
-	// initialize the polygon selection mechanic
-	constexpr bool bAutoBuild = true;
-	const FDynamicMesh3* DynamicMesh = PreviewMesh->GetPreviewDynamicMesh();
-	SelectionTopology = MakeUnique<UE::Geometry::FTriangleGroupTopology>(DynamicMesh, bAutoBuild);
-	MeshSpatial = MakeUnique<UE::Geometry::FDynamicMeshAABBTree3>(DynamicMesh, bAutoBuild);
-	PolygonSelectionMechanic->Initialize(
-		DynamicMesh,
-		FTransform::Identity,
-		Component->GetWorld(),
-		SelectionTopology.Get(),
-		[this]() { return MeshSpatial.Get(); }
-	);
 	// secondary triangle buffer used to render face selection
 	PreviewMesh->EnableSecondaryTriangleBuffers([this](const FDynamicMesh3* Mesh, int32 TriangleID)
 	{
@@ -860,9 +974,8 @@ void USkinWeightsPaintTool::Setup()
 		PreviewMesh->FastNotifySecondaryTrianglesChanged();
 	});
 	
-	SmoothWeightsDataSource = MakeUnique<FPaintToolWeightsDataSource>(&Weights);
-	SmoothWeightsOp = MakeUnique<UE::Geometry::TSmoothBoneWeights<int32, float>>(PreviewMesh->GetMesh(), SmoothWeightsDataSource.Get());
-	SmoothWeightsOp->MinimumWeightThreshold = MinimumWeightThreshold;
+	// initialize smooth operator (this must be done after PreviewMesh & Weights have been initialized)
+	InitializeSmoothWeightsOperator();
 
 	if (EditorContext.IsValid())
 	{
@@ -1093,7 +1206,7 @@ bool USkinWeightsPaintTool::HitTest(const FRay& Ray, FHitResult& OutHit)
 	FViewCameraState StateOut;
 	GetToolManager()->GetContextQueriesAPI()->GetCurrentViewState(StateOut);
 	FVector3d LocalEyePosition(CurTargetTransform.InverseTransformPosition((FVector3d)StateOut.Position));
-	const int32 TriID = TrianglesOctree.FindNearestHitObject(
+	const int32 TriID = TrianglesOctree->FindNearestHitObject(
 		LocalRay,
 		[this, Mesh, &LocalEyePosition](int TriangleID)
 	{
@@ -1208,7 +1321,7 @@ void USkinWeightsPaintTool::CalculateVertexROI(
 		const float RadiusSqr = CurrentBrushRadius * CurrentBrushRadius;
 		const FDynamicMesh3* Mesh = PreviewMesh->GetPreviewDynamicMesh();
 		const FAxisAlignedBox3d QueryBox(StampPosLocal, CurrentBrushRadius);
-		VerticesOctree.RangeQuery(QueryBox,
+		VerticesOctree->RangeQuery(QueryBox,
 			[&](int32 VertexID) { return FVector3d::DistSquared(Mesh->GetVertex(VertexID), StampPosLocal) < RadiusSqr; },
 			VertexIDs);
 
@@ -1607,6 +1720,69 @@ void USkinWeightsPaintTool::RelaxWeightOnVertices(
 	}
 }
 
+void USkinWeightsPaintTool::InitializeOctrees()
+{
+	if (!ensure(PreviewMesh && PreviewMesh->GetMesh()))
+	{
+		return;
+	}
+
+	// build octree for vertices
+	VerticesOctree = MakeUnique<DynamicVerticesOctree>();
+	VerticesOctree->Initialize(PreviewMesh->GetMesh(), true);
+
+	// build octree for triangles
+	TrianglesOctree = MakeUnique<DynamicTrianglesOctree>();
+	{
+		TRACE_CPUPROFILER_EVENT_SCOPE(SkinTool::InitTriangleOctree);
+		
+		TriangleOctreeFuture = Async(SkinPaintToolAsyncExecTarget, [&]()
+		{
+			TRACE_CPUPROFILER_EVENT_SCOPE(SkinTool::InitTriangleOctreeRun);
+			TrianglesOctree->Initialize(PreviewMesh->GetMesh());
+		});
+	}
+}
+
+void USkinWeightsPaintTool::InitializeSelectionMechanic()
+{
+	const USkeletalMeshComponent* Component = GetSkeletalMeshComponent(Target);
+	if (!ensure(Component))
+	{
+		return;
+	}
+	
+	if (!ensure(PreviewMesh && PreviewMesh->GetMesh()))
+	{
+		return;
+	}
+
+	static constexpr bool bAutoBuild = true;
+	const FDynamicMesh3* DynamicMesh = PreviewMesh->GetMesh();
+	SelectionTopology = MakeUnique<UE::Geometry::FTriangleGroupTopology>(DynamicMesh, bAutoBuild);
+	MeshSpatial = MakeUnique<FDynamicMeshAABBTree3>(DynamicMesh, bAutoBuild);
+	PolygonSelectionMechanic->Initialize(
+		DynamicMesh,
+		FTransform::Identity,
+		Component->GetWorld(),
+		SelectionTopology.Get(),
+		[this]() { return MeshSpatial.Get(); }
+	);
+}
+
+void USkinWeightsPaintTool::InitializeSmoothWeightsOperator()
+{
+	if (!ensure(PreviewMesh && PreviewMesh->GetMesh()))
+	{
+		return;
+	}
+
+	// NOTE: this could probably be initialized lazily as it's only used with the relax brush 
+	SmoothWeightsDataSource = MakeUnique<FPaintToolWeightsDataSource>(&Weights);
+	SmoothWeightsOp = MakeUnique<UE::Geometry::TSmoothBoneWeights<int32, float>>(PreviewMesh->GetMesh(), SmoothWeightsDataSource.Get());
+	SmoothWeightsOp->MinimumWeightThreshold = MinimumWeightThreshold;
+}
+
 void USkinWeightsPaintTool::ApplyWeightEditsToMesh(
 	const SkinPaintTool::FMultiBoneWeightEdits& WeightEdits,
 	const FText& TransactionLabel,
@@ -1717,11 +1893,16 @@ void USkinWeightsPaintTool::OnShutdown(EToolShutdownType ShutdownType)
 	if (ShutdownType == EToolShutdownType::Accept)
 	{
 		// apply the weights to the mesh description
-		Weights.ApplyCurrentWeightsToMeshDescription(EditedMesh.Get());
+		Weights.ApplyCurrentWeightsToMeshDescription(EditedMesh);
 
 		// this block bakes the modified DynamicMeshComponent back into the StaticMeshComponent inside an undo transaction
 		GetToolManager()->BeginUndoTransaction(LOCTEXT("SkinWeightsPaintTool", "Paint Skin Weights"));
-		UE::ToolTarget::CommitMeshDescriptionUpdate(Target, EditedMesh.Get());
+		static constexpr bool bHaveTargetLOD = true;
+		for (auto& [LOD, MeshDescription] : EditedMeshes)
+		{
+			const FCommitMeshParameters CommitParams(bHaveTargetLOD, LOD);
+			UE::ToolTarget::CommitMeshDescriptionUpdate(Target, &MeshDescription, nullptr, CommitParams);	
+		}
 		GetToolManager()->EndUndoTransaction();
 	}
 
@@ -1748,7 +1929,9 @@ void USkinWeightsPaintTool::OnShutdown(EToolShutdownType ShutdownType)
 
 void USkinWeightsPaintTool::BeginChange()
 {
-	ActiveChange = MakeUnique<FMeshSkinWeightsChange>();
+	const EMeshLODIdentifier LOD = GetLODId(WeightToolProperties->ActiveLOD);
+	const FName SkinProfile = WeightToolProperties->ActiveSkinWeightProfile;
+	ActiveChange = MakeUnique<FMeshSkinWeightsChange>(LOD, SkinProfile);
 }
 
 void USkinWeightsPaintTool::EndChange(const FText& TransactionLabel)
@@ -1784,6 +1967,40 @@ void USkinWeightsPaintTool::ExternalUpdateWeights(const int32 BoneIndex, const T
 	Weights.UpdateIsBoneWeighted(BoneIndex);
 }
 
+void USkinWeightsPaintTool::ExternalUpdateSkinWeightLayer(const EMeshLODIdentifier InLOD, const FName InSkinWeightProfile)
+{
+	enum class ESkinWeightChangeState
+	{
+		SkinProfile,
+		LOD,
+		None
+	} State = ESkinWeightChangeState::None;
+	
+	if (InSkinWeightProfile != WeightToolProperties->ActiveSkinWeightProfile)
+	{
+		WeightToolProperties->ActiveSkinWeightProfile = InSkinWeightProfile;
+		State = ESkinWeightChangeState::SkinProfile;
+	}
+
+	const FName LODName = GetLODName(InLOD);
+	if (LODName != NAME_None && LODName != WeightToolProperties->ActiveLOD)
+	{
+		WeightToolProperties->ActiveLOD = LODName;
+		State = ESkinWeightChangeState::LOD;
+	}
+
+	switch (State)
+	{
+	case ESkinWeightChangeState::SkinProfile:
+		return OnActiveSkinWeightProfileChanged();
+	case ESkinWeightChangeState::LOD:
+		return OnActiveLODChanged();
+	case ESkinWeightChangeState::None:
+		default:
+		break;
+	}
+}
+
 void FSkinMirrorData::RegenerateMirrorData(
     const TArray<FName>& BoneNames,
     const TMap<FName, BoneIndex>& BoneNameToIndexMap,
@@ -2211,6 +2428,115 @@ void USkinWeightsPaintTool::HandleSkeletalMeshModified(const TArray<FName>& InBo
 	}
 }
 
+void USkinWeightsPaintTool::OnActiveLODChanged()
+{
+	const USkeletalMeshComponent* Component = GetSkeletalMeshComponent(Target);
+	if (!ensure(Component))
+	{
+		return;
+	}
+
+	// apply previous changes
+	Weights.ApplyCurrentWeightsToMeshDescription(EditedMesh);
+
+	// update EditedMesh using the new LOD
+	const EMeshLODIdentifier LODId = GetLODId(WeightToolProperties->ActiveLOD);
+	const FGetMeshParameters Params(true, LODId);
+	EditedMesh = EditedMeshes.Find(LODId);
+	if (!EditedMesh)
+	{
+		EditedMesh = &EditedMeshes.Emplace(LODId);
+		*EditedMesh = *UE::ToolTarget::GetMeshDescription(Target, Params);
+	}
+
+	// update the preview mesh
+	PreviewMesh->ReplaceMesh(UE::ToolTarget::GetDynamicMeshCopy(Target, Params));
+	PreviewMesh->EditMesh([](FDynamicMesh3& Mesh)
+	{
+		Mesh.EnableAttributes();
+		Mesh.Attributes()->DisablePrimaryColors();
+		Mesh.Attributes()->EnablePrimaryColors();
+		Mesh.Attributes()->PrimaryColors()->CreateFromPredicate([](int ParentVID, int TriIDA, int TriIDB){return true;}, 0.f);
+	});
+
+	// update vertices & triangle octrees (this must be done after PreviewMesh has been updated)
+	InitializeOctrees();
+
+	// update weights
+	Weights = FSkinToolWeights();
+	if (!IsProfileValid(WeightToolProperties->ActiveSkinWeightProfile))
+	{
+		WeightToolProperties->ActiveSkinWeightProfile = FSkeletalMeshAttributesShared::DefaultSkinWeightProfileName;
+	}
+	Weights.Profile = WeightToolProperties->ActiveSkinWeightProfile;
+	Weights.InitializeSkinWeights(Component, EditedMesh);
+	bVisibleWeightsValid = false;
+
+	// update the polygon selection mechanic (this must be done after PreviewMesh has been updated)
+	InitializeSelectionMechanic();
+
+	// update smooth operator (this must be done after PreviewMesh & Weights have been updated)
+	InitializeSmoothWeightsOperator();
+}
+
+void USkinWeightsPaintTool::OnActiveSkinWeightProfileChanged()
+{
+	const USkeletalMeshComponent* SkeletalMeshComponent = GetSkeletalMeshComponent(Target);
+	if (!SkeletalMeshComponent)
+	{
+		return;
+	}
+
+	if (!IsProfileValid(WeightToolProperties->ActiveSkinWeightProfile))
+	{
+		WeightToolProperties->ActiveSkinWeightProfile = FSkeletalMeshAttributesShared::DefaultSkinWeightProfileName;
+	}
+	
+	if (WeightToolProperties->ActiveSkinWeightProfile == Weights.Profile)
+	{
+		return;
+	}
+
+	// apply previous changes
+	Weights.ApplyCurrentWeightsToMeshDescription(EditedMesh);
+
+	// re-init Weights with new skin profile
+	Weights = FSkinToolWeights();
+	Weights.Profile = WeightToolProperties->ActiveSkinWeightProfile;
+	Weights.InitializeSkinWeights(SkeletalMeshComponent, EditedMesh);
+	bVisibleWeightsValid = false;
+}
+
+bool USkinWeightsPaintTool::IsProfileValid(const FName InProfileName) const
+{
+	const USkeletalMeshComponent* SkeletalMeshComponent = GetSkeletalMeshComponent(Target);
+	if (!SkeletalMeshComponent)
+	{
+		return false;
+	}
+
+	// check SkeletalMesh
+	const TArray<FSkinWeightProfileInfo>& SkeletalMeshProfiles = SkeletalMeshComponent->GetSkeletalMeshAsset()->GetSkinWeightProfiles();
+	bool bHasProfile = SkeletalMeshProfiles.ContainsByPredicate([InProfileName](const FSkinWeightProfileInfo& Info)
+	{
+		return Info.Name == InProfileName;
+	});
+	if (!bHasProfile)
+	{
+		return false;
+	}
+
+	// check current MeshDescription
+	const FSkeletalMeshConstAttributes MeshAttribs(*EditedMesh);
+	const TArray<FName> MeshDescProfiles = MeshAttribs.GetSkinWeightProfileNames();
+	bHasProfile = MeshDescProfiles.ContainsByPredicate([InProfileName](const FName Name)
+	{
+		return Name == InProfileName;
+	});
+	
+	return bHasProfile;
+}
+
 void USkinWeightsPaintTool::ToggleEditingMode()
 {
 	Weights.Deformer.SetAllVerticesToBeUpdated();
diff --git a/Engine/Plugins/Experimental/MeshModelingToolsetExp/Source/MeshModelingToolsEditorOnlyExp/Public/SkeletalMesh/SkinWeightsPaintTool.h b/Engine/Plugins/Experimental/MeshModelingToolsetExp/Source/MeshModelingToolsEditorOnlyExp/Public/SkeletalMesh/SkinWeightsPaintTool.h
--- a/Engine/Plugins/Experimental/MeshModelingToolsetExp/Source/MeshModelingToolsEditorOnlyExp/Public/SkeletalMesh/SkinWeightsPaintTool.h
+++ b/Engine/Plugins/Experimental/MeshModelingToolsetExp/Source/MeshModelingToolsEditorOnlyExp/Public/SkeletalMesh/SkinWeightsPaintTool.h
@@ -2,19 +2,19 @@
 
 #pragma once
 
-#include "AssetViewerSettings.h"
 #include "DynamicMeshBrushTool.h"
 #include "BaseTools/MeshSurfacePointMeshEditingTool.h"
 #include "MeshDescription.h"
 #include "DynamicMesh/DynamicVerticesOctree3.h"
 #include "BoneWeights.h"
 #include "GroupTopology.h"
-#include "PreviewProfileController.h"
+#include "SkeletalMeshAttributes.h"
 #include "Misc/Optional.h"
 #include "Containers/Map.h"
 #include "DynamicMesh/DynamicMeshOctree3.h"
 #include "SkeletalMesh/SkeletalMeshEditionInterface.h"
 #include "Engine/SkeletalMesh.h"
+#include "TargetInterfaces/MeshTargetInterfaceTypes.h"
 
 #include "SkinWeightsPaintTool.generated.h"
 
@@ -153,7 +153,11 @@ namespace SkinPaintTool
 	class FMeshSkinWeightsChange : public FToolCommandChange
 	{
 	public:
-		FMeshSkinWeightsChange() : FToolCommandChange()	{ }
+		FMeshSkinWeightsChange(const EMeshLODIdentifier InLOD, const FName InSkinWeightProfile)
+			: FToolCommandChange()
+			, LOD(InLOD)
+			, SkinWeightProfile(InSkinWeightProfile)
+		{}
 
 		virtual FString ToString() const override
 		{
@@ -168,6 +172,8 @@ namespace SkinPaintTool
 
 	private:
 		FMultiBoneWeightEdits AllWeightEdits;
+		EMeshLODIdentifier LOD = EMeshLODIdentifier::Default;
+		FName SkinWeightProfile = FSkeletalMeshAttributesShared::DefaultSkinWeightProfileName;
 	};
 
 	// intermediate storage of the weight maps for duration of tool
@@ -224,6 +230,9 @@ namespace SkinPaintTool
 
 		// update deformation when vertex weights are modified
 		FSkinToolDeformer Deformer;
+
+		// which skin profile is currently edited
+		FName Profile = FSkeletalMeshAttributesShared::DefaultSkinWeightProfileName;
 	};
 
 	struct FSkinMirrorData
@@ -335,6 +344,12 @@ public:
 	UPROPERTY(Config)
 	FSkinWeightBrushConfig BrushConfigRelax;
 
+	// skin weight layer properties
+	UPROPERTY(EditAnywhere, Category = SkinWeightLayer, meta = (DisplayName = "Active LOD", GetOptions = GetLODsFunc))
+	FName ActiveLOD = "LOD0";
+	UPROPERTY(EditAnywhere, Category = SkinWeightLayer, meta = (DisplayName = "Active Profile", GetOptions = GetSkinWeightProfilesFunc))
+	FName ActiveSkinWeightProfile = FSkeletalMeshAttributesShared::DefaultSkinWeightProfileName;
+	
 	// pointer back to paint tool
 	TObjectPtr<USkinWeightsPaintTool> WeightTool;
 
@@ -342,6 +357,13 @@ public:
 	void SetFalloffMode(EWeightBrushFalloffMode InFalloffMode);
 	void SetColorMode(EWeightColorMode InColorMode);
 	void SetBrushMode(EWeightEditOperation InBrushMode);
+
+private:
+	
+	UFUNCTION()
+	TArray<FName> GetLODsFunc() const;
+	UFUNCTION()
+	TArray<FName> GetSkinWeightProfilesFunc() const;
 };
 
 // An interactive tool for painting and editing skin weights.
@@ -379,6 +401,7 @@ public:
 
 	// using when ToolChange is applied via Undo/Redo
 	void ExternalUpdateWeights(const int32 BoneIndex, const TMap<int32, float>& IndexValues);
+	void ExternalUpdateSkinWeightLayer(const EMeshLODIdentifier InLOD, const FName InSkinWeightProfile);
 
 	// weight editing operations (selection based)
 	void MirrorWeights(EAxis::Type Axis, EMirrorDirection Direction);
@@ -484,18 +507,22 @@ protected:
 		SkinPaintTool::FMultiBoneWeightEdits& InOutWeightEdits);
 
 	// used to accelerate mesh queries
-	UE::Geometry::TDynamicVerticesOctree3<FDynamicMesh3> VerticesOctree;
-	UE::Geometry::FDynamicMeshOctree3 TrianglesOctree;
+	using DynamicVerticesOctree = UE::Geometry::TDynamicVerticesOctree3<FDynamicMesh3>;
+	TUniquePtr<DynamicVerticesOctree> VerticesOctree;
+	using DynamicTrianglesOctree = UE::Geometry::FDynamicMeshOctree3;
+	TUniquePtr<DynamicTrianglesOctree> TrianglesOctree;
 	TFuture<void> TriangleOctreeFuture;
 	TArray<int32> TrianglesToReinsert;
+	void InitializeOctrees();
 
 	// tool properties
 	UPROPERTY()
 	TObjectPtr<USkinWeightsPaintToolProperties> WeightToolProperties;
 	virtual void OnPropertyModified(UObject* ModifiedObject, FProperty* ModifiedProperty) override;
 	
-	// the currently edited mesh description
-	TUniquePtr<FMeshDescription> EditedMesh;
+	// the currently edited mesh descriptions
+	mutable TMap<EMeshLODIdentifier, FMeshDescription> EditedMeshes;
+	FMeshDescription* EditedMesh = nullptr;
 
 	// storage of vertex weights per bone 
 	SkinPaintTool::FSkinToolWeights Weights;
@@ -509,6 +536,7 @@ protected:
 	// Smooth weights data source and operator
 	TUniquePtr<UE::Geometry::TBoneWeightsDataSource<int32, float>> SmoothWeightsDataSource;
 	TUniquePtr<UE::Geometry::TSmoothBoneWeights<int32, float>> SmoothWeightsOp;
+	void InitializeSmoothWeightsOperator();
 
 	// vertex colors updated when switching current bone or editing weights
 	void UpdateCurrentBoneVertexColors();
@@ -535,6 +563,12 @@ protected:
 	TObjectPtr<UPolygonSelectionMechanic> PolygonSelectionMechanic;
 	TUniquePtr<UE::Geometry::FDynamicMeshAABBTree3> MeshSpatial = nullptr;
 	TUniquePtr<UE::Geometry::FTriangleGroupTopology> SelectionTopology = nullptr;
+	void InitializeSelectionMechanic();
+
+	// skin weight layer
+	void OnActiveLODChanged();
+	void OnActiveSkinWeightProfileChanged();
+	bool IsProfileValid(const FName InProfileName) const;
 
 	UPROPERTY()
 	TWeakObjectPtr<USkeletalMeshEditorContextObjectBase> EditorContext = nullptr;
