Split the spline management logic of the triangulate splines and revolve splines tools into a base class to bring them to feature parity and make it easier to add new spline-based tools

diff --git a/Engine/Plugins/Experimental/MeshModelingToolsetExp/Source/MeshModelingToolsExp/Private/RevolveSplineTool.cpp b/Engine/Plugins/Experimental/MeshModelingToolsetExp/Source/MeshModelingToolsExp/Private/RevolveSplineTool.cpp
--- a/Engine/Plugins/Experimental/MeshModelingToolsetExp/Source/MeshModelingToolsExp/Private/RevolveSplineTool.cpp
+++ b/Engine/Plugins/Experimental/MeshModelingToolsetExp/Source/MeshModelingToolsExp/Private/RevolveSplineTool.cpp
@@ -34,35 +34,11 @@ void URevolveSplineTool::Setup()
 	ToolActions->Initialize(this);
 	AddToolPropertySource(ToolActions);
 
-	MaterialProperties = NewObject<UNewMeshMaterialProperties>(this);
-	AddToolPropertySource(MaterialProperties);
-	MaterialProperties->RestoreProperties(this);
-
-	OutputTypeProperties = NewObject<UCreateMeshObjectTypeProperties>(this);
-	OutputTypeProperties->InitializeDefault();
-	OutputTypeProperties->RestoreProperties(this);
-	OutputTypeProperties->WatchProperty(OutputTypeProperties->OutputType, [this](FString) { OutputTypeProperties->UpdatePropertyVisibility(); });
-	AddToolPropertySource(OutputTypeProperties);
-
 	SetToolDisplayName(LOCTEXT("RevolveSplineToolName", "Revolve Spline"));
 	GetToolManager()->DisplayMessage(
 		LOCTEXT("RevolveSplineToolDescription", "Revolve the selected spline to create a mesh."),
 		EToolMessageLevel::UserNotification);
 
-	Preview = NewObject<UMeshOpPreviewWithBackgroundCompute>(this);
-	Preview->Setup(GetTargetWorld(), this);
-	Preview->PreviewMesh->EnableWireframe(MaterialProperties->bShowWireframe);
-	Preview->ConfigureMaterials(MaterialProperties->Material.Get(),
-		ToolSetupUtil::GetDefaultWorkingMaterial(GetToolManager()));
-	ToolSetupUtil::ApplyRenderingConfigurationToPreview(Preview->PreviewMesh, nullptr);
-
-	Preview->OnMeshUpdated.AddLambda(
-		[this](const UMeshOpPreviewWithBackgroundCompute* UpdatedPreview)
-		{
-			UpdateAcceptWarnings(UpdatedPreview->HaveEmptyResult() ? EAcceptWarning::EmptyForbidden : EAcceptWarning::NoWarning);
-		}
-	);
-
 	// TODO: We'll probably want a click behavior someday for clicking on the spline to align to a tangent at a point
 
 	// The plane mechanic is used for the revolution axis.
@@ -88,8 +64,6 @@ void URevolveSplineTool::Setup()
 	// don't have a way to pass in a custom alignment raycast, which we would want in order to snap
 	// and align to spline points.
 
-	PollSplineUpdates();
-
 	if (Settings->bResetAxisOnStart)
 	{
 		ResetAxis();
@@ -98,11 +72,14 @@ void URevolveSplineTool::Setup()
 	{
 		UpdateRevolutionAxis();
 	}
+
+	Super::Setup();
 }
 
 void URevolveSplineTool::ResetAxis()
 {
-	if (!Spline.IsValid())
+	USplineComponent* Spline = GetFirstSpline();
+	if (!Spline)
 	{
 		return;
 	}
@@ -131,50 +108,10 @@ void URevolveSplineTool::ResetAxis()
 	UpdateRevolutionAxis();
 }
 
-void URevolveSplineTool::PollSplineUpdates()
-{
-	if (bLostInputSpline)
-	{
-		return;
-	}
-
-	if (!Spline.IsValid())
-	{
-		// Attempt to recapture the spline.
-		TArray<USplineComponent*> Splines;
-		if (SplineOwningActor.IsValid())
-		{
-			SplineOwningActor->GetComponents<USplineComponent>(Splines);
-		}
-		if (SplineComponentIndex < Splines.Num() && ensure(SplineComponentIndex >= 0))
-		{
-			Spline = Splines[SplineComponentIndex];
-		}
-		else
-		{
-			GetToolManager()->DisplayMessage(
-				LOCTEXT("LostSpline", "Tool lost reference to the input spline. The spline input will no longer be updated."),
-				EToolMessageLevel::UserWarning);
-			bLostInputSpline = true;
-			return;
-		}
-	}
-
-	if (Spline->SplineCurves.Version == LastSplineVersion && !bForceSplineUpdate)
-	{
-		// No update necessary
-		return;
-	}
-
-	UpdatePointsFromSpline();
-
-	LastSplineVersion = Spline->SplineCurves.Version;
-	bForceSplineUpdate = false;
-}
-
-void URevolveSplineTool::UpdatePointsFromSpline()
+void URevolveSplineTool::OnSplineUpdate()
 {
-	if (!Spline.IsValid())
+	USplineComponent* Spline = GetFirstSpline();
+	if (!Spline)
 	{
 		return;
 	}
@@ -226,17 +163,12 @@ void URevolveSplineTool::UpdatePointsFromSpline()
 
 void URevolveSplineTool::OnTick(float DeltaTime)
 {
-	PollSplineUpdates();
+	Super::OnTick(DeltaTime);
 
 	if (PlaneMechanic)
 	{
 		PlaneMechanic->Tick(DeltaTime);
 	}
-
-	if (Preview)
-	{
-		Preview->Tick(DeltaTime);
-	}
 }
 
 void URevolveSplineTool::Render(IToolsContextRenderAPI* RenderAPI)
@@ -268,97 +200,45 @@ void URevolveSplineTool::Render(IToolsContextRenderAPI* RenderAPI)
 
 void URevolveSplineTool::OnPropertyModified(UObject* PropertySet, FProperty* Property)
 {
-	if (PropertySet == OutputTypeProperties)
-	{
-		return;
-	}
-
 	if (Property)
 	{
-		if (Property->GetFName() == GET_MEMBER_NAME_CHECKED(UNewMeshMaterialProperties, Material))
-		{
-			Preview->ConfigureMaterials(MaterialProperties->Material.Get(),
-				ToolSetupUtil::GetDefaultWorkingMaterial(GetToolManager()));
-		}
-		else if (Property->GetFName() == GET_MEMBER_NAME_CHECKED(URevolveSplineToolProperties, SampleMode)
+		if (Property->GetFName() == GET_MEMBER_NAME_CHECKED(URevolveSplineToolProperties, SampleMode)
 			|| Property->GetFName() == GET_MEMBER_NAME_CHECKED(URevolveSplineToolProperties, ErrorTolerance)
 			|| Property->GetFName() == GET_MEMBER_NAME_CHECKED(URevolveSplineToolProperties, MaxSampleDistance))
 		{
-			UpdatePointsFromSpline();
+			OnSplineUpdate();
 		}
 
 		// Checking the name for these settings doesn't work, since the reported names are the low level components, like "X" or "Y"
 		// So we'll simply update the axis whenever any property changes. It's overkill but probably not too bad.
 		PlaneMechanic->SetPlaneWithoutBroadcast(FFrame3d(Settings->AxisOrigin,
 			FRotator(Settings->AxisOrientation.X, Settings->AxisOrientation.Y, 0).Quaternion()));
-		UpdateRevolutionAxis();		
+		UpdateRevolutionAxis();
 	}
 
-	Preview->PreviewMesh->EnableWireframe(MaterialProperties->bShowWireframe);
-
-	Preview->InvalidateResult();
+	Super::OnPropertyModified(PropertySet, Property);
 }
 
 void URevolveSplineTool::Shutdown(EToolShutdownType ShutdownType)
 {
 	Settings->SaveProperties(this);
-	OutputTypeProperties->SaveProperties(this);
-	MaterialProperties->SaveProperties(this);
-
-	FDynamicMeshOpResult Result = Preview->Shutdown();
-
-	if (ShutdownType == EToolShutdownType::Accept)
-	{
-		GetToolManager()->BeginUndoTransaction(LOCTEXT("RevolveSplineAction", "Revolve Spline"));
 
-		// Generate the result asset
-		GenerateAsset(Result);
-
-		GetToolManager()->EndUndoTransaction();
-	}
+	Super::Shutdown(ShutdownType);
 
 	PlaneMechanic->Shutdown();
 
-	Preview = nullptr;
 	Settings = nullptr;
-	MaterialProperties = nullptr;
-	OutputTypeProperties = nullptr;
 	ToolActions = nullptr;
 	PlaneMechanic = nullptr;
-
-	Super::Shutdown(ShutdownType);
 }
 
-void URevolveSplineTool::GenerateAsset(const FDynamicMeshOpResult& OpResult)
+FString URevolveSplineTool::GeneratedAssetBaseName() const
 {
-	if (OpResult.Mesh.Get() == nullptr) return;
-
-	if (!ensure(OpResult.Mesh.Get() && OpResult.Mesh->TriangleCount() > 0))
-	{
-		return;
-	}
-
-	GetToolManager()->BeginUndoTransaction(LOCTEXT("RevolveSplineToolTransactionName", "Revolve Spline"));
-
-	FCreateMeshObjectParams NewMeshObjectParams;
-	NewMeshObjectParams.TargetWorld = GetTargetWorld();
-	NewMeshObjectParams.Transform = (FTransform)OpResult.Transform;
-	NewMeshObjectParams.BaseName = TEXT("RevolveSpline");
-	NewMeshObjectParams.Materials.Add(MaterialProperties->Material.Get());
-	NewMeshObjectParams.SetMesh(OpResult.Mesh.Get());
-	OutputTypeProperties->ConfigureCreateMeshObjectParams(NewMeshObjectParams);
-	FCreateMeshObjectResult Result = UE::Modeling::CreateMeshObject(GetToolManager(), MoveTemp(NewMeshObjectParams));
-	if (Result.IsOK() && Result.NewActor != nullptr)
-	{
-		ToolSelectionUtil::SetNewActorSelection(GetToolManager(), Result.NewActor);
-	}
-
-	GetToolManager()->EndUndoTransaction();
+	return TEXT("RevolveSpline");
 }
-
-bool URevolveSplineTool::CanAccept() const
+FText URevolveSplineTool::TransactionName() const
 {
-	return Preview->HaveValidNonEmptyResult();
+	return LOCTEXT("RevolveSplinesAction", "Revolve Spline");
 }
 
 TUniquePtr<FDynamicMeshOperator> URevolveSplineTool::MakeNewOperator()
@@ -410,19 +290,6 @@ void URevolveSplineTool::UpdateRevolutionAxis()
 	}
 }
 
-// To be called by builder
-void URevolveSplineTool::SetSpline(USplineComponent* SplineComponent)
-{
-	Spline = SplineComponent;
-	SplineOwningActor = Spline->GetOwner();
-	bForceSplineUpdate = true;
-
-	TArray<USplineComponent*> Splines;
-	SplineOwningActor->GetComponents<USplineComponent>(Splines);
-
-	SplineComponentIndex = Splines.IndexOfByKey(Spline.Get());
-}
-
 void URevolveSplineTool::RequestAction(ERevolveSplineToolAction Action)
 {
 	if (Action == ERevolveSplineToolAction::ResetAxis)
@@ -434,28 +301,10 @@ void URevolveSplineTool::RequestAction(ERevolveSplineToolAction Action)
 
 /// Tool builder:
 
-bool URevolveSplineToolBuilder::CanBuildTool(const FToolBuilderState& SceneState) const
-{
-	int32 NumSplines = ToolBuilderUtil::CountComponents(SceneState, [&](UActorComponent* Object) -> bool
-	{
-		return Object->IsA<USplineComponent>();
-	});
-	return NumSplines == 1;
-}
-
 UInteractiveTool* URevolveSplineToolBuilder::BuildTool(const FToolBuilderState& SceneState) const
 {
 	URevolveSplineTool* NewTool = NewObject<URevolveSplineTool>(SceneState.ToolManager);
-	USplineComponent* Spline = Cast<USplineComponent>(ToolBuilderUtil::FindFirstComponent(SceneState, [&](UActorComponent* Object)
-	{
-		return Object->IsA<USplineComponent>();
-	}));
-
-	ensure(Spline);
-	NewTool->SetSpline(Spline);
-
-	NewTool->SetWorld(SceneState.World);
-
+	InitializeNewTool(NewTool, SceneState);
 	return NewTool;
 }
 
diff --git a/Engine/Plugins/Experimental/MeshModelingToolsetExp/Source/MeshModelingToolsExp/Private/Splines/BaseMeshFromSplinesTool.cpp b/Engine/Plugins/Experimental/MeshModelingToolsetExp/Source/MeshModelingToolsExp/Private/Splines/BaseMeshFromSplinesTool.cpp
new file mode 100644
--- /dev/null
+++ b/Engine/Plugins/Experimental/MeshModelingToolsetExp/Source/MeshModelingToolsExp/Private/Splines/BaseMeshFromSplinesTool.cpp
@@ -0,0 +1,297 @@
+// Copyright Epic Games, Inc. All Rights Reserved.
+
+#include "Spline/BaseMeshFromSplinesTool.h"
+
+#include "InteractiveToolManager.h"
+#include "ToolBuilderUtil.h"
+#include "ToolSetupUtil.h"
+#include "Selection/ToolSelectionUtil.h"
+#include "DynamicMesh/MeshTransforms.h"
+#include "Properties/MeshMaterialProperties.h"
+
+#include "Engine/World.h"
+
+#include "Components/SplineComponent.h"
+
+#include UE_INLINE_GENERATED_CPP_BY_NAME(BaseMeshFromSplinesTool)
+
+#define LOCTEXT_NAMESPACE "UBaseMeshFromSplinesTool"
+
+using namespace UE::Geometry;
+
+
+
+void UBaseMeshFromSplinesTool::Setup()
+{
+	UInteractiveTool::Setup();
+
+	// initialize our properties
+
+	OutputTypeProperties = NewObject<UCreateMeshObjectTypeProperties>(this);
+	OutputTypeProperties->InitializeDefault();
+	OutputTypeProperties->RestoreProperties(this);
+	OutputTypeProperties->WatchProperty(OutputTypeProperties->OutputType, [this](FString) { OutputTypeProperties->UpdatePropertyVisibility(); });
+	AddToolPropertySource(OutputTypeProperties);
+
+	MaterialProperties = NewObject<UNewMeshMaterialProperties>(this);
+	AddToolPropertySource(MaterialProperties);
+	MaterialProperties->RestoreProperties(this);
+
+	Preview = NewObject<UMeshOpPreviewWithBackgroundCompute>(this);
+	Preview->Setup(GetTargetWorld(), this);
+	Preview->PreviewMesh->EnableWireframe(MaterialProperties->bShowWireframe);
+	Preview->ConfigureMaterials(MaterialProperties->Material.Get(),
+		ToolSetupUtil::GetDefaultWorkingMaterial(GetToolManager()));
+	ToolSetupUtil::ApplyRenderingConfigurationToPreview(Preview->PreviewMesh, nullptr);
+
+	Preview->OnMeshUpdated.AddLambda(
+		[this](const UMeshOpPreviewWithBackgroundCompute* UpdatedPreview)
+		{
+			UpdateAcceptWarnings(UpdatedPreview->HaveEmptyResult() ? EAcceptWarning::EmptyForbidden : EAcceptWarning::NoWarning);
+		}
+	);
+
+	PollSplineUpdates();
+}
+
+void UBaseMeshFromSplinesTool::PollSplineUpdates()
+{
+	if (bLostInputSpline)
+	{
+		return;
+	}
+
+	bool bSplinesUpdated = false;
+	int32 SplineIdx = 0;
+	EnumerateSplines([this, &bSplinesUpdated, &SplineIdx](USplineComponent* SplineComponent)
+	{
+		int32 Version = SplineComponent->SplineCurves.Version;
+		FTransform Transform = SplineComponent->GetComponentTransform();
+		if (SplineIdx >= LastSplineVersions.Num())
+		{
+			bSplinesUpdated = true;
+			LastSplineVersions.Add(Version);
+			LastSplineTransforms.Add(Transform);
+		}
+		else if (LastSplineVersions[SplineIdx] != Version || !LastSplineTransforms[SplineIdx].Equals(Transform))
+		{
+			bSplinesUpdated = true;
+		}
+		LastSplineVersions[SplineIdx] = Version;
+		LastSplineTransforms[SplineIdx] = Transform;
+		++SplineIdx;
+	});
+	if (LastSplineVersions.Num() != SplineIdx)
+	{
+		if (SplineIdx < LastSplineVersions.Num())
+		{
+			bLostInputSpline = true;
+			GetToolManager()->DisplayMessage(
+				LOCTEXT("LostSpline", "Tool lost reference to an input spline; cannot respond to further spline changes."),
+				EToolMessageLevel::UserWarning);
+			return;
+		}
+		LastSplineVersions.SetNum(SplineIdx);
+		LastSplineTransforms.SetNum(SplineIdx);
+		bSplinesUpdated = true;
+	}
+	if (bSplinesUpdated)
+	{
+		OnSplineUpdate();
+
+		Preview->InvalidateResult();
+	}
+}
+
+USplineComponent* UBaseMeshFromSplinesTool::GetFirstSpline() const
+{
+	USplineComponent* ToRet = nullptr;
+	for (int32 ActorIdx = 0; !ToRet && ActorIdx < ActorsWithSplines.Num(); ++ActorIdx)
+	{
+		if (AActor* Actor = ActorsWithSplines[ActorIdx].Get())
+		{
+			Actor->ForEachComponent<USplineComponent>(false, [&ToRet](USplineComponent* SplineComponent)
+			{
+				if (!ToRet)
+				{
+					ToRet = SplineComponent;
+				}
+			});
+		}
+	}
+	return ToRet;
+}
+
+USplineComponent* UBaseMeshFromSplinesTool::GetLastSpline() const
+{
+	USplineComponent* ToRet = nullptr;
+	for (int32 ActorIdx = ActorsWithSplines.Num() - 1; !ToRet && ActorIdx >= 0; --ActorIdx)
+	{
+		if (AActor* Actor = ActorsWithSplines[ActorIdx].Get())
+		{
+			Actor->ForEachComponent<USplineComponent>(false, [&ToRet](USplineComponent* SplineComponent)
+			{
+				ToRet = SplineComponent;
+			});
+		}
+	}
+	return ToRet;
+}
+
+void UBaseMeshFromSplinesTool::OnTick(float DeltaTime)
+{
+	PollSplineUpdates();
+
+	if (Preview)
+	{
+		Preview->Tick(DeltaTime);
+	}
+}
+
+void UBaseMeshFromSplinesTool::SetWorld(UWorld* World)
+{
+	TargetWorld = World;
+}
+
+UWorld* UBaseMeshFromSplinesTool::GetTargetWorld()
+{
+	return TargetWorld.Get();
+}
+
+FTransform3d UBaseMeshFromSplinesTool::HandleOperatorTransform(const FDynamicMeshOpResult& OpResult)
+{
+	FTransform3d NewTransform;
+	if (ActorsWithSplines.Num() == 1 && ActorsWithSplines[0].IsValid()) // in the single-actor case, shove the result back into the original actor transform space
+	{
+		FTransform3d ActorToWorld = (FTransform3d)ActorsWithSplines[0]->GetTransform();
+		MeshTransforms::ApplyTransform(*OpResult.Mesh, OpResult.Transform, true);
+		MeshTransforms::ApplyTransformInverse(*OpResult.Mesh, ActorToWorld, true);
+		NewTransform = ActorToWorld;
+	}
+	else // in the multi-selection case, center the pivot for the combined result
+	{
+		FVector3d Center = OpResult.Mesh->GetBounds().Center();
+		double Rescale = OpResult.Transform.GetScale().X;
+		FTransform3d LocalTransform(-Center * Rescale);
+		LocalTransform.SetScale3D(FVector3d(Rescale, Rescale, Rescale));
+		MeshTransforms::ApplyTransform(*OpResult.Mesh, LocalTransform, true);
+		NewTransform = OpResult.Transform;
+		NewTransform.SetScale3D(FVector3d::One());
+		NewTransform.SetTranslation(NewTransform.GetTranslation() + NewTransform.TransformVector(Center * Rescale));
+	}
+	return NewTransform;
+}
+
+void UBaseMeshFromSplinesTool::GenerateAsset(const FDynamicMeshOpResult& OpResult)
+{
+	if (OpResult.Mesh.Get() == nullptr) return;
+
+	FTransform3d NewTransform = HandleOperatorTransform(OpResult);
+
+	FString BaseName = GeneratedAssetBaseName();
+
+	FCreateMeshObjectParams NewMeshObjectParams;
+	NewMeshObjectParams.TargetWorld = GetTargetWorld();
+	NewMeshObjectParams.Transform = (FTransform)NewTransform;
+	NewMeshObjectParams.BaseName = BaseName;
+	NewMeshObjectParams.Materials.Add(MaterialProperties->Material.Get());
+	NewMeshObjectParams.SetMesh(OpResult.Mesh.Get());
+	OutputTypeProperties->ConfigureCreateMeshObjectParams(NewMeshObjectParams);
+	FCreateMeshObjectResult Result = UE::Modeling::CreateMeshObject(GetToolManager(), MoveTemp(NewMeshObjectParams));
+	if (Result.IsOK() && Result.NewActor != nullptr)
+	{
+		ToolSelectionUtil::SetNewActorSelection(GetToolManager(), Result.NewActor);
+	}
+}
+
+void UBaseMeshFromSplinesTool::OnPropertyModified(UObject* PropertySet, FProperty* Property)
+{
+	if (Property && (PropertySet == OutputTypeProperties))
+	{
+		return;
+	}
+
+	if (Property && Property->GetFName() == GET_MEMBER_NAME_CHECKED(UNewMeshMaterialProperties, Material))
+	{
+		Preview->ConfigureMaterials(MaterialProperties->Material.Get(),
+			ToolSetupUtil::GetDefaultWorkingMaterial(GetToolManager()));
+	}
+
+	Preview->PreviewMesh->EnableWireframe(MaterialProperties->bShowWireframe);
+
+	Preview->InvalidateResult();
+}
+
+void UBaseMeshFromSplinesTool::Shutdown(EToolShutdownType ShutdownType)
+{
+	OutputTypeProperties->SaveProperties(this);
+	MaterialProperties->SaveProperties(this);
+
+	FDynamicMeshOpResult Result = Preview->Shutdown();
+
+	if (ShutdownType == EToolShutdownType::Accept)
+	{
+		GetToolManager()->BeginUndoTransaction(LOCTEXT("SweepSplineAction", "Spline Triangulation"));
+
+		// Generate the result asset
+		GenerateAsset(Result);
+
+		GetToolManager()->EndUndoTransaction();
+	}
+
+	TargetWorld = nullptr;
+	Preview = nullptr;
+	MaterialProperties = nullptr;
+	OutputTypeProperties = nullptr;
+
+	Super::Shutdown(ShutdownType);
+}
+
+bool UBaseMeshFromSplinesTool::CanAccept() const
+{
+	return Preview->HaveValidNonEmptyResult();
+}
+
+TUniquePtr<FDynamicMeshOperator> UBaseMeshFromSplinesTool::MakeNewOperator()
+{
+	checkNoEntry();
+	return TUniquePtr<FDynamicMeshOperator>();
+}
+
+/// Tool builder
+
+bool UBaseMeshFromSplinesToolBuilder::CanBuildTool(const FToolBuilderState& SceneState) const
+{
+	int32 NumSplines = ToolBuilderUtil::CountComponents(SceneState, [&](UActorComponent* Object) -> bool
+		{
+			return Object->IsA<USplineComponent>();
+		});
+	FIndex2i SupportedRange = GetSupportedSplineCountRange();
+	return (NumSplines >= SupportedRange.A && (SupportedRange.B == -1 || NumSplines <= SupportedRange.B));
+}
+
+void UBaseMeshFromSplinesToolBuilder::InitializeNewTool(UBaseMeshFromSplinesTool* NewTool, const FToolBuilderState& SceneState) const
+{
+	TArray<UActorComponent*> Components = ToolBuilderUtil::FindAllComponents(SceneState, [&](UActorComponent* Object)
+		{
+			return Object->IsA<USplineComponent>();
+		});
+	TArray<TWeakObjectPtr<AActor>> ActorsWithSplines;
+	TSet<AActor*> FoundActors;
+	for (UActorComponent* Component : Components)
+	{
+		AActor* ActorWithSpline = Component->GetOwner();
+		if (!FoundActors.Contains(ActorWithSpline))
+		{
+			FoundActors.Add(ActorWithSpline);
+			ActorsWithSplines.Add(ActorWithSpline);
+		}
+	}
+	NewTool->SetSplineActors(MoveTemp(ActorsWithSplines));
+	NewTool->SetWorld(SceneState.World);
+}
+
+
+
+#undef LOCTEXT_NAMESPACE
+
diff --git a/Engine/Plugins/Experimental/MeshModelingToolsetExp/Source/MeshModelingToolsExp/Private/TriangulateSplinesTool.cpp b/Engine/Plugins/Experimental/MeshModelingToolsetExp/Source/MeshModelingToolsExp/Private/TriangulateSplinesTool.cpp
--- a/Engine/Plugins/Experimental/MeshModelingToolsetExp/Source/MeshModelingToolsExp/Private/TriangulateSplinesTool.cpp
+++ b/Engine/Plugins/Experimental/MeshModelingToolsetExp/Source/MeshModelingToolsExp/Private/TriangulateSplinesTool.cpp
@@ -2,11 +2,7 @@
 
 #include "TriangulateSplinesTool.h"
 
-#include "InteractiveToolManager.h"
-#include "ToolBuilderUtil.h"
-#include "ToolSetupUtil.h"
-#include "Selection/ToolSelectionUtil.h"
-#include "DynamicMesh/MeshTransforms.h"
+#include "Properties/MeshMaterialProperties.h"
 
 #include "Engine/World.h"
 
@@ -22,179 +18,31 @@ using namespace UE::Geometry;
 
 void UTriangulateSplinesTool::Setup()
 {
-	UInteractiveTool::Setup();
-
 	// initialize our properties
 
 	TriangulateProperties = NewObject<UTriangulateSplinesToolProperties>(this);
 	TriangulateProperties->RestoreProperties(this);
 	AddToolPropertySource(TriangulateProperties);
 
-	OutputTypeProperties = NewObject<UCreateMeshObjectTypeProperties>(this);
-	OutputTypeProperties->InitializeDefault();
-	OutputTypeProperties->RestoreProperties(this);
-	OutputTypeProperties->WatchProperty(OutputTypeProperties->OutputType, [this](FString) { OutputTypeProperties->UpdatePropertyVisibility(); });
-	AddToolPropertySource(OutputTypeProperties);
+	TriangulateProperties->WatchProperty(TriangulateProperties->ErrorTolerance, [this](double ErrorTolerance)
+		{
+			OnSplineUpdate();
+		});
 
 	SetToolDisplayName(LOCTEXT("TriangulateSplinesToolName", "Triangulate Splines"));
 	GetToolManager()->DisplayMessage(
 		LOCTEXT("TriangulateSplinesToolToolDescription", "Triangulate the shapes of the selected splines."),
 		EToolMessageLevel::UserNotification);
 
-	Preview = NewObject<UMeshOpPreviewWithBackgroundCompute>(this);
-	Preview->Setup(GetTargetWorld(), this);
-	ToolSetupUtil::ApplyRenderingConfigurationToPreview(Preview->PreviewMesh, nullptr);
-
-	Preview->OnMeshUpdated.AddLambda(
-		[this](const UMeshOpPreviewWithBackgroundCompute* UpdatedPreview)
-		{
-			UpdateAcceptWarnings(UpdatedPreview->HaveEmptyResult() ? EAcceptWarning::EmptyForbidden : EAcceptWarning::NoWarning);
-		}
-	);
-
-	PollSplineUpdates();
-}
-
-void UTriangulateSplinesTool::PollSplineUpdates()
-{
-	bool bSplinesUpdated = false;
-	int32 SplineIdx = 0;
-	TArray<TObjectPtr<USplineComponent>, TInlineAllocator<8>> SplineComponents;
-	for (int32 ActorIdx = 0; ActorIdx < ActorsWithSplines.Num(); ++ActorIdx)
-	{
-		if (AActor* Actor = ActorsWithSplines[ActorIdx].Get())
-		{
-			SplineComponents.Reset();
-			Actor->GetComponents<USplineComponent, TInlineAllocator<8>>(SplineComponents);
-			for (USplineComponent* SplineComponent : SplineComponents)
-			{
-				int32 Version = SplineComponent->SplineCurves.Version;
-				FTransform Transform = SplineComponent->GetComponentTransform();
-				if (SplineIdx >= LastSplineVersions.Num())
-				{
-					bSplinesUpdated = true;
-					LastSplineVersions.Add(Version);
-					LastSplineTransforms.Add(Transform);
-				}
-				else if (LastSplineVersions[SplineIdx] != Version || !LastSplineTransforms[SplineIdx].Equals(Transform))
-				{
-					bSplinesUpdated = true;
-				}
-				LastSplineVersions[SplineIdx] = Version;
-				LastSplineTransforms[SplineIdx] = Transform;
-				++SplineIdx;
-			}
-		}
-	}
-	if (LastSplineVersions.Num() != SplineIdx)
-	{
-		LastSplineVersions.SetNum(SplineIdx);
-		LastSplineTransforms.SetNum(SplineIdx);
-		bSplinesUpdated = true;
-	}
-	if (bSplinesUpdated)
-	{
-		Preview->InvalidateResult();
-	}
-}
-
-void UTriangulateSplinesTool::OnTick(float DeltaTime)
-{
-	PollSplineUpdates();
-
-	Preview->Tick(DeltaTime);
-}
-
-void UTriangulateSplinesTool::SetWorld(UWorld* World)
-{
-	TargetWorld = World;
-}
-
-UWorld* UTriangulateSplinesTool::GetTargetWorld()
-{
-	return TargetWorld.Get();
-}
-
-void UTriangulateSplinesTool::GenerateAsset(const FDynamicMeshOpResult& OpResult)
-{
-	if (OpResult.Mesh.Get() == nullptr) return;
-
-	FTransform3d NewTransform;
-	if (ActorsWithSplines.Num() == 1 && ActorsWithSplines[0].IsValid()) // in the single-actor case, shove the result back into the original actor transform space
-	{
-		FTransform3d ActorToWorld = (FTransform3d)ActorsWithSplines[0]->GetTransform();
-		MeshTransforms::ApplyTransform(*OpResult.Mesh, OpResult.Transform, true);
-		MeshTransforms::ApplyTransformInverse(*OpResult.Mesh, ActorToWorld, true);
-		NewTransform = ActorToWorld;
-	}
-	else // in the multi-selection case, center the pivot for the combined result
-	{
-		FVector3d Center = OpResult.Mesh->GetBounds().Center();
-		double Rescale = OpResult.Transform.GetScale().X;
-		FTransform3d LocalTransform(-Center * Rescale);
-		LocalTransform.SetScale3D(FVector3d(Rescale, Rescale, Rescale));
-		MeshTransforms::ApplyTransform(*OpResult.Mesh, LocalTransform, true);
-		NewTransform = OpResult.Transform;
-		NewTransform.SetScale3D(FVector3d::One());
-		NewTransform.SetTranslation(NewTransform.GetTranslation() + NewTransform.TransformVector(Center * Rescale));
-	}
-
-	FString BaseName = TEXT("Triangulation");
-
-	FCreateMeshObjectParams NewMeshObjectParams;
-	NewMeshObjectParams.TargetWorld = GetTargetWorld();
-	NewMeshObjectParams.Transform = (FTransform)NewTransform;
-	NewMeshObjectParams.BaseName = BaseName;
-	NewMeshObjectParams.Materials = Preview->StandardMaterials;
-	NewMeshObjectParams.SetMesh(OpResult.Mesh.Get());
-	OutputTypeProperties->ConfigureCreateMeshObjectParams(NewMeshObjectParams);
-	FCreateMeshObjectResult Result = UE::Modeling::CreateMeshObject(GetToolManager(), MoveTemp(NewMeshObjectParams));
-	if (Result.IsOK() && Result.NewActor != nullptr)
-	{
-		ToolSelectionUtil::SetNewActorSelection(GetToolManager(), Result.NewActor);
-	}
-}
-
-void UTriangulateSplinesTool::OnPropertyModified(UObject* PropertySet, FProperty* Property)
-{
-	if (Property &&
-		(PropertySet == OutputTypeProperties
-			))
-	{
-		// nothing
-	}
-	else
-	{
-		Preview->InvalidateResult();
-	}
+	Super::Setup();
 }
 
 void UTriangulateSplinesTool::Shutdown(EToolShutdownType ShutdownType)
 {
-	OutputTypeProperties->SaveProperties(this);
 	TriangulateProperties->SaveProperties(this);
-
-	FDynamicMeshOpResult Result = Preview->Shutdown();
-
-	if (ShutdownType == EToolShutdownType::Accept)
-	{
-		GetToolManager()->BeginUndoTransaction(LOCTEXT("TriangulateSplinesAction", "Spline Triangulation"));
-
-		// Generate the result asset
-		GenerateAsset(Result);
-
-		GetToolManager()->EndUndoTransaction();
-	}
-
-	TargetWorld = nullptr;
 	Super::Shutdown(ShutdownType);
 }
 
-bool UTriangulateSplinesTool::CanAccept() const
-{
-	return Preview->HaveValidNonEmptyResult();
-}
-
 TUniquePtr<FDynamicMeshOperator> UTriangulateSplinesTool::MakeNewOperator()
 {
 	TUniquePtr<FTriangulateCurvesOp> Op = MakeUnique<FTriangulateCurvesOp>();
@@ -216,37 +64,49 @@ TUniquePtr<FDynamicMeshOperator> UTriangulateSplinesTool::MakeNewOperator()
 	Op->OffsetJoinMethod = TriangulateProperties->JoinMethod;
 	Op->OpenEndShape = TriangulateProperties->EndShapes;
 	Op->MiterLimit = TriangulateProperties->MiterLimit;
+	Op->UVScaleFactor = MaterialProperties->UVScale;
+	Op->bWorldSpaceUVScale = MaterialProperties->bWorldSpaceUVScale;
 	
 	Op->bFlipResult = TriangulateProperties->bFlipResult;
 
-	TArray<TObjectPtr<USplineComponent>, TInlineAllocator<8>> SplineComponents;
-	for (int32 ActorIdx = 0; ActorIdx < ActorsWithSplines.Num(); ++ActorIdx)
+	for (FPathCache& Path : SplinesCache)
 	{
-		if (AActor* Actor = ActorsWithSplines[ActorIdx].Get())
-		{
-			SplineComponents.Reset();
-			Actor->GetComponents<USplineComponent, TInlineAllocator<8>>(SplineComponents);
-			for (USplineComponent* SplineComponent : SplineComponents)
-			{
-				Op->AddSpline(SplineComponent, TriangulateProperties->ErrorTolerance);
-			}
-		}
+		Op->AddWorldCurve(Path.Vertices, Path.bClosed, Path.ComponentTransform);
 	}
 
 	return Op;
 }
 
-/// Tool builder
-
-bool UTriangulateSplinesToolBuilder::CanBuildTool(const FToolBuilderState& SceneState) const
+void UTriangulateSplinesTool::OnSplineUpdate()
 {
-	int32 NumSplines = ToolBuilderUtil::CountComponents(SceneState, [&](UActorComponent* Object) -> bool
+	if (bLostInputSpline)
 	{
-		return Object->IsA<USplineComponent>();
+		return;
+	}
+
+	SplinesCache.Reset();
+	EnumerateSplines([this](USplineComponent* SplineComponent)
+	{
+		FPathCache& Path = SplinesCache.Emplace_GetRef();
+		SplineComponent->ConvertSplineToPolyLine(ESplineCoordinateSpace::Type::World, 
+			TriangulateProperties->ErrorTolerance * TriangulateProperties->ErrorTolerance, Path.Vertices);
+		Path.ComponentTransform = SplineComponent->GetComponentTransform();
+		Path.bClosed = SplineComponent->IsClosedLoop();
 	});
-	return NumSplines > 0;
 }
 
+FString UTriangulateSplinesTool::GeneratedAssetBaseName() const
+{
+	return TEXT("Triangulation");
+}
+FText UTriangulateSplinesTool::TransactionName() const
+{
+	return LOCTEXT("TriangulateSplinesAction", "Spline Triangulation");
+}
+
+
+/// Tool builder
+
 UInteractiveTool* UTriangulateSplinesToolBuilder::BuildTool(const FToolBuilderState& SceneState) const
 {
 	UTriangulateSplinesTool* NewTool = NewObject<UTriangulateSplinesTool>(SceneState.ToolManager);
@@ -254,28 +114,6 @@ UInteractiveTool* UTriangulateSplinesToolBuilder::BuildTool(const FToolBuilderSt
 	return NewTool;
 }
 
-void UTriangulateSplinesToolBuilder::InitializeNewTool(UTriangulateSplinesTool* NewTool, const FToolBuilderState& SceneState) const
-{
-	TArray<UActorComponent*> Components = ToolBuilderUtil::FindAllComponents(SceneState, [&](UActorComponent* Object)
-	{
-		return Object->IsA<USplineComponent>();
-	});
-	TArray<TWeakObjectPtr<AActor>> ActorsWithSplines;
-	TSet<AActor*> FoundActors;
-	for (UActorComponent* Component : Components)
-	{
-		AActor* ActorWithSpline = Component->GetOwner();
-		if (!FoundActors.Contains(ActorWithSpline))
-		{
-			FoundActors.Add(ActorWithSpline);
-			ActorsWithSplines.Add(ActorWithSpline);
-		}
-	}
-	NewTool->SetSplineActors(MoveTemp(ActorsWithSplines));
-	NewTool->SetWorld(SceneState.World);
-}
-
-
 
 #undef LOCTEXT_NAMESPACE
 
diff --git a/Engine/Plugins/Experimental/MeshModelingToolsetExp/Source/MeshModelingToolsExp/Public/RevolveSplineTool.h b/Engine/Plugins/Experimental/MeshModelingToolsetExp/Source/MeshModelingToolsExp/Public/RevolveSplineTool.h
--- a/Engine/Plugins/Experimental/MeshModelingToolsetExp/Source/MeshModelingToolsExp/Public/RevolveSplineTool.h
+++ b/Engine/Plugins/Experimental/MeshModelingToolsetExp/Source/MeshModelingToolsExp/Public/RevolveSplineTool.h
@@ -4,25 +4,13 @@
 
 #include "CoreMinimal.h"
 
-#include "Components/SplineComponent.h" // (to use with TWeakObjectPtr)
-#include "Engine/World.h" // (to use with TWeakObjectPtr)
-#include "InteractiveTool.h"
-#include "InteractiveToolBuilder.h"
-#include "InteractiveToolQueryInterfaces.h"
-#include "ModelingOperators.h" // IDynamicMeshOperatorFactory
 #include "Properties/RevolveProperties.h"
+#include "Spline/BaseMeshFromSplinesTool.h"
 
 #include "RevolveSplineTool.generated.h"
 
-class AActor;
-struct FDynamicMeshOpResult;
 class UConstructionPlaneMechanic;
-class UCreateMeshObjectTypeProperties;
-class UMeshOpPreviewWithBackgroundCompute;
-class UNewMeshMaterialProperties;
 class URevolveSplineTool;
-class USplineComponent;
-class UWorld;
 
 //~ TODO: Might want to have some shared enum for sampling splines in a util folder,
 //~ but hesitant to prescribe it until other tools want it.
@@ -124,80 +112,52 @@ public:
  * Revolves a selected spline to create a new mesh.
  */
 UCLASS()
-class MESHMODELINGTOOLSEXP_API URevolveSplineTool : public UInteractiveTool, 
-	public IInteractiveToolEditorGizmoAPI, 
-	public UE::Geometry::IDynamicMeshOperatorFactory
+class MESHMODELINGTOOLSEXP_API URevolveSplineTool : public UBaseMeshFromSplinesTool
 {
 	GENERATED_BODY()
 
 public:
 
-	virtual void SetSpline(USplineComponent* SplineComponent);
-
-	virtual void SetWorld(UWorld* World) { TargetWorld = World; }
-	virtual UWorld* GetTargetWorld() { return TargetWorld.Get(); }
 	virtual void RequestAction(ERevolveSplineToolAction ActionType);
 
 	// UInteractiveTool
 	virtual void Setup() override;
 	virtual void Shutdown(EToolShutdownType ShutdownType) override;
-	virtual void OnTick(float DeltaTime) override;
-	virtual bool HasCancel() const override { return true; }
-	virtual bool HasAccept() const override { return true; }
-	virtual bool CanAccept() const override;
 	virtual void Render(IToolsContextRenderAPI* RenderAPI) override;
 	virtual void OnPropertyModified(UObject* PropertySet, FProperty* Property) override;
+	virtual void OnTick(float DeltaTime) override;
 
 	// IDynamicMeshOperatorFactory
-	virtual TUniquePtr<UE::Geometry::FDynamicMeshOperator> MakeNewOperator();
+	virtual TUniquePtr<UE::Geometry::FDynamicMeshOperator> MakeNewOperator() override;
 
-	// IInteractiveToolEditorGizmoAPI
-	virtual bool GetAllowStandardEditorGizmos() override { return true; }
+	virtual FString GeneratedAssetBaseName() const override;
+	virtual FText TransactionName() const override;
 
-private:
+protected:
+	// Update the profile curve and fit plane from spline
+	virtual void OnSplineUpdate() override;
 
-	UPROPERTY()
-	TObjectPtr<URevolveSplineToolProperties> Settings = nullptr;
+	// Keep the result mesh in the same space as set by the operator result
+	virtual FTransform3d HandleOperatorTransform(const FDynamicMeshOpResult& OpResult) override
+	{
+		return OpResult.Transform;
+	}
 
-	UPROPERTY()
-	TObjectPtr<UNewMeshMaterialProperties> MaterialProperties = nullptr;
+private:
 
 	UPROPERTY()
-	TObjectPtr<UCreateMeshObjectTypeProperties> OutputTypeProperties = nullptr;
+	TObjectPtr<URevolveSplineToolProperties> Settings = nullptr;
 
 	UPROPERTY()
 	TObjectPtr<URevolveSplineToolActionPropertySet> ToolActions = nullptr;
 
-	UPROPERTY()
-	TObjectPtr<UMeshOpPreviewWithBackgroundCompute> Preview = nullptr;
-
 	UPROPERTY()
 	TObjectPtr<UConstructionPlaneMechanic> PlaneMechanic = nullptr;
 
-	TWeakObjectPtr<UWorld> TargetWorld = nullptr;
-
-	TWeakObjectPtr<USplineComponent> Spline = nullptr;
-	
-	// Used to try to re-acquire the spline if it is inside a BP actor, where it will be destroyed and recreated
-	// whenever the user edits it.
-	TWeakObjectPtr<AActor> SplineOwningActor = nullptr;
-	int32 SplineComponentIndex;
-	// If failed to reacquire once, used to avoid trying to reaquire again.
-	bool bLostInputSpline = false;
-
 	// The actual points to be resolved, sampled from the spline
 	TArray<FVector3d> ProfileCurve;
 	bool bProfileCurveIsClosed;
 
-	// See if the spline has changed, and update our data if so
-	void PollSplineUpdates();
-	// Track the spline 'Version' integer, which is incremented when splines are changed
-	uint32 LastSplineVersion = 0;
-	// Used to make sure we update when we first get a spline, regardless of LastSplineVersion
-	bool bForceSplineUpdate = true;
-	// Update the profile curve and fit plane from spline
-	void UpdatePointsFromSpline();
-
 	// Axis direction in vector form (since the user modifiable values are a pitch and yaw)
 	FVector3d RevolutionAxisDirection;
 	// This duplicates Settings->AxisOrigin, but kept for cleanliness since we do need RevolutionAxisDirection
@@ -208,21 +168,22 @@ private:
 
 	FVector3d SplineFitPlaneOrigin;
 	FVector3d SplineFitPlaneNormal;
-
-	virtual void GenerateAsset(const FDynamicMeshOpResult& OpResult);
 };
 
 
 
 UCLASS(Transient)
-class MESHMODELINGTOOLSEXP_API URevolveSplineToolBuilder : public UInteractiveToolBuilder
+class MESHMODELINGTOOLSEXP_API URevolveSplineToolBuilder : public UBaseMeshFromSplinesToolBuilder
 {
 	GENERATED_BODY()
 
 public:
-	/** @return true if spline component sources can be found in the active selection */
-	virtual bool CanBuildTool(const FToolBuilderState& SceneState) const override;
-
 	/** @return new Tool instance initialized with selected spline source(s) */
 	virtual UInteractiveTool* BuildTool(const FToolBuilderState& SceneState) const override;
+
+	virtual UE::Geometry::FIndex2i GetSupportedSplineCountRange() const override
+	{
+		return UE::Geometry::FIndex2i(1, 1);
+	}
 };
+
diff --git a/Engine/Plugins/Experimental/MeshModelingToolsetExp/Source/MeshModelingToolsExp/Public/Spline/BaseMeshFromSplinesTool.h b/Engine/Plugins/Experimental/MeshModelingToolsetExp/Source/MeshModelingToolsExp/Public/Spline/BaseMeshFromSplinesTool.h
new file mode 100644
--- /dev/null
+++ b/Engine/Plugins/Experimental/MeshModelingToolsetExp/Source/MeshModelingToolsExp/Public/Spline/BaseMeshFromSplinesTool.h
@@ -0,0 +1,184 @@
+// Copyright Epic Games, Inc. All Rights Reserved.
+
+#pragma once
+
+#include "CoreMinimal.h"
+#include "UObject/NoExportTypes.h"
+
+#include "Components/SplineComponent.h" // (to use with TWeakObjectPtr)
+#include "Engine/World.h" // (to use with TWeakObjectPtr)
+#include "IndexTypes.h"
+#include "InteractiveTool.h"
+#include "InteractiveToolBuilder.h"
+#include "InteractiveToolManager.h"
+#include "MeshOpPreviewHelpers.h"
+#include "InteractiveToolQueryInterfaces.h"
+#include "PropertySets/CreateMeshObjectTypeProperties.h"
+
+
+#include "BaseMeshFromSplinesTool.generated.h"
+
+class USplineComponent;
+class UNewMeshMaterialProperties;
+class UWorld;
+class AActor;
+struct FDynamicMeshOpResult;
+class UCreateMeshObjectTypeProperties;
+class UMeshOpPreviewWithBackgroundCompute;
+class UNewMeshMaterialProperties;
+
+/**
+ * Tool to create a mesh from a set of selected Spline Components
+ */
+UCLASS()
+class MESHMODELINGTOOLSEXP_API UBaseMeshFromSplinesTool : public UInteractiveTool, public IInteractiveToolEditorGizmoAPI, public UE::Geometry::IDynamicMeshOperatorFactory
+{
+	GENERATED_BODY()
+
+public:
+
+	// IInteractiveToolEditorGizmoAPI -- allow editor gizmo so users can live-edit the splines
+	virtual bool GetAllowStandardEditorGizmos() override
+	{
+		return true;
+	}
+
+	// InteractiveTool API
+	virtual void Setup() override;
+	virtual void Shutdown(EToolShutdownType ShutdownType) override;
+	virtual void OnTick(float DeltaTime) override;
+	virtual bool HasCancel() const override { return true; }
+	virtual bool HasAccept() const override { return true; }
+	virtual bool CanAccept() const override;
+
+	virtual void OnPropertyModified(UObject* PropertySet, FProperty* Property) override;
+
+	virtual void SetSplineActors(TArray<TWeakObjectPtr<AActor>> InSplineActors)
+	{
+		ActorsWithSplines = MoveTemp(InSplineActors);
+	}
+
+	virtual void SetWorld(UWorld* World);
+	virtual UWorld* GetTargetWorld();
+
+	// IDynamicMeshOperatorFactory API
+	virtual TUniquePtr<UE::Geometry::FDynamicMeshOperator> MakeNewOperator();
+
+protected:
+
+	// Override to respond to spline changes
+	virtual void OnSplineUpdate() {};
+
+
+	virtual void GenerateAsset(const FDynamicMeshOpResult& OpResult);
+
+	//
+	// API for asset generation: Override the below methods to customize common parts of spline tool asset generation:
+	// 
+	
+	// Called by GenerateAsset to set the desired asset transform and if needed transform the result to the corresponding local space
+	// @return The transform to use for the generated asset
+	virtual FTransform3d HandleOperatorTransform(const FDynamicMeshOpResult& OpResult);
+	// Override with an appropriate prefix for generated meshes
+	virtual FString GeneratedAssetBaseName() const
+	{
+		return FString(TEXT("SplineMesh"));
+	}
+	// Override with an appropriate tool transaction name
+	virtual FText TransactionName() const
+	{
+		return FText();
+	}
+
+	//
+	// Methods to access the active splines 
+	//
+
+	template<typename Func>
+	void EnumerateSplines(Func SplineComponentFunc) const
+	{
+		if (bLostInputSpline)
+		{
+			return;
+		}
+
+		for (int32 ActorIdx = 0; ActorIdx < ActorsWithSplines.Num(); ++ActorIdx)
+		{
+			if (AActor* Actor = ActorsWithSplines[ActorIdx].Get())
+			{
+				Actor->ForEachComponent<USplineComponent>(false, [&SplineComponentFunc](USplineComponent* SplineComponent)
+					{
+						SplineComponentFunc(SplineComponent);
+					});
+			}
+		}
+	}
+
+	inline int32 NumSplines() const
+	{
+		int32 SplineCount = 0;
+		EnumerateSplines([&](USplineComponent* Unused) -> void
+			{
+				SplineCount++;
+			});
+		return SplineCount;
+	}
+
+	USplineComponent* GetFirstSpline() const;
+	USplineComponent* GetLastSpline() const;
+
+	// Common spline tool properties
+
+	UPROPERTY()
+	TObjectPtr<UCreateMeshObjectTypeProperties> OutputTypeProperties = nullptr;
+
+	UPROPERTY()
+	TObjectPtr<UNewMeshMaterialProperties> MaterialProperties = nullptr;
+
+	UPROPERTY()
+	TObjectPtr<UMeshOpPreviewWithBackgroundCompute> Preview;
+
+	UPROPERTY()
+	TWeakObjectPtr<UWorld> TargetWorld = nullptr;
+
+	// Note: We track actors instead of the USplineComponents here because the USplineComponents objects are often deleted / swapped for identical but new objects
+	UPROPERTY()
+	TArray<TWeakObjectPtr<AActor>> ActorsWithSplines;
+
+	// Helper to track the splines we are triangulating, so we can re-triangulate when they are moved or changed
+	void PollSplineUpdates();
+	// Track the spline 'Version' integer, which is incremented when splines are changed
+	TArray<uint32> LastSplineVersions;
+	// Track the spline component's transform (to world space)
+	TArray<FTransform> LastSplineTransforms;
+
+	// If failed to reacquire once, used to avoid trying to reaquire again.
+	bool bLostInputSpline = false;
+};
+
+
+
+/**
+ * Base Tool Builder for tools that operate on a selection of Spline Components
+ */
+UCLASS(Transient)
+class MESHMODELINGTOOLSEXP_API UBaseMeshFromSplinesToolBuilder : public UInteractiveToolBuilder
+{
+	GENERATED_BODY()
+
+public:
+	/** @return true if spline component sources can be found in the active selection */
+	virtual bool CanBuildTool(const FToolBuilderState& SceneState) const override;
+
+	/** Called by BuildTool to configure the Tool with the input spline source(s) based on the SceneState */
+	virtual void InitializeNewTool(UBaseMeshFromSplinesTool* Tool, const FToolBuilderState& SceneState) const;
+
+	// @return the min and max (inclusive) number of splines allowed in the selection for the tool to be built. A value of -1 can be used to indicate there is no maximum.
+	virtual UE::Geometry::FIndex2i GetSupportedSplineCountRange() const
+	{
+		return UE::Geometry::FIndex2i(1, -1);
+	}
+};
+
+
+
diff --git a/Engine/Plugins/Experimental/MeshModelingToolsetExp/Source/MeshModelingToolsExp/Public/TriangulateSplinesTool.h b/Engine/Plugins/Experimental/MeshModelingToolsetExp/Source/MeshModelingToolsExp/Public/TriangulateSplinesTool.h
--- a/Engine/Plugins/Experimental/MeshModelingToolsetExp/Source/MeshModelingToolsExp/Public/TriangulateSplinesTool.h
+++ b/Engine/Plugins/Experimental/MeshModelingToolsetExp/Source/MeshModelingToolsExp/Public/TriangulateSplinesTool.h
@@ -5,20 +5,12 @@
 #include "CoreMinimal.h"
 #include "UObject/NoExportTypes.h"
 
-#include "InteractiveTool.h"
-#include "InteractiveToolBuilder.h"
-#include "InteractiveToolManager.h"
-#include "MeshOpPreviewHelpers.h"
-#include "InteractiveToolQueryInterfaces.h"
-#include "PropertySets/CreateMeshObjectTypeProperties.h"
-
 #include "CurveOps/TriangulateCurvesOp.h"
 
+#include "Spline/BaseMeshFromSplinesTool.h"
 
-#include "TriangulateSplinesTool.generated.h"
 
-class USplineComponent;
-class UWorld;
+#include "TriangulateSplinesTool.generated.h"
 
 
 /**
@@ -52,7 +44,7 @@ public:
 
 	// How to handle open curves: Either offset them, or treat them as closed curves
 	UPROPERTY(EditAnywhere, Category = Offset, meta = (EditCondition = "FlattenMethod != EFlattenCurveMethod::DoNotFlatten", EditConditionHides))
-	EOffsetOpenCurvesMethod OpenCurves = EOffsetOpenCurvesMethod::Offset;
+	EOffsetOpenCurvesMethod OpenCurves = EOffsetOpenCurvesMethod::TreatAsClosed;
 
 	// How much offset to apply to curves
 	UPROPERTY(EditAnywhere, Category = Offset, meta = (EditCondition = "FlattenMethod != EFlattenCurveMethod::DoNotFlatten", EditConditionHides))
@@ -77,99 +69,52 @@ public:
 };
 
 /**
- * Tool to create a mesh from a set of selected Spline Components
+ * Tool to create a mesh by triangulating the shapes outlined or traced by a set of selected Spline Components, with optional offset and extrusion
  */
 UCLASS()
-class MESHMODELINGTOOLSEXP_API UTriangulateSplinesTool : public UInteractiveTool, public IInteractiveToolEditorGizmoAPI, public UE::Geometry::IDynamicMeshOperatorFactory
+class MESHMODELINGTOOLSEXP_API UTriangulateSplinesTool : public UBaseMeshFromSplinesTool
 {
 	GENERATED_BODY()
 
 public:
 
-	UTriangulateSplinesTool() = default;
-
-	// IInteractiveToolEditorGizmoAPI -- allow editor gizmo so users can live-edit the splines
-
-	virtual bool GetAllowStandardEditorGizmos() override
-	{
-		return true;
-	}
-
-	//
-	// InteractiveTool API - generally does not need to be modified by subclasses
-	//
-
 	virtual void Setup() override;
 	virtual void Shutdown(EToolShutdownType ShutdownType) override;
 
-	virtual void OnTick(float DeltaTime) override;
-
-	virtual bool HasCancel() const override { return true; }
-	virtual bool HasAccept() const override { return true; }
-	virtual bool CanAccept() const override;
-
-	virtual void OnPropertyModified(UObject* PropertySet, FProperty* Property) override;
-
-	void SetSplineActors(TArray<TWeakObjectPtr<AActor>> InSplineActors)
-	{
-		ActorsWithSplines = MoveTemp(InSplineActors);
-	}
+	// IDynamicMeshOperatorFactory API
+	virtual TUniquePtr<UE::Geometry::FDynamicMeshOperator> MakeNewOperator() override;
 
-	virtual void SetWorld(UWorld* World);
-	virtual UWorld* GetTargetWorld();
+	virtual FString GeneratedAssetBaseName() const override;
+	virtual FText TransactionName() const override;
 
-	// IDynamicMeshOperatorFactory API
-	virtual TUniquePtr<UE::Geometry::FDynamicMeshOperator> MakeNewOperator();
+protected:
+	virtual void OnSplineUpdate() override;
 
 private:
 
-	virtual void GenerateAsset(const FDynamicMeshOpResult& OpResult);
-
 	UPROPERTY()
 	TObjectPtr<UTriangulateSplinesToolProperties> TriangulateProperties;
 
-	UPROPERTY()
-	TObjectPtr<UCreateMeshObjectTypeProperties> OutputTypeProperties;
-
-	UPROPERTY()
-	TObjectPtr<UMeshOpPreviewWithBackgroundCompute> Preview;
-
-	UPROPERTY()
-	TWeakObjectPtr<UWorld> TargetWorld = nullptr;
-
-	// Note: We track actors instead of the USplineComponents here because the USplineComponents objects are often deleted / swapped for identical but new objects
-	UPROPERTY()
-	TArray<TWeakObjectPtr<AActor>> ActorsWithSplines;
-
-private:
+	// Sampled splines, computed in OnSplineUpdate
+	struct FPathCache
+	{
+		TArray<FVector3d> Vertices;
+		bool bClosed;
+		FTransform ComponentTransform;
+	};
+	TArray<FPathCache> SplinesCache;
 
-	// Helper to track the splines we are triangulating, so we can re-triangulate when they are moved or changed
-	void PollSplineUpdates();
-	// Track the spline 'Version' integer, which is incremented when splines are changed
-	TArray<uint32> LastSplineVersions;
-	// Track the spline component's transform (to world space)
-	TArray<FTransform> LastSplineTransforms;
 };
 
 
-
-/**
- * Base Tool Builder for tools that operate on a selection of Spline Components
- */
 UCLASS(Transient)
-class MESHMODELINGTOOLSEXP_API UTriangulateSplinesToolBuilder : public UInteractiveToolBuilder
+class MESHMODELINGTOOLSEXP_API UTriangulateSplinesToolBuilder : public UBaseMeshFromSplinesToolBuilder
 {
 	GENERATED_BODY()
 
 public:
-	/** @return true if spline component sources can be found in the active selection */
-	virtual bool CanBuildTool(const FToolBuilderState& SceneState) const override;
-
 	/** @return new Tool instance initialized with selected spline source(s) */
 	virtual UInteractiveTool* BuildTool(const FToolBuilderState& SceneState) const override;
-
-	/** Called by BuildTool to configure the Tool with the input spline source(s) based on the SceneState */
-	virtual void InitializeNewTool(UTriangulateSplinesTool* Tool, const FToolBuilderState& SceneState) const;
 };
 
 
diff --git a/Engine/Plugins/Runtime/MeshModelingToolset/Source/ModelingOperators/Private/CurveOps/TriangulateCurvesOp.cpp b/Engine/Plugins/Runtime/MeshModelingToolset/Source/ModelingOperators/Private/CurveOps/TriangulateCurvesOp.cpp
--- a/Engine/Plugins/Runtime/MeshModelingToolset/Source/ModelingOperators/Private/CurveOps/TriangulateCurvesOp.cpp
+++ b/Engine/Plugins/Runtime/MeshModelingToolset/Source/ModelingOperators/Private/CurveOps/TriangulateCurvesOp.cpp
@@ -29,6 +29,17 @@ void FTriangulateCurvesOp::AddSpline(USplineComponent* Spline, double ErrorToler
 	}
 }
 
+void FTriangulateCurvesOp::AddWorldCurve(TArrayView<const FVector3d> WorldSpaceVertices, bool bClosed, const FTransform& ReferenceTransform)
+{
+	FCurvePath& Path = Paths.Emplace_GetRef();
+	Path.Vertices.Append(WorldSpaceVertices);
+	Path.bClosed = bClosed;
+	if (Paths.Num() == 1)
+	{
+		FirstPathTransform = ReferenceTransform;
+	}
+}
+
 void FTriangulateCurvesOp::CalculateResult(FProgressCancel* Progress)
 {
 	if (Progress && Progress->Cancelled())
@@ -47,7 +58,8 @@ void FTriangulateCurvesOp::CalculateResult(FProgressCancel* Progress)
 	{
 		InputBounds.Contain(Path.Vertices);
 	}
-	double UseUVScaleFactor = UVScaleFactor / InputBounds.MaxDim();
+	double UVScaleDenom = bWorldSpaceUVScale ? 100.0 : InputBounds.MaxDim();
+	double UseUVScaleFactor = UVScaleFactor / UVScaleDenom;
 
 	auto AppendTriangles = [this, UseUVScaleFactor](FDynamicMesh3& Mesh, const TArray<FVector3d>& Vertices, const TArray<FIndex3i>& Tris, int32 GroupID, FVector3d PlaneOrigin, FVector3d Normal, bool bFlip) -> void
 	{
diff --git a/Engine/Plugins/Runtime/MeshModelingToolset/Source/ModelingOperators/Public/CurveOps/TriangulateCurvesOp.h b/Engine/Plugins/Runtime/MeshModelingToolset/Source/ModelingOperators/Public/CurveOps/TriangulateCurvesOp.h
--- a/Engine/Plugins/Runtime/MeshModelingToolset/Source/ModelingOperators/Public/CurveOps/TriangulateCurvesOp.h
+++ b/Engine/Plugins/Runtime/MeshModelingToolset/Source/ModelingOperators/Public/CurveOps/TriangulateCurvesOp.h
@@ -98,8 +98,16 @@ public:
 	// Inputs
 	//
 
+	// Sample a spline with the given ErrorTolerance and add it as a curve
 	void AddSpline(USplineComponent* Spline, double ErrorTolerance);
 
+	// Add a curve in world space
+	// @param WorldSpaceVertices The vertices of the curve, in world space
+	// @param bClosed Whether the curve is closed, i.e. should include an edge from the last vertex back to the first vertex
+	// @param ReferenceTransform The transform that would take the curve vertices from local space to world space
+	// Note: The first curve's Reference Transform is used as the local reference frame for the triangulation, and will be the operator's Result Transform
+	void AddWorldCurve(TArrayView<const FVector3d> WorldSpaceVertices, bool bClosed, const FTransform& ReferenceTransform);
+
 	//
 	// Parameters
 	//
@@ -109,6 +117,9 @@ public:
 	// Scaling applied to the default UV values
 	double UVScaleFactor = 1.0;
 
+	// If true, UVs will be consistently scaled relative to the world space, otherwise UVs will be relative to the mesh bounds.
+	bool bWorldSpaceUVScale = false;
+
 	// If > 0, thicken the result mesh to make a solid
 	double Thickness = 0.0;
 
