MeshModelingToolset: Move SDraggableBox slate widget into a new ModelingUI module, allowing it's usage in non-Editor Only contexts.

diff --git a/Engine/Plugins/Experimental/MeshModelingToolsetExp/MeshModelingToolsetExp.uplugin b/Engine/Plugins/Experimental/MeshModelingToolsetExp/MeshModelingToolsetExp.uplugin
--- a/Engine/Plugins/Experimental/MeshModelingToolsetExp/MeshModelingToolsetExp.uplugin
+++ b/Engine/Plugins/Experimental/MeshModelingToolsetExp/MeshModelingToolsetExp.uplugin
@@ -36,6 +36,11 @@
 			"Type": "Editor",
 			"LoadingPhase": "Default"
 		},
+		{
+			"Name": "ModelingUI",
+			"Type": "Runtime",
+			"LoadingPhase": "Default"
+		},
 		{
 			"Name": "SkeletalMeshModifiers",
 			"Type": "Editor",
diff --git a/Engine/Plugins/Experimental/MeshModelingToolsetExp/Source/ModelingEditorUI/ModelingEditorUI.build.cs b/Engine/Plugins/Experimental/MeshModelingToolsetExp/Source/ModelingEditorUI/ModelingEditorUI.build.cs
--- a/Engine/Plugins/Experimental/MeshModelingToolsetExp/Source/ModelingEditorUI/ModelingEditorUI.build.cs
+++ b/Engine/Plugins/Experimental/MeshModelingToolsetExp/Source/ModelingEditorUI/ModelingEditorUI.build.cs
@@ -39,6 +39,7 @@ public class ModelingEditorUI : ModuleRules
 				"EditorWidgets",
 				"DeveloperSettings",
 				"Persona",
+				"ModelingUI",
 				"MeshModelingToolsEditorOnlyExp"
 			}
 			);
diff --git a/Engine/Plugins/Experimental/MeshModelingToolsetExp/Source/ModelingEditorUI/Private/ModelingWidgets/SDraggableBox.cpp b/Engine/Plugins/Experimental/MeshModelingToolsetExp/Source/ModelingEditorUI/Private/ModelingWidgets/SDraggableBox.cpp
deleted file mode 100644
diff --git a/Engine/Plugins/Experimental/MeshModelingToolsetExp/Source/ModelingEditorUI/Public/ModelingWidgets/SDraggableBox.h b/Engine/Plugins/Experimental/MeshModelingToolsetExp/Source/ModelingEditorUI/Public/ModelingWidgets/SDraggableBox.h
deleted file mode 100644
diff --git a/Engine/Plugins/Experimental/MeshModelingToolsetExp/Source/ModelingUI/ModelingUI.build.cs b/Engine/Plugins/Experimental/MeshModelingToolsetExp/Source/ModelingUI/ModelingUI.build.cs
new file mode 100644
--- /dev/null
+++ b/Engine/Plugins/Experimental/MeshModelingToolsetExp/Source/ModelingUI/ModelingUI.build.cs
@@ -0,0 +1,40 @@
+// Copyright Epic Games, Inc. All Rights Reserved.
+
+using UnrealBuildTool;
+
+public class ModelingUI : ModuleRules
+{
+	public ModelingUI(ReadOnlyTargetRules Target) : base(Target)
+	{
+		PCHUsage = ModuleRules.PCHUsageMode.UseExplicitOrSharedPCHs;
+		
+		PublicDependencyModuleNames.AddRange(
+			new string[]
+			{
+				"Core",
+				// ... add other public dependencies that you statically link with here ...
+			}
+            );
+			
+		
+		PrivateDependencyModuleNames.AddRange(
+			new string[]
+			{
+                "CoreUObject",
+                "ApplicationCore",
+                "Slate",
+                "SlateCore",
+                "Engine",
+                "InputCore"
+			}
+			);
+
+
+		DynamicallyLoadedModuleNames.AddRange(
+			new string[]
+			{
+				// ... add any modules that your module loads dynamically here ...
+			}
+			);
+	}
+}
diff --git a/Engine/Plugins/Experimental/MeshModelingToolsetExp/Source/ModelingUI/Private/ModelingUIModule.cpp b/Engine/Plugins/Experimental/MeshModelingToolsetExp/Source/ModelingUI/Private/ModelingUIModule.cpp
new file mode 100644
--- /dev/null
+++ b/Engine/Plugins/Experimental/MeshModelingToolsetExp/Source/ModelingUI/Private/ModelingUIModule.cpp
@@ -0,0 +1,17 @@
+// Copyright Epic Games, Inc. All Rights Reserved.
+
+#include "ModelingUIModule.h"
+
+#define LOCTEXT_NAMESPACE "FModelingUIModule"
+
+void FModelingUIModule::StartupModule()
+{
+}
+
+void FModelingUIModule::ShutdownModule()
+{
+}
+
+#undef LOCTEXT_NAMESPACE
+	
+IMPLEMENT_MODULE(FModelingUIModule, ModelingUI)
\ No newline at end of file
diff --git a/Engine/Plugins/Experimental/MeshModelingToolsetExp/Source/ModelingUI/Private/ModelingWidgets/SDraggableBox.cpp b/Engine/Plugins/Experimental/MeshModelingToolsetExp/Source/ModelingUI/Private/ModelingWidgets/SDraggableBox.cpp
new file mode 100644
--- /dev/null
+++ b/Engine/Plugins/Experimental/MeshModelingToolsetExp/Source/ModelingUI/Private/ModelingWidgets/SDraggableBox.cpp
@@ -0,0 +1,180 @@
+// Copyright Epic Games, Inc. All Rights Reserved.
+
+#include "ModelingWidgets/SDraggableBox.h"
+
+#include "Framework/Application/SlateApplication.h"
+#include "Input/DragAndDrop.h" // FDragDropEvent
+#include "Styling/SlateTypes.h" // FWindowStyle
+#include "Widgets/Layout/SBackgroundBlur.h"
+#include "Widgets/Layout/SBorder.h"
+#include "Widgets/SBoxPanel.h" // SVerticalBox
+#include "Widgets/SWindow.h"
+
+
+void SDraggableBoxOverlay::Construct(const FArguments& InArgs)
+{
+	// Depending on whether we're positioning relative to the bottom or top, we'll update the margin of the
+	// draggable box differently.
+	TAttribute<FMargin> DraggableBoxPaddingAttribute;
+	SDraggableBox::FOnDragComplete OnDragCompleteFunction;
+
+	if (InArgs._bPositionRelativeToBottom)
+	{
+		DraggableBoxPaddingAttribute = MakeAttributeLambda([this]()
+		{
+			return FMargin(DraggableBoxPaddingHorizontal, 0, 0, DraggableBoxPaddingVertical);
+		});
+		OnDragCompleteFunction = SDraggableBox::FOnDragComplete::CreateLambda([this](const FVector2D& ScreenSpacePosition)
+		{
+			FVector2D BottomLeftScreenSpace(ScreenSpacePosition.X, ScreenSpacePosition.Y + DraggableBox->GetTickSpaceGeometry().GetAbsoluteSize().Y);
+			FVector2D BottomLeftLocal = ContainingBox->GetTickSpaceGeometry().AbsoluteToLocal(BottomLeftScreenSpace);
+
+			DraggableBoxPaddingHorizontal = BottomLeftLocal.X;
+			DraggableBoxPaddingVertical = ContainingBox->GetTickSpaceGeometry().GetLocalSize().Y - BottomLeftLocal.Y;
+		});
+	}
+	else
+	{
+		DraggableBoxPaddingAttribute = MakeAttributeLambda([this]()
+		{
+			return FMargin(DraggableBoxPaddingHorizontal, DraggableBoxPaddingVertical, 0, 0);
+		});
+		OnDragCompleteFunction = SDraggableBox::FOnDragComplete::CreateLambda([this](const FVector2D& ScreenSpacePosition)
+		{
+			FVector2D Position = ContainingBox->GetTickSpaceGeometry().AbsoluteToLocal(ScreenSpacePosition);
+			DraggableBoxPaddingHorizontal = Position.X;
+			DraggableBoxPaddingVertical = Position.Y;
+		});
+	}
+
+	ChildSlot
+	[
+		SAssignNew(ContainingBox, SVerticalBox)
+		.Visibility(EVisibility::SelfHitTestInvisible)
+
+		+ SVerticalBox::Slot()
+
+		.HAlign(HAlign_Left)
+		.VAlign(InArgs._bPositionRelativeToBottom ? VAlign_Bottom : VAlign_Top)
+		.Padding(DraggableBoxPaddingAttribute)
+		[
+			// The draggable box itself
+			SAssignNew(DraggableBox, SDraggableBox)
+			.OnDragComplete(OnDragCompleteFunction)
+			[
+				InArgs._Content.Widget
+			]
+		]
+	]
+	;
+}
+
+void SDraggableBoxOverlay::SetBoxPosition(float HorizontalPosition, float VerticalPosition)
+{
+	DraggableBoxPaddingHorizontal = HorizontalPosition;
+	DraggableBoxPaddingVertical = VerticalPosition;
+}
+
+// Draggable box:
+
+void SDraggableBox::Construct(const FArguments& InArgs)
+{
+	OnDragComplete = InArgs._OnDragComplete;
+
+	ChildSlot
+	[
+		SAssignNew(InnerWidget, SBackgroundBlur)
+		.BlurStrength(2)
+		.Padding(0.0f)
+		.CornerRadius(FVector4(4.0f, 4.0f, 4.0f, 4.0f))
+		[
+			SNew(SBorder)
+			.Padding(0.0f)
+			[
+				InArgs._Content.Widget
+			]
+		]
+	];
+}
+
+FReply SDraggableBox::OnMouseButtonDown(const FGeometry& MyGeometry, const FPointerEvent& MouseEvent)
+{
+	// Need to remember where within the box we grabbed. We do this here instead of OnDragDetected because 
+	// our mouse can potentially travel some distance before OnDragDetected fires.
+	ScreenSpaceOffsetOfGrab = MouseEvent.GetScreenSpacePosition() - MyGeometry.GetAbsolutePosition();
+
+	return FReply::Handled().DetectDrag(SharedThis(this), EKeys::LeftMouseButton);
+};
+
+FReply SDraggableBox::OnDragDetected(const FGeometry& MyGeometry, const FPointerEvent& MouseEvent)
+{
+	// The drag/drop operation moves a transparent copy of our contents, so we should hide ourselves to make it
+	// look like the contents are moving.
+	SetVisibility(EVisibility::Collapsed);
+
+	// This will be called when the drag/drop operation is done
+	FOnDragComplete OnUIDropped = FOnDragComplete::CreateLambda([this] (const FVector2D& ScreenSpacePosition) {
+		OnDragComplete.ExecuteIfBound(ScreenSpacePosition);
+		SetVisibility(EVisibility::SelfHitTestInvisible);
+	});
+
+	TSharedRef<FDraggableBoxUIDragOperation> DragDropOperation =
+		FDraggableBoxUIDragOperation::New(
+			// Note that it is important to give the inner widget here, since SDraggableBox itself is invisible now
+			InnerWidget.ToSharedRef(),
+			ScreenSpaceOffsetOfGrab,
+			GetDesiredSize(),
+			OnUIDropped
+		);
+
+	return FReply::Handled().BeginDragDrop(DragDropOperation);
+}
+
+// Drag operation:
+
+void FDraggableBoxUIDragOperation::OnDrop(bool bDropWasHandled, const FPointerEvent& MouseEvent)
+{
+	FVector2D ScreenSpaceDropLocation = MouseEvent.GetScreenSpacePosition() - DecoratorOffsetFromCursor;
+
+	OnDragComplete.ExecuteIfBound(ScreenSpaceDropLocation);
+
+	// Call super
+	FDragDropOperation::OnDrop(bDropWasHandled, MouseEvent);
+
+	UIBeingDragged.Reset();
+}
+
+// These are mostly copied from FInViewportUIDragOperation:
+void FDraggableBoxUIDragOperation::OnDragged(const FDragDropEvent& DragDropEvent)
+{
+	// The tab is being dragged. Move the the decorator window to match the cursor position.
+	FVector2D TargetPosition = DragDropEvent.GetScreenSpacePosition() - DecoratorOffsetFromCursor;
+	CursorDecoratorWindow->UpdateMorphTargetShape(FSlateRect(TargetPosition.X, TargetPosition.Y, TargetPosition.X + LastContentSize.X, TargetPosition.Y + LastContentSize.Y));
+	CursorDecoratorWindow->MoveWindowTo(TargetPosition);
+}
+
+TSharedRef<FDraggableBoxUIDragOperation> FDraggableBoxUIDragOperation::New(const TSharedRef<SWidget>& InUIToBeDragged, 
+	const FVector2D InDecoratorOffset, const FVector2D& OwnerAreaSize, SDraggableBox::FOnDragComplete& OnDragComplete)
+{
+	const TSharedRef<FDraggableBoxUIDragOperation> Operation = MakeShareable(new FDraggableBoxUIDragOperation(InUIToBeDragged, InDecoratorOffset, OwnerAreaSize, OnDragComplete));
+	return Operation;
+}
+
+FDraggableBoxUIDragOperation::FDraggableBoxUIDragOperation(const TSharedRef<class SWidget>& InUIToBeDragged, 
+	const FVector2D InDecoratorOffset, const FVector2D& OwnerAreaSize, SDraggableBox::FOnDragComplete& InOnDragComplete)
+	: UIBeingDragged(InUIToBeDragged)
+	, DecoratorOffsetFromCursor(InDecoratorOffset)
+	, LastContentSize(OwnerAreaSize)
+	, OnDragComplete(InOnDragComplete)
+{
+	// Create the decorator window that we will use during this drag and drop to make the user feel like
+	// they are actually dragging a piece of UI.
+	const bool bShowImmediately = true;
+	CursorDecoratorWindow = FSlateApplication::Get().AddWindow(SWindow::MakeStyledCursorDecorator(FAppStyle::Get().GetWidgetStyle<FWindowStyle>("InViewportDecoratorWindow")), bShowImmediately);
+	CursorDecoratorWindow->SetOpacity(0.45f);
+	CursorDecoratorWindow->SetContent
+	(
+		InUIToBeDragged
+	);
+
+}
\ No newline at end of file
diff --git a/Engine/Plugins/Experimental/MeshModelingToolsetExp/Source/ModelingUI/Public/ModelingUIModule.h b/Engine/Plugins/Experimental/MeshModelingToolsetExp/Source/ModelingUI/Public/ModelingUIModule.h
new file mode 100644
--- /dev/null
+++ b/Engine/Plugins/Experimental/MeshModelingToolsetExp/Source/ModelingUI/Public/ModelingUIModule.h
@@ -0,0 +1,18 @@
+// Copyright Epic Games, Inc. All Rights Reserved.
+
+#pragma once
+
+#include "CoreMinimal.h"
+#include "Modules/ModuleManager.h"
+
+class FModelingUIModule : public IModuleInterface
+{
+public:
+
+	/** IModuleInterface implementation */
+	virtual void StartupModule() override;
+	virtual void ShutdownModule() override;
+
+private:
+	TArray<FName> CustomizedClasses;
+};
diff --git a/Engine/Plugins/Experimental/MeshModelingToolsetExp/Source/ModelingUI/Public/ModelingWidgets/SDraggableBox.h b/Engine/Plugins/Experimental/MeshModelingToolsetExp/Source/ModelingUI/Public/ModelingWidgets/SDraggableBox.h
new file mode 100644
--- /dev/null
+++ b/Engine/Plugins/Experimental/MeshModelingToolsetExp/Source/ModelingUI/Public/ModelingWidgets/SDraggableBox.h
@@ -0,0 +1,130 @@
+// Copyright Epic Games, Inc. All Rights Reserved.
+
+#pragma once
+
+#include "CoreMinimal.h"
+
+#include "Widgets/DeclarativeSyntaxSupport.h"
+#include "Widgets/SCompoundWidget.h"
+
+/**
+ * Class that can be used to place a draggable box into a viewport or some other large widget as an
+ * overlay. Just place the widget that you want to be draggable as the contents of SDraggableBoxOverlay.
+ */
+class MODELINGUI_API SDraggableBoxOverlay : public SCompoundWidget
+{
+public:
+
+	SLATE_BEGIN_ARGS(SDraggableBoxOverlay) {}
+
+	// When true, the positioning of the box is relative to the bottom of the widget rather than the top,
+	// so the VerticalPosition in SetBoxPosition is interpreted as distance from the bottom of the
+	// containing widget. The choice here depends on how your widget is generally positioned, for instance
+	// if it is near the bottom, you want this to be true so that making the window smaller does not
+	// immediately clip your box.
+	SLATE_ARGUMENT(bool, bPositionRelativeToBottom)
+
+	SLATE_DEFAULT_SLOT(FArguments, Content)
+	SLATE_END_ARGS()
+
+	void Construct(const FArguments& InArgs);
+
+	/**
+	 * Sets the box position in the overlay. Horizontal position is distance from left, and vertical position
+	 * is distance from bottom or from top depending on whether bPositionRelativeToBottom == true.
+	 */
+	void SetBoxPosition(float HorizontalPosition, float VerticalPosition);
+
+protected:
+
+	TSharedPtr<SWidget> DraggableBox;
+	TSharedPtr<SWidget> ContainingBox;
+
+	float DraggableBoxPaddingHorizontal = 0;
+	float DraggableBoxPaddingVertical = 0;
+};
+
+/**
+ * A widget for the draggable box itself, which requires its parent to handle its positioning in
+ * response to the drag.
+ * 
+ * Users probably shouldn't use this class directly; rather, they should use SDraggableBoxOverlay, 
+ * which will put its contents into a draggable box and properly handle the dragging without the
+ * user having to set it up.
+ */
+class MODELINGUI_API SDraggableBox : public SCompoundWidget
+{
+public:
+
+	DECLARE_DELEGATE_OneParam(FOnDragComplete, const FVector2D& /*ScreenSpacePosition*/);
+
+	SLATE_BEGIN_ARGS(SDraggableBox) {}
+	SLATE_EVENT(FOnDragComplete, OnDragComplete)
+	SLATE_DEFAULT_SLOT(FArguments, Content)
+	SLATE_END_ARGS()
+
+	void Construct(const FArguments& InArgs);
+
+	// SWidget overrides necessary for box to be draggable
+	FReply OnMouseButtonDown(const FGeometry& MyGeometry, const FPointerEvent& MouseEvent) override;
+	FReply OnDragDetected(const FGeometry& MyGeometry, const FPointerEvent& MouseEvent) override;
+
+protected:
+	TSharedPtr<SWidget> InnerWidget;
+	FOnDragComplete OnDragComplete;
+
+	// Remembers the point in the box that we grabbed (on click, rather than when the
+	// drag was confirmed)
+	FVector2D ScreenSpaceOffsetOfGrab;
+};
+
+
+/**
+ * A drag/drop operation used by SDraggableBox, largely modeled on FInViewportUIDragOperation, except
+ * instead of requiring the drop location to be able to handle that particular class to trigger the
+ * OnDragComplete delegate, it just triggers on any drop. This makes it possible to use it in any
+ * viewports, not just the level editor.
+ */
+class MODELINGUI_API FDraggableBoxUIDragOperation : public FDragDropOperation
+{
+public:
+
+	DRAG_DROP_OPERATOR_TYPE(FDraggableBoxUIDragOperation, FDragDropOperation)
+
+	// FDragDropOperation
+	virtual void OnDrop(bool bDropWasHandled, const FPointerEvent& MouseEvent) override;
+	virtual void OnDragged(const FDragDropEvent& DragDropEvent) override;
+
+	/**
+	 * Create this Drag and Drop Content
+	 *
+	 * @param InUIToBeDragged	  The UI being dragged
+	 * @param InDecoratorOffset   Where within the UI we grabbed, so we're not dragging by the upper left of the UI.
+	 * @param OwnerAreaSize       Size of the DockArea at the time when we start dragging.
+	 * @param OnDragComplete      Delegate to call when dropped. Gets passed in the screen space location of the top left corner.
+	 *
+	 * @return a new FDockingDragOperation
+	 */
+	static TSharedRef<FDraggableBoxUIDragOperation> New(const TSharedRef<class SWidget>& InUIToBeDragged, const FVector2D InDecoratorOffset, 
+		const FVector2D& OwnerAreaSize, SDraggableBox::FOnDragComplete& OnDragComplete);
+
+	virtual ~FDraggableBoxUIDragOperation() {}
+
+protected:
+	/** The constructor is protected, so that this class can only be instantiated as a shared pointer. */
+	FDraggableBoxUIDragOperation(const TSharedRef<class SWidget>& InUIToBeDragged, const FVector2D InDecoratorOffsetFromCursor, 
+		const FVector2D& OwnerAreaSize, SDraggableBox::FOnDragComplete& OnDragCompletee);
+
+	/**
+	 * Shared pointer to our contents, which are a weak pointer.
+	 */
+	TSharedPtr<class SWidget> UIBeingDragged;
+
+	// Screen space offset of grab location from the top left corner of the rectangle being dragged.
+	FVector2D DecoratorOffsetFromCursor;
+
+	// Size of rectangle being dragged
+	FVector2D LastContentSize;
+
+	SDraggableBox::FOnDragComplete OnDragComplete;
+};
