Clearer tooltip on min group size in the generate polygroups tool

diff --git a/Engine/Plugins/Experimental/MeshModelingToolsetExp/Source/MeshModelingToolsExp/Public/ConvertToPolygonsTool.h b/Engine/Plugins/Experimental/MeshModelingToolsetExp/Source/MeshModelingToolsExp/Public/ConvertToPolygonsTool.h
--- a/Engine/Plugins/Experimental/MeshModelingToolsetExp/Source/MeshModelingToolsExp/Public/ConvertToPolygonsTool.h
+++ b/Engine/Plugins/Experimental/MeshModelingToolsetExp/Source/MeshModelingToolsExp/Public/ConvertToPolygonsTool.h
@@ -114,7 +114,7 @@ public:
 	bool bRespectHardNormals = false;
 
 
-	/** group filtering */
+	/** Minimum number of triangles to include in a group. Any group containing fewer triangles will be merged with an adjacent group (if possible). */
 	UPROPERTY(EditAnywhere, Category = Filtering, meta = (UIMin = "1", UIMax = "100", ClampMin = "1", ClampMax = "10000", EditCondition = "ConversionMode != EConvertToPolygonsMode::CopyFromLayer"))
 	int32 MinGroupSize = 2;
 	
