Added ability to isolate component selections within the skin paint tool.

diff --git a/Engine/Plugins/Experimental/Animation/SkeletalMeshModelingTools/Source/SkeletalMeshModelingTools/Private/DetailCustomization/SkinWeightDetailCustomization.cpp b/Engine/Plugins/Experimental/Animation/SkeletalMeshModelingTools/Source/SkeletalMeshModelingTools/Private/DetailCustomization/SkinWeightDetailCustomization.cpp
--- a/Engine/Plugins/Experimental/Animation/SkeletalMeshModelingTools/Source/SkeletalMeshModelingTools/Private/DetailCustomization/SkinWeightDetailCustomization.cpp
+++ b/Engine/Plugins/Experimental/Animation/SkeletalMeshModelingTools/Source/SkeletalMeshModelingTools/Private/DetailCustomization/SkinWeightDetailCustomization.cpp
@@ -146,6 +146,9 @@ void FSkinWeightDetailCustomization::CustomizeDetails(IDetailLayoutBuilder& Deta
 
 void FSkinWeightDetailCustomization::AddBrushUI(IDetailLayoutBuilder& DetailBuilder)
 {
+	// pointer to capture
+	TObjectPtr<USkinWeightsPaintToolProperties> ToolSettings = SkinToolSettings.Pin().Get();
+	
 	// custom display of falloff mode as segmented toggle buttons
 	IDetailCategoryBuilder& BrushCategory = DetailBuilder.EditCategory("Brush", FText::GetEmpty(), ECategoryPriority::Important);
 
@@ -165,13 +168,13 @@ void FSkinWeightDetailCustomization::AddBrushUI(IDetailLayoutBuilder& DetailBuil
 				"This command operates on the selected bone(s) and selected vertices.\n"
 				"If no bones are selected, ALL bones are considered.\n"
 				"If no vertices are selected, ALL vertices are considered."))
-			.Value_Lambda([this]()
+			.Value_Lambda([ToolSettings]()
 			{
-				return SkinToolSettings->BrushMode;
+				return ToolSettings->BrushMode;
 			})
-			.OnValueChanged_Lambda([this](EWeightEditOperation Mode)
+			.OnValueChanged_Lambda([ToolSettings](EWeightEditOperation Mode)
 			{
-				SkinToolSettings->SetBrushMode(Mode);
+				ToolSettings->SetBrushMode(Mode);
 			})
 			+SSegmentedControl<EWeightEditOperation>::Slot(EWeightEditOperation::Add)
 			.Text(LOCTEXT("BrushAddMode", "Add"))
@@ -195,13 +198,13 @@ void FSkinWeightDetailCustomization::AddBrushUI(IDetailLayoutBuilder& DetailBuil
 			.ToolTipText(LOCTEXT("BrushFalloffModeTooltip",
 					"Surface: falloff is based on the distance along the surface from the brush center to nearby connected vertices.\n"
 					"Volume: falloff is based on the straight-line distance from the brush center to surrounding vertices.\n"))
-			.Value_Lambda([this]()
+			.Value_Lambda([ToolSettings]()
 			{
-				return SkinToolSettings->GetBrushConfig().FalloffMode;
+				return ToolSettings->GetBrushConfig().FalloffMode;
 			})
-			.OnValueChanged_Lambda([this](EWeightBrushFalloffMode Mode)
+			.OnValueChanged_Lambda([ToolSettings](EWeightBrushFalloffMode Mode)
 			{
-				SkinToolSettings->SetFalloffMode(Mode);
+				ToolSettings->SetFalloffMode(Mode);
 			})
 			+SSegmentedControl<EWeightBrushFalloffMode>::Slot(EWeightBrushFalloffMode::Surface)
 			.Text(LOCTEXT("SurfaceMode", "Surface"))
@@ -225,20 +228,20 @@ void FSkinWeightDetailCustomization::AddBrushUI(IDetailLayoutBuilder& DetailBuil
 		.MaxSliderValue(20.f)
 		.Value(10.0f)
 		.SupportDynamicSliderMaxValue(true)
-		.Value_Lambda([this]()
+		.Value_Lambda([ToolSettings]()
 		{
-			return SkinToolSettings->GetBrushConfig().Radius;
+			return ToolSettings->GetBrushConfig().Radius;
 		})
-		.OnValueChanged_Lambda([this](float NewValue)
+		.OnValueChanged_Lambda([ToolSettings](float NewValue)
 		{
-			SkinToolSettings->BrushRadius = NewValue;
-			SkinToolSettings->GetBrushConfig().Radius = NewValue;
+			ToolSettings->BrushRadius = NewValue;
+			ToolSettings->GetBrushConfig().Radius = NewValue;
 			FPropertyChangedEvent PropertyChangedEvent(UBrushBaseProperties::StaticClass()->FindPropertyByName(GET_MEMBER_NAME_CHECKED(UBrushBaseProperties, BrushRadius)));
-			SkinToolSettings->PostEditChangeProperty(PropertyChangedEvent);
+			ToolSettings->PostEditChangeProperty(PropertyChangedEvent);
 		})
-		.OnValueCommitted_Lambda([this](float NewValue, ETextCommit::Type CommitType)
+		.OnValueCommitted_Lambda([ToolSettings](float NewValue, ETextCommit::Type CommitType)
 		{
-			SkinToolSettings->SaveConfig();
+			ToolSettings->SaveConfig();
 		})
 	];
 
@@ -258,20 +261,20 @@ void FSkinWeightDetailCustomization::AddBrushUI(IDetailLayoutBuilder& DetailBuil
 		.MaxSliderValue(1.f)
 		.Value(1.0f)
 		.SupportDynamicSliderMaxValue(true)
-		.Value_Lambda([this]()
+		.Value_Lambda([ToolSettings]()
 		{
-			return SkinToolSettings->GetBrushConfig().Strength;
+			return ToolSettings->GetBrushConfig().Strength;
 		})
-		.OnValueChanged_Lambda([this](float NewValue)
+		.OnValueChanged_Lambda([ToolSettings](float NewValue)
 		{
-			SkinToolSettings->BrushStrength = NewValue;
-			SkinToolSettings->GetBrushConfig().Strength = NewValue;
+			ToolSettings->BrushStrength = NewValue;
+			ToolSettings->GetBrushConfig().Strength = NewValue;
 			FPropertyChangedEvent PropertyChangedEvent(UBrushBaseProperties::StaticClass()->FindPropertyByName(GET_MEMBER_NAME_CHECKED(UBrushBaseProperties, BrushStrength)));
-			SkinToolSettings->PostEditChangeProperty(PropertyChangedEvent);
+			ToolSettings->PostEditChangeProperty(PropertyChangedEvent);
 		})
-		.OnValueCommitted_Lambda([this](float NewValue, ETextCommit::Type CommitType)
+		.OnValueCommitted_Lambda([ToolSettings](float NewValue, ETextCommit::Type CommitType)
 		{
-			SkinToolSettings->SaveConfig();
+			ToolSettings->SaveConfig();
 		})
 	];
 
@@ -288,26 +291,30 @@ void FSkinWeightDetailCustomization::AddBrushUI(IDetailLayoutBuilder& DetailBuil
 		SNew(SSpinBox<float>)
 		.MinValue(0.f)
 		.MaxValue(1.f)
-		.Value_Lambda([this]()
+		.Value_Lambda([ToolSettings]()
 		{
-			return SkinToolSettings->GetBrushConfig().Falloff;
+			return ToolSettings->GetBrushConfig().Falloff;
 		})
-		.OnValueChanged_Lambda([this](float NewValue)
+		.OnValueChanged_Lambda([ToolSettings](float NewValue)
 		{
-			SkinToolSettings->BrushFalloffAmount = NewValue;
-			SkinToolSettings->GetBrushConfig().Falloff = NewValue;
+			ToolSettings->BrushFalloffAmount = NewValue;
+			ToolSettings->GetBrushConfig().Falloff = NewValue;
 			FPropertyChangedEvent PropertyChangedEvent(UBrushBaseProperties::StaticClass()->FindPropertyByName(GET_MEMBER_NAME_CHECKED(UBrushBaseProperties, BrushFalloffAmount)));
-			SkinToolSettings->PostEditChangeProperty(PropertyChangedEvent);
+			ToolSettings->PostEditChangeProperty(PropertyChangedEvent);
 		})
-		.OnValueCommitted_Lambda([this](float NewValue, ETextCommit::Type CommitType)
+		.OnValueCommitted_Lambda([ToolSettings](float NewValue, ETextCommit::Type CommitType)
 		{
-			SkinToolSettings->SaveConfig();
+			ToolSettings->SaveConfig();
 		})
 	];
 }
 
 void FSkinWeightDetailCustomization::AddSelectionUI(IDetailLayoutBuilder& DetailBuilder)
 {
+	// pointers to capture
+	TObjectPtr<USkinWeightsPaintToolProperties> ToolSettings = SkinToolSettings.Get();
+	TObjectPtr<USkinWeightsPaintTool> Tool = SkinToolSettings->WeightTool;
+	
 	// custom display of weight editing tools
 	IDetailCategoryBuilder& EditSelectionCategory = DetailBuilder.EditCategory("Edit Selection", FText::GetEmpty(), ECategoryPriority::Important);
 	EditSelectionCategory.InitiallyCollapsed(true);
@@ -319,22 +326,21 @@ void FSkinWeightDetailCustomization::AddSelectionUI(IDetailLayoutBuilder& Detail
 
 	ToolbarBuilder.BeginSection("SelectionFilter");
 	ToolbarBuilder.BeginBlockGroup();
-
-	TObjectPtr<USkinWeightsPaintToolProperties> ToolProperties = SkinToolSettings.Get();
-	auto AddToggleButtonForBool = [&ToolbarBuilder, ToolProperties](EComponentSelectionMode Mode, const FText& Label, const FText& Tooltip, const FName IconName)
+	
+	auto AddToggleButtonForBool = [&ToolbarBuilder, ToolSettings](EComponentSelectionMode Mode, const FText& Label, const FText& Tooltip, const FName IconName)
 	{
 		ToolbarBuilder.AddToolBarButton(FUIAction(
-		FExecuteAction::CreateLambda([ToolProperties, Mode]()
+		FExecuteAction::CreateLambda([ToolSettings, Mode]()
 		{
-			ToolProperties->SetComponentMode(Mode);
+			ToolSettings->SetComponentMode(Mode);
 		}),
-		FCanExecuteAction::CreateLambda([ToolProperties]()
+		FCanExecuteAction::CreateLambda([ToolSettings]()
 		{
-			return ToolProperties->EditingMode == EWeightEditMode::Mesh;
+			return ToolSettings->EditingMode == EWeightEditMode::Mesh;
 		}),
-		FIsActionChecked::CreateLambda([ToolProperties, Mode]()
+		FIsActionChecked::CreateLambda([ToolSettings, Mode]()
 		{
-			return ToolProperties->ComponentSelectionMode == Mode;
+			return ToolSettings->ComponentSelectionMode == Mode;
 		})),
 		NAME_None,	// Extension hook
 		Label,		// Label
@@ -366,64 +372,122 @@ void FSkinWeightDetailCustomization::AddSelectionUI(IDetailLayoutBuilder& Detail
 	EditSelectionCategory.AddCustomRow(LOCTEXT("EditSelectionRow", "Edit Selection"), false)
 	.WholeRowContent()
 	[
-		SNew(SHorizontalBox)
-
-		+SHorizontalBox::Slot()
+		SNew(SVerticalBox)
+		+SVerticalBox::Slot()
+		.AutoHeight()
 		[
 			SNew(SHorizontalBox)
+
 			+SHorizontalBox::Slot()
-			.HAlign(HAlign_Left)
-			[
-				ToolbarBuilder.MakeWidget()
-			]
-		]
-		
-		+SHorizontalBox::Slot()
-		.Padding(2.f, WeightEditVerticalPadding)
-		[
-			SNew(SButton)
 			.HAlign(HAlign_Center)
 			.VAlign(VAlign_Center)
-			.Text(LOCTEXT("GrowSelectionButtonLabel", "Grow"))
-			.ToolTipText(LOCTEXT("GrowSelectionTooltip",
-					"Grow the current selection by adding connected neighbors to current selection.\n"))
-			.OnClicked_Lambda([this]()
-			{
-				SkinToolSettings->WeightTool->GrowSelection();
-				return FReply::Handled();
-			})
-		]
+			.Padding(2.0f, WeightEditVerticalPadding)
+			[
+				SNew(SHorizontalBox)
+				+SHorizontalBox::Slot()
+				[
+					ToolbarBuilder.MakeWidget()
+				]
 
-		+SHorizontalBox::Slot()
-		.Padding(2.f, WeightEditVerticalPadding)
-		[
-			SNew(SButton)
-			.HAlign(HAlign_Center)
-			.VAlign(VAlign_Center)
-			.Text(LOCTEXT("ShrinkSelectionButtonLabel", "Shrink"))
-			.ToolTipText(LOCTEXT("ShrinkSelectionTooltip",
-					"Shrink the current selection by removing components on the border of the current selection.\n"))
-			.OnClicked_Lambda([this]()
-			{
-				SkinToolSettings->WeightTool->ShrinkSelection();
-				return FReply::Handled();
-			})
+				+SHorizontalBox::Slot()
+				.AutoWidth()
+				.VAlign(VAlign_Center)
+				[
+					SNew(SCheckBox)
+					.Style(FAppStyle::Get(), "ToggleButtonCheckbox")
+					.HAlign(HAlign_Center)
+					.ToolTipText(LOCTEXT("IsolateSelectedTooltip",
+							"Shows only the selected faces in the viewport.\n"
+							"Weight editing operations will not affect hidden vertices.\n "))
+					.IsEnabled_Lambda([Tool]()
+					{
+						const bool bHasSelection = Tool->IsAnyComponentSelected();
+						const bool bAlreadyIsolatingSelection = Tool->IsSelectionIsolated();
+						return bHasSelection ||  bAlreadyIsolatingSelection;
+					})
+					.IsChecked_Lambda([Tool]()
+					{
+						const bool bAlreadyIsolatingSelection = Tool->IsSelectionIsolated();
+						return bAlreadyIsolatingSelection ? ECheckBoxState::Checked : ECheckBoxState::Unchecked;
+					})
+					.OnCheckStateChanged_Lambda([Tool](ECheckBoxState InCheckBoxState)
+					{
+						if (InCheckBoxState == ECheckBoxState::Checked)
+						{
+							Tool->SetIsolateSelected(true);	
+						}
+						else
+						{
+							Tool->SetIsolateSelected(false);	
+						}
+					})
+					[
+						SNew(STextBlock)
+						.Text_Lambda([Tool]()
+						{
+							if (Tool->IsSelectionIsolated())
+							{
+								return LOCTEXT("ShowAllButtonLabel", "Show All");
+							}
+								
+							return LOCTEXT("IsolateButtonLabel", "Isolate Selected");
+						})
+					]
+				]
+			]
 		]
 
-		+SHorizontalBox::Slot()
-		.Padding(2.f, WeightEditVerticalPadding)
+		+SVerticalBox::Slot()
 		[
-			SNew(SButton)
-			.HAlign(HAlign_Center)
-			.VAlign(VAlign_Center)
-			.Text(LOCTEXT("FloodSelectionButtonLabel", "Flood"))
-			.ToolTipText(LOCTEXT("FloodSelectionTooltip",
-					"Flood the current selection by adding all connected components to the current selection.\n"))
-			.OnClicked_Lambda([this]()
-			{
-				SkinToolSettings->WeightTool->FloodSelection();
-				return FReply::Handled();
-			})
+			SNew(SHorizontalBox)
+		
+			+SHorizontalBox::Slot()
+			.Padding(2.f, WeightEditVerticalPadding)
+			[
+				SNew(SButton)
+				.HAlign(HAlign_Center)
+				.VAlign(VAlign_Center)
+				.Text(LOCTEXT("GrowSelectionButtonLabel", "Grow"))
+				.ToolTipText(LOCTEXT("GrowSelectionTooltip",
+						"Grow the current selection by adding connected neighbors to current selection.\n"))
+				.OnClicked_Lambda([Tool]()
+				{
+					Tool->GrowSelection();
+					return FReply::Handled();
+				})
+			]
+
+			+SHorizontalBox::Slot()
+			.Padding(2.f, WeightEditVerticalPadding)
+			[
+				SNew(SButton)
+				.HAlign(HAlign_Center)
+				.VAlign(VAlign_Center)
+				.Text(LOCTEXT("ShrinkSelectionButtonLabel", "Shrink"))
+				.ToolTipText(LOCTEXT("ShrinkSelectionTooltip",
+						"Shrink the current selection by removing components on the border of the current selection.\n"))
+				.OnClicked_Lambda([Tool]()
+				{
+					Tool->ShrinkSelection();
+					return FReply::Handled();
+				})
+			]
+
+			+SHorizontalBox::Slot()
+			.Padding(2.f, WeightEditVerticalPadding)
+			[
+				SNew(SButton)
+				.HAlign(HAlign_Center)
+				.VAlign(VAlign_Center)
+				.Text(LOCTEXT("FloodSelectionButtonLabel", "Flood"))
+				.ToolTipText(LOCTEXT("FloodSelectionTooltip",
+						"Flood the current selection by adding all connected components to the current selection.\n"))
+				.OnClicked_Lambda([Tool]()
+				{
+					Tool->FloodSelection();
+					return FReply::Handled();
+				})
+			]
 		]
 	];
 
@@ -461,17 +525,17 @@ void FSkinWeightDetailCustomization::AddSelectionUI(IDetailLayoutBuilder& Detail
 				.MaxSliderValue(1.f)
 				.Value(1.0f)
 				.SupportDynamicSliderMaxValue(true)
-				.Value_Lambda([this]()
+				.Value_Lambda([ToolSettings]()
 				{
-					return SkinToolSettings->FloodValue;
+					return ToolSettings->FloodValue;
 				})
-				.OnValueChanged_Lambda([this](float NewValue)
+				.OnValueChanged_Lambda([ToolSettings](float NewValue)
 				{
-					SkinToolSettings->FloodValue = NewValue;
+					ToolSettings->FloodValue = NewValue;
 				})
-				.OnValueCommitted_Lambda([this](float NewValue, ETextCommit::Type CommitType)
+				.OnValueCommitted_Lambda([ToolSettings](float NewValue, ETextCommit::Type CommitType)
 				{
-					SkinToolSettings->SaveConfig();
+					ToolSettings->SaveConfig();
 				})
 			]
 		]
@@ -500,9 +564,9 @@ void FSkinWeightDetailCustomization::AddSelectionUI(IDetailLayoutBuilder& Detail
 					.HAlign(HAlign_Center)
 					.Text(LOCTEXT("AddWeightsButtonLabel", "Add"))
 					.ToolTipText(LOCTEXT("AddOpTooltip", "Add: adds the Flood Amount value to the current weight."))
-					.OnClicked_Lambda([this]()
+					.OnClicked_Lambda([ToolSettings, Tool]()
 					{
-						SkinToolSettings->WeightTool->FloodWeights(SkinToolSettings->FloodValue, EWeightEditOperation::Add);
+						Tool->FloodWeights(ToolSettings->FloodValue, EWeightEditOperation::Add);
 						return FReply::Handled();
 					})
 				]
@@ -512,9 +576,9 @@ void FSkinWeightDetailCustomization::AddSelectionUI(IDetailLayoutBuilder& Detail
 					.HAlign(HAlign_Center)
 					.Text(LOCTEXT("ReplaceWeightsButtonLabel", "Replace"))
 					.ToolTipText(LOCTEXT("ReplaceOpTooltip", "Replace: applies the current weight minus the flood amount to the new weight."))
-					.OnClicked_Lambda([this]()
+					.OnClicked_Lambda([ToolSettings, Tool]()
 					{
-						SkinToolSettings->WeightTool->FloodWeights(SkinToolSettings->FloodValue, EWeightEditOperation::Replace);
+						Tool->FloodWeights(ToolSettings->FloodValue, EWeightEditOperation::Replace);
 						return FReply::Handled();
 					})
 				]
@@ -524,9 +588,9 @@ void FSkinWeightDetailCustomization::AddSelectionUI(IDetailLayoutBuilder& Detail
 					.HAlign(HAlign_Center)
 					.Text(LOCTEXT("MultiplyeightsButtonLabel", "Multiply"))
 					.ToolTipText(LOCTEXT("MultiplyOpTooltip", "Multiply: applies the current weight multiplied by the flood amount to the new weight."))
-					.OnClicked_Lambda([this]()
+					.OnClicked_Lambda([ToolSettings, Tool]()
 					{
-						SkinToolSettings->WeightTool->FloodWeights(SkinToolSettings->FloodValue, EWeightEditOperation::Multiply);
+						Tool->FloodWeights(ToolSettings->FloodValue, EWeightEditOperation::Multiply);
 						return FReply::Handled();
 					})
 				]
@@ -536,9 +600,9 @@ void FSkinWeightDetailCustomization::AddSelectionUI(IDetailLayoutBuilder& Detail
 					.HAlign(HAlign_Center)
 					.Text(LOCTEXT("RelaxWeightsButtonLabel", "Relax"))
 					.ToolTipText(LOCTEXT("RelaxeOpTooltip", "Relax: applies the average of the connected (by edge) vertex weights to the new vertex weight, scaled by the flood amount."))
-					.OnClicked_Lambda([this]()
+					.OnClicked_Lambda([ToolSettings, Tool]()
 					{
-						SkinToolSettings->WeightTool->FloodWeights(SkinToolSettings->FloodValue, EWeightEditOperation::Relax);
+						Tool->FloodWeights(ToolSettings->FloodValue, EWeightEditOperation::Relax);
 						return FReply::Handled();
 					})
 				]
@@ -564,9 +628,9 @@ void FSkinWeightDetailCustomization::AddSelectionUI(IDetailLayoutBuilder& Detail
 					"This command operates on the selected bone(s) and selected vertices.\n "
 					"If no bones are selected, ALL bone weights are considered.\n "
 					"If no vertices are selected, ALL vertices are considered."))
-			.OnClicked_Lambda([this]()
+			.OnClicked_Lambda([Tool]()
 			{
-				SkinToolSettings->WeightTool->AverageWeights();
+				Tool->AverageWeights();
 				return FReply::Handled();
 			})
 		]
@@ -582,13 +646,13 @@ void FSkinWeightDetailCustomization::AddSelectionUI(IDetailLayoutBuilder& Detail
 					"Forces the weights on the selected vertices to sum to 1.\n"
 					"This command operates on the selected vertices.\n "
 					"If no vertices are selected, ALL vertices are considered."))
-			.IsEnabled_Lambda([this]()
+			.IsEnabled_Lambda([ToolSettings]()
 			{
-				return SkinToolSettings->EditingMode == EWeightEditMode::Mesh;
+				return ToolSettings->EditingMode == EWeightEditMode::Mesh;
 			})
-			.OnClicked_Lambda([this]()
+			.OnClicked_Lambda([Tool]()
 			{
-				SkinToolSettings->WeightTool->NormalizeWeights();
+				Tool->NormalizeWeights();
 				return FReply::Handled();
 			})
 		]
@@ -620,17 +684,17 @@ void FSkinWeightDetailCustomization::AddSelectionUI(IDetailLayoutBuilder& Detail
 				SNew(SSpinBox<float>)
 				.MinValue(0.f)
 				.MaxValue(1.f)
-				.Value_Lambda([this]()
+				.Value_Lambda([ToolSettings]()
 				{
-					return SkinToolSettings->PruneValue;
+					return ToolSettings->PruneValue;
 				})
-				.OnValueChanged_Lambda([this](float NewValue)
+				.OnValueChanged_Lambda([ToolSettings](float NewValue)
 				{
-					SkinToolSettings->PruneValue = NewValue;
+					ToolSettings->PruneValue = NewValue;
 				})
-				.OnValueCommitted_Lambda([this](float NewValue, ETextCommit::Type CommitType)
+				.OnValueCommitted_Lambda([ToolSettings](float NewValue, ETextCommit::Type CommitType)
 				{
-					SkinToolSettings->SaveConfig();
+					ToolSettings->SaveConfig();
 				})
 			]
 
@@ -688,13 +752,13 @@ void FSkinWeightDetailCustomization::AddSelectionUI(IDetailLayoutBuilder& Detail
 						"X: copies weights across the YZ plane.\n"
 						"Y: copies weights across the XZ plane.\n"
 						"Z: copies weights across the XY plane."))
-					.Value_Lambda([this]()
+					.Value_Lambda([ToolSettings]()
 					{
-						return SkinToolSettings->MirrorAxis;
+						return ToolSettings->MirrorAxis;
 					})
-					.OnValueChanged_Lambda([this](EAxis::Type Mode)
+					.OnValueChanged_Lambda([ToolSettings](EAxis::Type Mode)
 					{
-						SkinToolSettings->MirrorAxis = Mode;
+						ToolSettings->MirrorAxis = Mode;
 					})
 					+ SSegmentedControl<EAxis::Type>::Slot(EAxis::X)
 					.Text(LOCTEXT("MirrorXLabel", "X"))
@@ -708,13 +772,13 @@ void FSkinWeightDetailCustomization::AddSelectionUI(IDetailLayoutBuilder& Detail
 				[
 					SNew(SSegmentedControl<EMirrorDirection>)
 					.ToolTipText(LOCTEXT("MirrorDirectionTooltip", "The direction that determines what side of the plane to copy weights from."))
-					.Value_Lambda([this]()
+					.Value_Lambda([ToolSettings]()
 					{
-						return SkinToolSettings->MirrorDirection;
+						return ToolSettings->MirrorDirection;
 					})
-					.OnValueChanged_Lambda([this](EMirrorDirection Mode)
+					.OnValueChanged_Lambda([ToolSettings](EMirrorDirection Mode)
 					{
-						SkinToolSettings->MirrorDirection = Mode;
+						ToolSettings->MirrorDirection = Mode;
 					})
 					+ SSegmentedControl<EMirrorDirection>::Slot(EMirrorDirection::PositiveToNegative)
 					.Text(LOCTEXT("MirrorPosToNegLabel", "+ to -"))
@@ -737,9 +801,9 @@ void FSkinWeightDetailCustomization::AddSelectionUI(IDetailLayoutBuilder& Detail
 					"This command operates on the selected bone(s) and selected vertices.\n "
 					"If no bones are selected, ALL bone weights are considered.\n "
 					"If no vertices are selected, ALL vertices are considered."))
-				.OnClicked_Lambda([this]()
+				.OnClicked_Lambda([ToolSettings, Tool]()
 				{
-					SkinToolSettings->WeightTool->MirrorWeights(SkinToolSettings->MirrorAxis, SkinToolSettings->MirrorDirection);
+					Tool->MirrorWeights(ToolSettings->MirrorAxis, ToolSettings->MirrorDirection);
 					return FReply::Handled();
 				})
 			]
diff --git a/Engine/Plugins/Experimental/MeshModelingToolsetExp/Source/MeshModelingToolsEditorOnlyExp/Private/SkeletalMesh/SkinWeightsPaintTool.cpp b/Engine/Plugins/Experimental/MeshModelingToolsetExp/Source/MeshModelingToolsEditorOnlyExp/Private/SkeletalMesh/SkinWeightsPaintTool.cpp
--- a/Engine/Plugins/Experimental/MeshModelingToolsetExp/Source/MeshModelingToolsEditorOnlyExp/Private/SkeletalMesh/SkinWeightsPaintTool.cpp
+++ b/Engine/Plugins/Experimental/MeshModelingToolsetExp/Source/MeshModelingToolsEditorOnlyExp/Private/SkeletalMesh/SkinWeightsPaintTool.cpp
@@ -28,12 +28,16 @@
 #include "Spatial/PointSetHashTable.h"
 #include "Operations/SmoothBoneWeights.h"
 #include "ContextObjectStore.h"
+#include "DynamicMeshToMeshDescription.h"
+#include "DynamicSubmesh3.h"
+#include "MeshDescriptionToDynamicMesh.h"
 #include "SkeletalDebugRendering.h"
 #include "Editor/Persona/Public/IPersonaEditorModeManager.h"
 #include "Editor/Persona/Public/PersonaModule.h"
 #include "Preferences/PersonaOptions.h"
 #include "PreviewProfileController.h"
 #include "Animation/SkinWeightProfile.h"
+#include "AnimationRuntime.h"
 
 #include UE_INLINE_GENERATED_CPP_BY_NAME(SkinWeightsPaintTool)
 
@@ -262,10 +266,10 @@
 	return 0.0f;
 }
 
-void FSkinToolDeformer::Initialize(const USkeletalMeshComponent* SkeletalMeshComponent, const FMeshDescription* Mesh)
+void FSkinToolDeformer::Initialize(const USkeletalMeshComponent* InSkelMeshComponent, const FMeshDescription* InMeshDescription)
 {
 	// get all bone transforms in the reference pose store a copy in component space
-	Component = SkeletalMeshComponent;
+	Component = InSkelMeshComponent;
 	const FReferenceSkeleton& RefSkeleton = Component->GetSkeletalMeshAsset()->GetRefSkeleton();
 	const TArray<FTransform> &LocalSpaceBoneTransforms = RefSkeleton.GetRefBonePose();
 	const int32 NumBones = LocalSpaceBoneTransforms.Num();
@@ -296,14 +300,14 @@
 	}
 
 	// store reference pose vertex positions
-	const TArrayView<const UE::Math::TVector<float>> VertexPositions = Mesh->GetVertexPositions().GetRawArray();
+	const TArrayView<const FVector3f> VertexPositions = InMeshDescription->GetVertexPositions().GetRawArray();
 	RefPoseVertexPositions = VertexPositions;
 
 	// set all vertices to be updated on first tick
 	SetAllVerticesToBeUpdated();
 
 	// record "prev" bone transforms to detect change in pose
-	PrevBoneTransforms = Component->GetComponentSpaceTransforms();
+	PreviousPoseComponentSpace = Component->GetComponentSpaceTransforms();
 }
 
 void FSkinToolDeformer::SetAllVerticesToBeUpdated()
@@ -315,23 +319,35 @@
 	}
 }
 
-void FSkinToolDeformer::UpdateVertexDeformation(USkinWeightsPaintTool* Tool)
+void FSkinToolDeformer::SetToRefPose(USkinWeightsPaintTool* Tool)
 {
+	// get ref pose
+	const FReferenceSkeleton& RefSkeleton = Component->GetSkeletalMeshAsset()->GetRefSkeleton();
+	const TArray<FTransform>& RefPoseLocalSpace = RefSkeleton.GetRefBonePose();
+	// convert to global space and store in current pose
+	FAnimationRuntime::FillUpComponentSpaceTransforms(RefSkeleton, RefPoseLocalSpace, RefPoseComponentSpace);
+	// update mesh to new pose
+	UpdateVertexDeformation(Tool, RefPoseComponentSpace);
+}
+
+void FSkinToolDeformer::UpdateVertexDeformation(
+	USkinWeightsPaintTool* Tool,
+	const TArray<FTransform>& PoseComponentSpace)
+{
 	TRACE_CPUPROFILER_EVENT_SCOPE(SkinTool::UpdateDeformationTotal);
 
 	// if no weights have been modified, we must check for a modified pose which requires re-calculation of skinning
 	if (VerticesWithModifiedWeights.IsEmpty())
 	{
-		const TArray<FTransform>& CurrentBoneTransforms = Component->GetComponentSpaceTransforms();
-		for (int32 BoneIndex=0; BoneIndex<CurrentBoneTransforms.Num(); ++BoneIndex)
+		for (int32 BoneIndex=0; BoneIndex<PoseComponentSpace.Num(); ++BoneIndex)
 		{
 			if (!Tool->Weights.IsBoneWeighted[BoneIndex])
 			{
 				continue;
 			}
 			
-			const FTransform& CurrentBoneTransform = CurrentBoneTransforms[BoneIndex];
-			const FTransform& PrevBoneTransform = PrevBoneTransforms[BoneIndex];
+			const FTransform& CurrentBoneTransform = PoseComponentSpace[BoneIndex];
+			const FTransform& PrevBoneTransform = PreviousPoseComponentSpace[BoneIndex];
 			if (!CurrentBoneTransform.Equals(PrevBoneTransform))
 			{
 				SetAllVerticesToBeUpdated();
@@ -348,20 +364,23 @@
 	// update vertex positions
 	UPreviewMesh* PreviewMesh = Tool->PreviewMesh;
 	const TArray<VertexWeights>& CurrentWeights = Tool->Weights.CurrentWeights;
-	PreviewMesh->DeferredEditMesh([this, &CurrentWeights](FDynamicMesh3& Mesh)
+	PreviewMesh->DeferredEditMesh([this, &CurrentWeights, &PoseComponentSpace](FDynamicMesh3& Mesh)
 	{
 		TRACE_CPUPROFILER_EVENT_SCOPE(SkinTool::UpdateDeformation);
-		const TArray<FTransform>& CurrentBoneTransforms = Component->GetComponentSpaceTransforms();
 		const TArray<int32> VertexIndices = VerticesWithModifiedWeights.Array();
 		
-		ParallelFor( VerticesWithModifiedWeights.Num(), [this, &VertexIndices, &Mesh, &CurrentBoneTransforms, &CurrentWeights](int32 Index)
+		ParallelFor( VerticesWithModifiedWeights.Num(), [this, &VertexIndices, &Mesh, &PoseComponentSpace, &CurrentWeights](int32 Index)
 		{
 			const int32 VertexID = VertexIndices[Index];
 			FVector VertexNewPosition = FVector::ZeroVector;
 			const VertexWeights& VertexPerBoneData = CurrentWeights[VertexID];
 			for (const FVertexBoneWeight& VertexData : VertexPerBoneData)
 			{
-				const FTransform& CurrentTransform = CurrentBoneTransforms[VertexData.BoneIndex];
+				if (VertexData.BoneIndex == INDEX_NONE)
+				{
+					continue;
+				}
+				const FTransform& CurrentTransform = PoseComponentSpace[VertexData.BoneIndex];
 				VertexNewPosition += CurrentTransform.TransformPosition(VertexData.VertexInBoneSpace) * VertexData.Weight;
 			}
 			
@@ -435,7 +454,7 @@
 	VerticesWithModifiedWeights.Reset();
 
 	// record the skeleton state we used to update the deformations
-	PrevBoneTransforms = Component->GetComponentSpaceTransforms();
+	PreviousPoseComponentSpace = PoseComponentSpace;
 }
 
 void FSkinToolDeformer::SetVertexNeedsUpdated(int32 VertexIndex)
@@ -445,7 +464,7 @@
 
 void FSkinToolWeights::InitializeSkinWeights(
 	const USkeletalMeshComponent* SkeletalMeshComponent,
-	FMeshDescription* Mesh)
+	const FMeshDescription* Mesh)
 {
 	// initialize deformer data
 	Deformer.Initialize(SkeletalMeshComponent, Mesh);
@@ -620,9 +639,9 @@
 		NewWeightValue);
 }
 
-void FSkinToolWeights::ApplyCurrentWeightsToMeshDescription(FMeshDescription* EditedMesh)
+void FSkinToolWeights::ApplyCurrentWeightsToMeshDescription(FMeshDescription* MeshDescription)
 {
-	FSkeletalMeshAttributes MeshAttribs(*EditedMesh);
+	FSkeletalMeshAttributes MeshAttribs(*MeshDescription);
 	FSkinWeightsVertexAttributesRef VertexSkinWeights = MeshAttribs.GetVertexSkinWeights(Profile);
 	
 	UE::AnimationCore::FBoneWeightsSettings Settings;
@@ -631,7 +650,13 @@
 	TArray<UE::AnimationCore::FBoneWeight> SourceBoneWeights;
 	SourceBoneWeights.Reserve(UE::AnimationCore::MaxInlineBoneWeightCount);
 
-	const int32 NumVertices = EditedMesh->Vertices().Num();
+	const int32 NumVertices = MeshDescription->Vertices().Num();
+	if (!ensure(CurrentWeights.Num() == NumVertices))
+	{
+		// weights are out of sync with mesh description you're trying to apply them to
+		return;
+	}
+	
 	for (int32 VertexIndex = 0; VertexIndex < NumVertices; VertexIndex++)
 	{
 		SourceBoneWeights.Reset();
@@ -875,11 +900,11 @@
 	// replace the base brush properties
 	ReplaceToolPropertySource(BrushProperties, WeightToolProperties);
 	BrushProperties = WeightToolProperties;
-
 	// brush render customization
 	BrushStampIndicator->bScaleNormalByStrength = true;
 	BrushStampIndicator->SecondaryLineThickness = 1.0f;
 	BrushStampIndicator->SecondaryLineColor = FLinearColor::Yellow;
+	RecalculateBrushRadius();
 
 	// default to the root bone as current bone
 	PendingCurrentBone = CurrentBone = Component->GetSkeletalMeshAsset()->GetRefSkeleton().GetBoneName(0);
@@ -887,44 +912,14 @@
 	// configure preview mesh
 	PreviewMesh->SetTangentsMode(EDynamicMeshComponentTangentsMode::AutoCalculated);
 	PreviewMesh->SetShadowsEnabled(false);
-	// enable vtx colors on preview mesh
-	PreviewMesh->EditMesh([](FDynamicMesh3& Mesh)
-	{
-		Mesh.EnableAttributes();
-		Mesh.Attributes()->DisablePrimaryColors();
-		Mesh.Attributes()->EnablePrimaryColors();
-		// Create an overlay that has no split elements, init with zero value.
-		Mesh.Attributes()->PrimaryColors()->CreateFromPredicate([](int ParentVID, int TriIDA, int TriIDB){return true;}, 0.f);
-	});
-	// optionally display vertex color material
-	SetDisplayVertexColors(WeightToolProperties->ColorMode != EWeightColorMode::FullMaterial);
 
-	// modify viewport render settings to optimize for painting weights
-	FPreviewProfileController PreviewProfileController;
-	PreviewProfileToRestore = PreviewProfileController.GetActiveProfile();
-	PreviewProfileController.SetActiveProfile(UDefaultEditorProfiles::EditingProfileName.ToString());
-	// turn on bone colors
-	bBoneColorsToRestore = GetDefault<UPersonaOptions>()->bShowBoneColors;
-	GetMutableDefault<UPersonaOptions>()->bShowBoneColors = true;
-
-	// initialize vertices & triangle octrees (this must be done after PreviewMesh has been initialized)
-	InitializeOctrees();
-
-	// initialize weight maps and deformation data
-	Weights.InitializeSkinWeights(Component, EditedMesh);
-	bVisibleWeightsValid = false;
-
-	RecalculateBrushRadius();
-
-	
-	const FLinearColor FaceSelectedOrange = FLinearColor(0.886f, 0.672f, 0.473f);
-	const FLinearColor VertexSelectedPurple = FLinearColor(0.78f, 0.f, 0.78f);
-	const FLinearColor VertexSelectedYellow = FLinearColor::Yellow;
-
+	// selection colors
+	constexpr FLinearColor FaceSelectedOrange = FLinearColor(0.886f, 0.672f, 0.473f);
+	constexpr FLinearColor VertexSelectedPurple = FLinearColor(0.78f, 0.f, 0.78f);
+	constexpr FLinearColor VertexSelectedYellow = FLinearColor(1.f,1.f,0.f);
 	// configure secondary render material for selected triangles
 	// NOTE: the material returned by ToolSetupUtil::GetSelectionMaterial has a checkerboard pattern on back faces which makes it hard to use
-	UMaterialInterface* Material = LoadObject<UMaterial>(nullptr, TEXT("/MeshModelingToolsetExp/Materials/SculptMaterial"));
-	if (Material)
+	if (UMaterialInterface* Material = LoadObject<UMaterial>(nullptr, TEXT("/MeshModelingToolsetExp/Materials/SculptMaterial")))
 	{
 		if (UMaterialInstanceDynamic* MatInstance = UMaterialInstanceDynamic::Create(Material, GetToolManager()))
 		{
@@ -932,33 +927,27 @@
 			PreviewMesh->SetSecondaryRenderMaterial(MatInstance);
 		}
 	}
-
 	// set up vertex selection mechanic
 	PolygonSelectionMechanic = NewObject<UPolygonSelectionMechanic>(this);
 	PolygonSelectionMechanic->bAddSelectionFilterPropertiesToParentTool = false;
 	PolygonSelectionMechanic->Setup(this);
 	PolygonSelectionMechanic->SetIsEnabled(false);
 	PolygonSelectionMechanic->OnSelectionChanged.AddLambda([this](){OnSelectionChanged.Broadcast();} );
-	// restore saved mode
-	SetComponentSelectionMode(WeightToolProperties->ComponentSelectionMode);
 	// adjust selection rendering for this context
 	PolygonSelectionMechanic->HilightRenderer.PointColor = FLinearColor::Blue;
 	PolygonSelectionMechanic->HilightRenderer.PointSize = 10.0f;
-
+	// vertex highlighting once selected
 	PolygonSelectionMechanic->SelectionRenderer.LineThickness = 1.0f;
 	PolygonSelectionMechanic->SelectionRenderer.PointColor = VertexSelectedYellow;
 	PolygonSelectionMechanic->SelectionRenderer.PointSize = 5.0f;
 	PolygonSelectionMechanic->SelectionRenderer.DepthBias = 2.0f;
-
 	// despite the name, this renders the vertices
 	PolygonSelectionMechanic->PolyEdgesRenderer.PointColor = VertexSelectedPurple;
 	PolygonSelectionMechanic->PolyEdgesRenderer.PointSize = 5.0f;
 	PolygonSelectionMechanic->PolyEdgesRenderer.DepthBias = 2.0f;
 	PolygonSelectionMechanic->PolyEdgesRenderer.LineThickness = 1.0f;
-
-	// initialize the polygon selection mechanic (this must be done after PreviewMesh has been initialized)
-	InitializeSelectionMechanic();
-	
+	// restore saved mode
+	SetComponentSelectionMode(WeightToolProperties->ComponentSelectionMode);
 	// secondary triangle buffer used to render face selection
 	PreviewMesh->EnableSecondaryTriangleBuffers([this](const FDynamicMesh3* Mesh, int32 TriangleID)
 	{
@@ -973,10 +962,11 @@
 	{
 		PreviewMesh->FastNotifySecondaryTrianglesChanged();
 	});
-	
-	// initialize smooth operator (this must be done after PreviewMesh & Weights have been initialized)
-	InitializeSmoothWeightsOperator();
 
+	// run all initialization for mesh/weights
+	PostEditMeshInitialization(Component, *PreviewMesh->GetMesh(), *EditedMesh);
+
+	// bind the skeletal mesh editor context
 	if (EditorContext.IsValid())
 	{
 		EditorContext->BindTo(this);
@@ -984,7 +974,15 @@
 
 	// trigger last used mode
 	ToggleEditingMode();
-	
+
+	// modify viewport render settings to optimize for painting weights
+	FPreviewProfileController PreviewProfileController;
+	PreviewProfileToRestore = PreviewProfileController.GetActiveProfile();
+	PreviewProfileController.SetActiveProfile(UDefaultEditorProfiles::EditingProfileName.ToString());
+	// turn on bone colors
+	bBoneColorsToRestore = GetDefault<UPersonaOptions>()->bShowBoneColors;
+	GetMutableDefault<UPersonaOptions>()->bShowBoneColors = true;
+	// set focus to viewport so brush hotkey works
 	SetFocusInViewport();
 	
 	// inform user of tool keys
@@ -1124,6 +1122,12 @@
 
 void USkinWeightsPaintTool::OnTick(float DeltaTime)
 {
+	if (bPendingUpdateFromPartialMesh)
+	{
+		FinishIsolatedSelection();
+		bPendingUpdateFromPartialMesh = false;
+	}
+	
 	if (bStampPending)
 	{
 		ApplyStamp(LastStamp);
@@ -1144,7 +1148,43 @@
 	}
 
 	// sparsely updates vertex positions (only on vertices with modified weights)
-	Weights.Deformer.UpdateVertexDeformation(this);
+	Weights.Deformer.UpdateVertexDeformation(this, Weights.Deformer.Component->GetComponentSpaceTransforms());
 }
 
+void USkinWeightsPaintTool::PostEditMeshInitialization(
+	const USkeletalMeshComponent* InComponent,
+	const FDynamicMesh3& InDynamicMesh,
+	const FMeshDescription& InMeshDescription)
+{
+	// update the preview mesh
+	PreviewMesh->ReplaceMesh(InDynamicMesh);
+	PreviewMesh->EditMesh([](FDynamicMesh3& Mesh)
+	{
+		Mesh.EnableAttributes();
+		Mesh.Attributes()->DisablePrimaryColors();
+		Mesh.Attributes()->EnablePrimaryColors();
+		Mesh.Attributes()->PrimaryColors()->CreateFromPredicate([](int ParentVID, int TriIDA, int TriIDB){return true;}, 0.f);
+	});
+	SetDisplayVertexColors(WeightToolProperties->ColorMode != EWeightColorMode::FullMaterial);
+	
+	// update vertices & triangle octrees (this must be done after PreviewMesh has been updated)
+	InitializeOctrees();
+
+	// update the polygon selection mechanic (this must be done after PreviewMesh has been updated)
+	InitializeSelectionMechanic();
+
+	// update weights
+	Weights = FSkinToolWeights();
+	if (!IsProfileValid(WeightToolProperties->ActiveSkinWeightProfile))
+	{
+		WeightToolProperties->ActiveSkinWeightProfile = FSkeletalMeshAttributesShared::DefaultSkinWeightProfileName;
+	}
+	Weights.Profile = WeightToolProperties->ActiveSkinWeightProfile;
+	Weights.InitializeSkinWeights(InComponent, &InMeshDescription);
+	bVisibleWeightsValid = false;
+	
+	// update smooth operator (this must be done after PreviewMesh & Weights have been updated)
+	InitializeSmoothWeightsOperator();
+}
+
 void USkinWeightsPaintToolProperties::SetComponentMode(EComponentSelectionMode InComponentMode)
@@ -1201,7 +1241,7 @@
 		CurTargetTransform.InverseTransformVector((FVector3d)Ray.Direction));
 	UE::Geometry::Normalize(LocalRay.Direction);
 	
-	const FDynamicMesh3* Mesh = PreviewMesh->GetPreviewDynamicMesh();
+	const FDynamicMesh3* Mesh = PreviewMesh->GetMesh();
 
 	FViewCameraState StateOut;
 	GetToolManager()->GetContextQueriesAPI()->GetCurrentViewState(StateOut);
@@ -1763,6 +1803,7 @@
 		SelectionTopology.Get(),
 		[this]() { return MeshSpatial.Get(); }
 	);
+	PolygonSelectionMechanic->ClearSelection();
 }
 
 void USkinWeightsPaintTool::InitializeSmoothWeightsOperator()
@@ -2431,6 +2472,11 @@
 		return;
 	}
 
+	if (IsSelectionIsolated())
+	{
+		FinishIsolatedSelection();
+	}
+
 	// apply previous changes
 	Weights.ApplyCurrentWeightsToMeshDescription(EditedMesh);
 
@@ -2444,34 +2490,9 @@
 		*EditedMesh = *UE::ToolTarget::GetMeshDescription(Target, Params);
 	}
 
-	// update the preview mesh
-	PreviewMesh->ReplaceMesh(UE::ToolTarget::GetDynamicMeshCopy(Target, Params));
-	PreviewMesh->EditMesh([](FDynamicMesh3& Mesh)
-	{
-		Mesh.EnableAttributes();
-		Mesh.Attributes()->DisablePrimaryColors();
-		Mesh.Attributes()->EnablePrimaryColors();
-		Mesh.Attributes()->PrimaryColors()->CreateFromPredicate([](int ParentVID, int TriIDA, int TriIDB){return true;}, 0.f);
-	});
-
-	// update vertices & triangle octrees (this must be done after PreviewMesh has been updated)
-	InitializeOctrees();
-
-	// update weights
-	Weights = FSkinToolWeights();
-	if (!IsProfileValid(WeightToolProperties->ActiveSkinWeightProfile))
-	{
-		WeightToolProperties->ActiveSkinWeightProfile = FSkeletalMeshAttributesShared::DefaultSkinWeightProfileName;
-	}
-	Weights.Profile = WeightToolProperties->ActiveSkinWeightProfile;
-	Weights.InitializeSkinWeights(Component, EditedMesh);
-	bVisibleWeightsValid = false;
-
-	// update the polygon selection mechanic (this must be done after PreviewMesh has been updated)
-	InitializeSelectionMechanic();
-
-	// update smooth operator (this must be done after PreviewMesh & Weights have been updated)
-	InitializeSmoothWeightsOperator();
+	// reinitialize all mesh data structures
+	const UE::Geometry::FDynamicMesh3 DynamicMesh = UE::ToolTarget::GetDynamicMeshCopy(Target, Params);
+	PostEditMeshInitialization(Component, DynamicMesh, *EditedMesh);
 }
 
 void USkinWeightsPaintTool::OnActiveSkinWeightProfileChanged()
@@ -2482,6 +2503,11 @@
 		return;
 	}
 
+	if (IsSelectionIsolated())
+	{
+		FinishIsolatedSelection();
+	}
+
 	if (!IsProfileValid(WeightToolProperties->ActiveSkinWeightProfile))
 	{
 		WeightToolProperties->ActiveSkinWeightProfile = FSkeletalMeshAttributesShared::DefaultSkinWeightProfileName;
@@ -2591,6 +2617,153 @@
 	PolygonSelectionMechanic->FloodSelection();
 }
 
+bool USkinWeightsPaintTool::IsAnyComponentSelected() const
+{
+	if (!PolygonSelectionMechanic)
+	{
+		return false;
+	}
+
+	return PolygonSelectionMechanic->HasSelection();
+}
+
+bool USkinWeightsPaintTool::IsSelectionIsolated() const
+{
+	return PartialMeshDescription.IsValid();
+}
+
+void USkinWeightsPaintTool::SetIsolateSelected(const bool bIsolateSelection)
+{
+	// if we are turning off an isolated selection, we must queue the Tick() to update the full mesh
+	if (!bIsolateSelection && PartialMeshDescription.IsValid())
+	{
+		bPendingUpdateFromPartialMesh = true;
+		return;
+	}
+	
+	if (PartialMeshDescription.IsValid())
+	{
+		ensure(false); // should be reset to null
+		return;
+	}
+	
+	if (!ensure(PolygonSelectionMechanic))
+	{
+		return;
+	}
+
+	const USkeletalMeshComponent* SkeletalMeshComponent = GetSkeletalMeshComponent(Target);
+	if (!ensure(SkeletalMeshComponent))
+	{
+		return;
+	}
+
+	if (!ensure(EditedMesh))
+	{
+		return;
+	}
+
+	// apply previous changes
+	Weights.ApplyCurrentWeightsToMeshDescription(EditedMesh);
+	
+	// put into ref pose, BEFORE copying the mesh, so that submesh deformer initializes with vertices in ref pose
+	Weights.Deformer.SetToRefPose(this);
+
+	// store selection to be restored
+	IsolatedSelectionToRestoreVertices.Reset();
+	IsolatedSelectionToRestoreEdges.Reset();
+	IsolatedSelectionToRestoreFaces.Reset();
+	IsolatedSelectionToRestoreVertices.ElementType = UE::Geometry::EGeometryElementType::Vertex;
+	IsolatedSelectionToRestoreEdges.ElementType = UE::Geometry::EGeometryElementType::Edge;
+	IsolatedSelectionToRestoreFaces.ElementType = UE::Geometry::EGeometryElementType::Face;
+	PolygonSelectionMechanic->GetSelection_AsTriangleTopology(IsolatedSelectionToRestoreVertices);
+	PolygonSelectionMechanic->GetSelection_AsTriangleTopology(IsolatedSelectionToRestoreEdges);
+	PolygonSelectionMechanic->GetSelection_AsTriangleTopology(IsolatedSelectionToRestoreFaces);
+
+	// store copy of original FDynamicMesh to restore
+	FDynamicMesh3 DynamicMesh;
+	FMeshDescriptionToDynamicMesh Converter;
+	Converter.Convert(EditedMesh, DynamicMesh);
+	FullDynamicMesh = MoveTemp(DynamicMesh);
+
+	// create a submesh from the selected triangles (or triangles connected to selected vertices/edges)
+	TArray<int32> TrianglesToIsolate;
+	GetSelectedTriangles(TrianglesToIsolate);
+	if (TrianglesToIsolate.IsEmpty())
+	{
+		return;
+	}
+	PartialSubMesh = UE::Geometry::FDynamicSubmesh3(&FullDynamicMesh, TrianglesToIsolate);
+
+	// create mesh description for sub-mesh
+	PartialMeshDescription = MakeShared<FMeshDescription>();
+	// registering skeletal mesh attributes is required to create room to copy attributes during conversion from dynamic mesh
+	FSkeletalMeshAttributes Attributes(*PartialMeshDescription);
+	Attributes.Register();
+	// convert the partial dynamic mesh to a mesh description
+	// NOTE: this copies vertex weights to partial mesh description (later used to load weights into the tool)
+	FDynamicMeshToMeshDescription DnyToDescConverter;
+	constexpr bool bCopyTangents = true;
+	DnyToDescConverter.Convert(&PartialSubMesh.GetSubmesh(), *PartialMeshDescription, bCopyTangents);
+	
+	// reinitialize all mesh data structures
+	PostEditMeshInitialization(SkeletalMeshComponent, PartialSubMesh.GetSubmesh(), *PartialMeshDescription.Get());
+}
+
+void USkinWeightsPaintTool::FinishIsolatedSelection()
+{
+	const USkeletalMeshComponent* SkeletalMeshComponent = GetSkeletalMeshComponent(Target);
+	if (!ensure(SkeletalMeshComponent))
+	{
+		return;
+	}
+
+	if (!PartialMeshDescription)
+	{
+		// nothing hidden
+		return;
+	}
+
+	// apply partial mesh weights to partial mesh description
+	Weights.ApplyCurrentWeightsToMeshDescription(PartialMeshDescription.Get());
+
+	// reinitialize with full mesh
+	PostEditMeshInitialization(SkeletalMeshComponent, FullDynamicMesh, *EditedMesh);
+
+	// copy the remapped weights back to the full mesh
+	const FSkeletalMeshConstAttributes MeshAttribs(*PartialMeshDescription.Get());
+	const FSkinWeightsVertexAttributesConstRef AllVertexWeights = MeshAttribs.GetVertexSkinWeights(WeightToolProperties->ActiveSkinWeightProfile);
+	const int32 NumVerticesInPartialMesh = PartialMeshDescription.Get()->Vertices().Num();
+	for (int32 VertexIndexPartial = 0; VertexIndexPartial < NumVerticesInPartialMesh; VertexIndexPartial++)
+	{
+		// get the equivalent vertex on the full mesh
+		const int32 VertexIndexFull = PartialSubMesh.MapVertexToBaseMesh(VertexIndexPartial);
+		// clear all the weights on this vertex
+		Weights.CurrentWeights[VertexIndexFull].Init(FVertexBoneWeight(), UE::AnimationCore::MaxInlineBoneWeightCount);
+		// replace with weights from partial mesh
+		const FVertexBoneWeightsConst& VertexWeightsPartial = AllVertexWeights.Get(VertexIndexPartial);
+		for (int32 InfluenceIndex=0; InfluenceIndex<VertexWeightsPartial.Num(); ++InfluenceIndex)
+		{
+			const UE::AnimationCore::FBoneWeight& SingleBoneWeight = VertexWeightsPartial[InfluenceIndex];
+			FVertexBoneWeight& VertexBoneWeight = Weights.CurrentWeights[VertexIndexFull][InfluenceIndex];
+			VertexBoneWeight.BoneIndex = SingleBoneWeight.GetBoneIndex();
+			VertexBoneWeight.Weight = SingleBoneWeight.GetWeight();
+			VertexBoneWeight.VertexInBoneSpace = Weights.Deformer.InvCSRefPoseTransforms[VertexBoneWeight.BoneIndex].TransformPosition(Weights.Deformer.RefPoseVertexPositions[VertexIndexFull]);
+		}
+	}
+	// sync both weight buffers
+	Weights.PreChangeWeights = Weights.CurrentWeights;
+	// apply full mesh weights to full mesh description
+	Weights.ApplyCurrentWeightsToMeshDescription(EditedMesh);
+
+	// restore selection (allows for easily adjusting crop)
+	PolygonSelectionMechanic->SetSelection_AsTriangleTopology(IsolatedSelectionToRestoreVertices);
+	PolygonSelectionMechanic->SetSelection_AsTriangleTopology(IsolatedSelectionToRestoreEdges);
+	PolygonSelectionMechanic->SetSelection_AsTriangleTopology(IsolatedSelectionToRestoreFaces);
+
+	PartialMeshDescription = nullptr;
+}
+
 void USkinWeightsPaintTool::GetSelectedVertices(TArray<int32>& OutVertexIndices) const
 {
 	OutVertexIndices.Empty();
@@ -2616,13 +2789,9 @@
 	};
 
 	// add selected vertices
-	if (PolygonSelectionMechanic->Properties->bSelectVertices)
-	{
-		AddVertices(Selection.SelectedCornerIDs);
-	}
+	AddVertices(Selection.SelectedCornerIDs);
 
 	// add vertices on selected edges
-	if (PolygonSelectionMechanic->Properties->bSelectEdges)
 	{
 		TSet<int32> VerticesInSelectedEdges;
 		for (const int32 SelectedEdgeIndex : Selection.SelectedEdgeIDs)
@@ -2636,7 +2805,6 @@
 	}
 
 	// add vertices in selected faces
-	if (PolygonSelectionMechanic->Properties->bSelectFaces)
 	{
 		TSet<int32> VerticesInSelectedFaces;
 		for (const int32 SelectedFaceIndex : Selection.SelectedGroupIDs)
@@ -2651,6 +2819,42 @@
 	}
 }
 
+void USkinWeightsPaintTool::GetSelectedTriangles(TArray<int32>& OutTriangleIndices) const
+{
+	OutTriangleIndices.Empty();
+	if (!PolygonSelectionMechanic)
+	{
+		return;
+	}
+	
+	const FGroupTopologySelection& Selection = PolygonSelectionMechanic->GetActiveSelection();
+	const FDynamicMesh3* DynamicMesh = PreviewMesh->GetMesh();
+	TSet<int32> TriangleSet;
+
+	// add triangles connected to selected vertices
+	for (const int32 VertexIndex : Selection.SelectedCornerIDs)
+	{
+		DynamicMesh->EnumerateVertexTriangles(VertexIndex, [&TriangleSet](int32 TriangleIndex)
+		{
+			TriangleSet.Add(TriangleIndex);
+		});
+	}
+
+	// add triangles connected to selected edges
+	for (const int32 EdgeIndex : Selection.SelectedEdgeIDs)
+	{
+		DynamicMesh->EnumerateEdgeTriangles(EdgeIndex, [&TriangleSet](int32 TriangleIndex)
+		{
+			TriangleSet.Add(TriangleIndex);
+		});
+	}
+
+	// add selected triangles
+	TriangleSet.Append(Selection.SelectedGroupIDs);	
+
+	OutTriangleIndices = TriangleSet.Array();
+}
+
 void USkinWeightsPaintTool::GetInfluences(const TArray<int32>& VertexIndices, TArray<BoneIndex>& OutBoneIndices)
 {
 	for (const int32 SelectedVertex : VertexIndices)
@@ -2677,6 +2881,11 @@
 	
 	for (const int32 SelectedVertex : VertexIndices)
 	{
+		if (!Weights.CurrentWeights.IsValidIndex(SelectedVertex))
+		{
+			continue;
+		}
+		
 		for (const FVertexBoneWeight& VertexBoneData : Weights.CurrentWeights[SelectedVertex])
 		{
 			if (VertexBoneData.BoneIndex == InBoneIndex)
diff --git a/Engine/Plugins/Experimental/MeshModelingToolsetExp/Source/MeshModelingToolsEditorOnlyExp/Public/SkeletalMesh/SkinWeightsPaintTool.h b/Engine/Plugins/Experimental/MeshModelingToolsetExp/Source/MeshModelingToolsEditorOnlyExp/Public/SkeletalMesh/SkinWeightsPaintTool.h
--- a/Engine/Plugins/Experimental/MeshModelingToolsetExp/Source/MeshModelingToolsEditorOnlyExp/Public/SkeletalMesh/SkinWeightsPaintTool.h
+++ b/Engine/Plugins/Experimental/MeshModelingToolsetExp/Source/MeshModelingToolsEditorOnlyExp/Public/SkeletalMesh/SkinWeightsPaintTool.h
@@ -7,6 +7,7 @@
 #include "MeshDescription.h"
 #include "DynamicMesh/DynamicVerticesOctree3.h"
 #include "BoneWeights.h"
+#include "DynamicSubmesh3.h"
 #include "GroupTopology.h"
 #include "SkeletalMeshAttributes.h"
 #include "Misc/Optional.h"
@@ -14,6 +15,7 @@
 #include "DynamicMesh/DynamicMeshOctree3.h"
 #include "SkeletalMesh/SkeletalMeshEditionInterface.h"
 #include "Engine/SkeletalMesh.h"
+#include "Selections/GeometrySelection.h"
 #include "TargetInterfaces/MeshTargetInterfaceTypes.h"
 
 #include "SkinWeightsPaintTool.generated.h"
@@ -26,6 +28,7 @@ class UPersonaEditorModeManagerContext;
 
 namespace UE::Geometry 
 {
+	struct FGeometrySelection;
 	template <typename BoneIndexType, typename BoneWeightType> class TBoneWeightsDataSource;
 	template <typename BoneIndexType, typename BoneWeightType> class TSmoothBoneWeights;
 }
@@ -107,11 +110,13 @@ namespace SkinPaintTool
 	// data required to preview the skinning deformations as you paint
 	struct FSkinToolDeformer
 	{
-		void Initialize(const USkeletalMeshComponent* SkeletalMeshComponent, const FMeshDescription* Mesh);
+		void Initialize(const USkeletalMeshComponent* InSkelMeshComponent, const FMeshDescription* InMeshDescription);
 
 		void SetAllVerticesToBeUpdated();
 
-		void UpdateVertexDeformation(USkinWeightsPaintTool* Tool);
+		void SetToRefPose(USkinWeightsPaintTool* Tool);
+
+		void UpdateVertexDeformation(USkinWeightsPaintTool* Tool, const TArray<FTransform>& PoseComponentSpace);
 
 		void SetVertexNeedsUpdated(int32 VertexIndex);
 		
@@ -122,7 +127,9 @@ namespace SkinPaintTool
 		// inverted, component space ref pose transform of each bone
 		TArray<FTransform> InvCSRefPoseTransforms;
 		// bones transforms used in last deformation update
-		TArray<FTransform> PrevBoneTransforms;
+		TArray<FTransform> PreviousPoseComponentSpace;
+		// bones transforms stored for duration of async deformation update
+		TArray<FTransform> RefPoseComponentSpace;
 		// bone index to bone name
 		TArray<FName> BoneNames;
 		TMap<FName, BoneIndex> BoneNameToIndexMap;
@@ -182,7 +189,7 @@ namespace SkinPaintTool
 		// copy the initial weight values from the skeletal mesh
 		void InitializeSkinWeights(
 			const USkeletalMeshComponent* SkeletalMeshComponent,
-			FMeshDescription* Mesh);
+			const FMeshDescription* Mesh);
 
 		// applies an edit to a single vertex weight on a single bone, then normalizes the remaining weights while
 		// keeping the edited weight intact (ie, adapts OTHER influences to achieve normalization)
@@ -192,7 +199,7 @@ namespace SkinPaintTool
 			float NewWeightValue,
 			FMultiBoneWeightEdits& WeightEdits);
 
-		void ApplyCurrentWeightsToMeshDescription(FMeshDescription* EditedMesh);
+		void ApplyCurrentWeightsToMeshDescription(FMeshDescription* MeshDescription);
 		
 		static float GetWeightOfBoneOnVertex(
 			const int32 BoneIndex,
@@ -425,9 +432,14 @@ public:
 	void GrowSelection() const;
 	void ShrinkSelection() const;
 	void FloodSelection() const;
+	// isolate selection
+	bool IsAnyComponentSelected() const;
+	bool IsSelectionIsolated() const;
+	void SetIsolateSelected(const bool bIsolateSelection);
 
-	// get a list of currently selected vertices
+	// get a list of currently selected components (converting other components)
 	void GetSelectedVertices(TArray<int32>& OutVertexIndices) const;
+	void GetSelectedTriangles(TArray<int32>& OutTriangleIndices) const;
 
 	// get the average weight value of each influence on the given vertices
 	void GetInfluences(const TArray<int32>& VertexIndices, TArray<BoneIndex>& OutBoneIndices);
@@ -476,6 +488,11 @@ protected:
 	void OnShutdown(EToolShutdownType ShutdownType) override;
 	void OnTick(float DeltaTime) override;
 
+	void PostEditMeshInitialization(
+		const USkeletalMeshComponent* Component,
+		const FDynamicMesh3& InDynamicMesh,
+		const FMeshDescription& InMeshDescription);
+
 	// stamp
 	float CalculateBrushFalloff(float Distance) const;
 	void CalculateVertexROI(
@@ -523,6 +540,14 @@ protected:
 	// the currently edited mesh descriptions
 	mutable TMap<EMeshLODIdentifier, FMeshDescription> EditedMeshes;
 	FMeshDescription* EditedMesh = nullptr;
+	// when selection is isolated, we hide the full mesh and show a submesh
+	// when islated selection is unhidden, we remap all changes from the submesh back to the full mesh
+	TSharedPtr<FMeshDescription> PartialMeshDescription = nullptr; // during isolated selection
+	UE::Geometry::FGeometrySelection IsolatedSelectionToRestoreVertices;
+	UE::Geometry::FGeometrySelection IsolatedSelectionToRestoreEdges;
+	UE::Geometry::FGeometrySelection IsolatedSelectionToRestoreFaces;
+	bool bPendingUpdateFromPartialMesh = false;
+	void FinishIsolatedSelection();
 
 	// storage of vertex weights per bone 
 	SkinPaintTool::FSkinToolWeights Weights;
@@ -565,6 +590,10 @@ protected:
 	TUniquePtr<UE::Geometry::FTriangleGroupTopology> SelectionTopology = nullptr;
 	void InitializeSelectionMechanic();
 
+	// isolate selection sub-meshes
+	UE::Geometry::FDynamicSubmesh3 PartialSubMesh;
+	UE::Geometry::FDynamicMesh3 FullDynamicMesh;
+
 	// skin weight layer
 	void OnActiveLODChanged();
 	void OnActiveSkinWeightProfileChanged();
