[GeometryProcessing] CombineMeshInstances - Fix crash when no approximated mesh is generated.

diff --git a/Engine/Plugins/Experimental/MeshModelingToolsetExp/Source/GeometryProcessingAdapters/Private/CombineMeshInstancesImpl.cpp b/Engine/Plugins/Experimental/MeshModelingToolsetExp/Source/GeometryProcessingAdapters/Private/CombineMeshInstancesImpl.cpp
--- a/Engine/Plugins/Experimental/MeshModelingToolsetExp/Source/GeometryProcessingAdapters/Private/CombineMeshInstancesImpl.cpp
+++ b/Engine/Plugins/Experimental/MeshModelingToolsetExp/Source/GeometryProcessingAdapters/Private/CombineMeshInstancesImpl.cpp
@@ -673,7 +673,10 @@ void ReplaceBadSimplifiedLODs(FMeshPartsAssembly& Assembly, const IGeometryProce
 			Selector.Initialize(Spatial.GetMesh(), &Spatial);
 			if ( k == OptimizedTargets.SimplifiedMeshLODs.Num()-1 )
 			{
-				Selector.AddGeneratedMesh(OptimizedTargets.ApproximateMeshLODs[0], 2);
+				if (OptimizedTargets.ApproximateMeshLODs.Num() > 0)
+				{
+					Selector.AddGeneratedMesh(OptimizedTargets.ApproximateMeshLODs[0], 2);
+				}
 			}
 			else
 			{
