Fixes skin paint issue where vertex flood does not work on some meshes.
Added CleanMesh() to skin paint tool.

diff --git a/Engine/Plugins/Experimental/MeshModelingToolsetExp/Source/MeshModelingToolsEditorOnlyExp/Private/SkeletalMesh/SkinWeightsPaintTool.cpp b/Engine/Plugins/Experimental/MeshModelingToolsetExp/Source/MeshModelingToolsEditorOnlyExp/Private/SkeletalMesh/SkinWeightsPaintTool.cpp
--- a/Engine/Plugins/Experimental/MeshModelingToolsetExp/Source/MeshModelingToolsEditorOnlyExp/Private/SkeletalMesh/SkinWeightsPaintTool.cpp
+++ b/Engine/Plugins/Experimental/MeshModelingToolsetExp/Source/MeshModelingToolsEditorOnlyExp/Private/SkeletalMesh/SkinWeightsPaintTool.cpp
@@ -38,6 +38,7 @@
 #include "PreviewProfileController.h"
 #include "Animation/SkinWeightProfile.h"
 #include "AnimationRuntime.h"
+#include "TargetInterfaces/DynamicMeshCommitter.h"
 
 #include UE_INLINE_GENERATED_CPP_BY_NAME(SkinWeightsPaintTool)
 
@@ -876,6 +877,9 @@ void USkinWeightsPaintTool::Setup()
 	const USkeletalMeshComponent* Component = Cast<USkeletalMeshComponent>(TargetComponent->GetOwnerComponent());
 	check(Component && Component->GetSkeletalMeshAsset())
 
+	// prepare mesh for skin editing
+	CleanMesh();
+
 	// create a mesh description for editing (this must be done before calling UpdateBonePositionInfos) 
 	bool bSupportsLODs = false;
 	const EMeshLODIdentifier DefaultLOD = UE::ToolTarget::GetTargetMeshDescriptionLOD(Target, bSupportsLODs);
@@ -1016,6 +1020,11 @@ void USkinWeightsPaintTool::Render(IToolsContextRenderAPI* RenderAPI)
 
 FBox USkinWeightsPaintTool::GetWorldSpaceFocusBox()
 {
+	if (!WeightToolProperties)
+	{
+		return PreviewMesh->GetActor()->GetComponentsBoundingBox();
+	}
+	
 	// 1. Prioritize Brush & Vertex modes
 	switch (WeightToolProperties->EditingMode)
 	{
@@ -1187,6 +1196,22 @@ void USkinWeightsPaintTool::PostEditMeshInitialization(
 	InitializeSmoothWeightsOperator();
 }
 
+void USkinWeightsPaintTool::CleanMesh() const
+{
+	if (PreviewMesh->GetMesh()->HasUnusedVertices())
+	{
+		// orphaned vertices wreak havoc on our selection tools
+		PreviewMesh->EditMesh([](FDynamicMesh3& Mesh)
+		{
+			Mesh.RemoveUnusedVertices();
+			Mesh.CompactInPlace();
+		});
+	
+		IDynamicMeshCommitter* DynamicMeshCommitter = Cast<IDynamicMeshCommitter>(Target);
+		DynamicMeshCommitter->CommitDynamicMesh(*PreviewMesh->GetMesh());
+	}
+}
+
 void USkinWeightsPaintToolProperties::SetComponentMode(EComponentSelectionMode InComponentMode)
 {
 	ComponentSelectionMode = InComponentMode;
diff --git a/Engine/Plugins/Experimental/MeshModelingToolsetExp/Source/MeshModelingToolsEditorOnlyExp/Public/SkeletalMesh/SkinWeightsPaintTool.h b/Engine/Plugins/Experimental/MeshModelingToolsetExp/Source/MeshModelingToolsEditorOnlyExp/Public/SkeletalMesh/SkinWeightsPaintTool.h
--- a/Engine/Plugins/Experimental/MeshModelingToolsetExp/Source/MeshModelingToolsEditorOnlyExp/Public/SkeletalMesh/SkinWeightsPaintTool.h
+++ b/Engine/Plugins/Experimental/MeshModelingToolsetExp/Source/MeshModelingToolsEditorOnlyExp/Public/SkeletalMesh/SkinWeightsPaintTool.h
@@ -493,6 +493,8 @@ protected:
 		const FDynamicMesh3& InDynamicMesh,
 		const FMeshDescription& InMeshDescription);
 
+	void CleanMesh() const;
+
 	// stamp
 	float CalculateBrushFalloff(float Distance) const;
 	void CalculateVertexROI(
diff --git a/Engine/Source/Runtime/GeometryCore/Private/DynamicMesh/DynamicMesh3_Edits.cpp b/Engine/Source/Runtime/GeometryCore/Private/DynamicMesh/DynamicMesh3_Edits.cpp
--- a/Engine/Source/Runtime/GeometryCore/Private/DynamicMesh/DynamicMesh3_Edits.cpp
+++ b/Engine/Source/Runtime/GeometryCore/Private/DynamicMesh/DynamicMesh3_Edits.cpp
@@ -302,6 +302,19 @@ void FDynamicMesh3::RemoveUnusedVertices()
 	UpdateChangeStamps(true, true);
 }
 
+bool FDynamicMesh3::HasUnusedVertices() const
+{
+	for (int32 VID = 0; VID < MaxVertexID(); ++VID)
+	{
+		// If vertex exists but is not referenced by any triangles
+		if (VertexRefCounts.GetRefCount(VID) == 1)
+		{
+			return true;
+		}
+	}
+	
+	return false;
+}
 
 
 void FDynamicMesh3::CompactInPlace(FCompactMaps* CompactInfo)
diff --git a/Engine/Source/Runtime/GeometryCore/Public/DynamicMesh/DynamicMesh3.h b/Engine/Source/Runtime/GeometryCore/Public/DynamicMesh/DynamicMesh3.h
--- a/Engine/Source/Runtime/GeometryCore/Public/DynamicMesh/DynamicMesh3.h
+++ b/Engine/Source/Runtime/GeometryCore/Public/DynamicMesh/DynamicMesh3.h
@@ -1212,6 +1212,11 @@ public:
 	 */
 	GEOMETRYCORE_API void RemoveUnusedVertices();
 
+	/**
+	 * @return true if any vertices are unused (not in any triangles)
+	 * */
+	GEOMETRYCORE_API bool HasUnusedVertices() const;
+
 	/**
 	 * Reverse the ccw/cw orientation of all triangles in the mesh, and
 	 * optionally flip the vertex normals if they exist
