Fixed issue where left click drag moves the camera while in a skin paint brush tool.

diff --git a/Engine/Plugins/Experimental/MeshModelingToolsetExp/Source/MeshModelingToolsEditorOnlyExp/Private/SkeletalMesh/SkinWeightsPaintTool.cpp b/Engine/Plugins/Experimental/MeshModelingToolsetExp/Source/MeshModelingToolsEditorOnlyExp/Private/SkeletalMesh/SkinWeightsPaintTool.cpp
--- a/Engine/Plugins/Experimental/MeshModelingToolsetExp/Source/MeshModelingToolsEditorOnlyExp/Private/SkeletalMesh/SkinWeightsPaintTool.cpp
+++ b/Engine/Plugins/Experimental/MeshModelingToolsetExp/Source/MeshModelingToolsEditorOnlyExp/Private/SkeletalMesh/SkinWeightsPaintTool.cpp
@@ -1129,6 +1129,26 @@ void USkinWeightsPaintTool::OnUpdateModifierState(int ModifierID, bool bIsOn)
 	Super::OnUpdateModifierState(ModifierID, bIsOn);
 }
 
+FInputRayHit USkinWeightsPaintTool::CanBeginClickDragSequence(const FInputDeviceRay& InPressPos)
+{
+	// NOTE: this function is only overridden to prevent left-click fly camera behavior while brushing
+	// this should eventually be removed once we have a clear way of disabling the fly-cam mode
+	
+	if (WeightToolProperties->EditingMode != EWeightEditMode::Brush)
+	{
+		return FInputRayHit(); // allow other behaviors to capture mouse while not brushing
+	}
+	
+	const FInputRayHit Hit = Super::CanBeginClickDragSequence(InPressPos);
+	if (Hit.bHit)
+	{
+		return Hit;
+	}
+
+	// always return a hit so we always capture and prevent accidental camera movement
+	return FInputRayHit(TNumericLimits<float>::Max());
+}
+
 void USkinWeightsPaintTool::OnTick(float DeltaTime)
 {
 	if (bPendingUpdateFromPartialMesh)
diff --git a/Engine/Plugins/Experimental/MeshModelingToolsetExp/Source/MeshModelingToolsEditorOnlyExp/Public/SkeletalMesh/SkinWeightsPaintTool.h b/Engine/Plugins/Experimental/MeshModelingToolsetExp/Source/MeshModelingToolsEditorOnlyExp/Public/SkeletalMesh/SkinWeightsPaintTool.h
--- a/Engine/Plugins/Experimental/MeshModelingToolsetExp/Source/MeshModelingToolsEditorOnlyExp/Public/SkeletalMesh/SkinWeightsPaintTool.h
+++ b/Engine/Plugins/Experimental/MeshModelingToolsetExp/Source/MeshModelingToolsEditorOnlyExp/Public/SkeletalMesh/SkinWeightsPaintTool.h
@@ -405,6 +405,7 @@ public:
 
 	// IClickDragBehaviorTarget implementation
 	virtual void OnUpdateModifierState(int ModifierID, bool bIsOn) override;
+	virtual FInputRayHit CanBeginClickDragSequence(const FInputDeviceRay& InPressPos) override;
 
 	// using when ToolChange is applied via Undo/Redo
 	void ExternalUpdateWeights(const int32 BoneIndex, const TMap<int32, float>& IndexValues);
