Fixes regression in skin paint vertex colors from ramp.

diff --git a/Engine/Plugins/Experimental/MeshModelingToolsetExp/Source/MeshModelingToolsEditorOnlyExp/Private/SkeletalMesh/SkinWeightsPaintTool.cpp b/Engine/Plugins/Experimental/MeshModelingToolsetExp/Source/MeshModelingToolsEditorOnlyExp/Private/SkeletalMesh/SkinWeightsPaintTool.cpp
--- a/Engine/Plugins/Experimental/MeshModelingToolsetExp/Source/MeshModelingToolsEditorOnlyExp/Private/SkeletalMesh/SkinWeightsPaintTool.cpp
+++ b/Engine/Plugins/Experimental/MeshModelingToolsetExp/Source/MeshModelingToolsEditorOnlyExp/Private/SkeletalMesh/SkinWeightsPaintTool.cpp
@@ -201,18 +201,24 @@ float FDirectEditWeightState::GetModeMaxValue()
 
 USkinWeightsPaintToolProperties::USkinWeightsPaintToolProperties()
 {
-	ColorRamp.Add(FLinearColor::Blue);
-	ColorRamp.Add(FLinearColor::Yellow);
-
-	MinColor = FLinearColor::Black;
-	MaxColor = FLinearColor::White;
-
 	BrushConfigs.Add(EWeightEditOperation::Add, &BrushConfigAdd);
 	BrushConfigs.Add(EWeightEditOperation::Replace, &BrushConfigReplace);
 	BrushConfigs.Add(EWeightEditOperation::Multiply, &BrushConfigMultiply);
 	BrushConfigs.Add(EWeightEditOperation::Relax, &BrushConfigRelax);
 
 	LoadConfig();
+
+	if (ColorRamp.IsEmpty())
+	{
+		// default color ramp simulates a heat map
+		ColorRamp.Add(FLinearColor(0.8f, 0.4f, 0.8f)); // Purple
+		ColorRamp.Add(FLinearColor(0.0f, 0.0f, 0.5f)); // Dark Blue
+		ColorRamp.Add(FLinearColor(0.2f, 0.2f, 1.0f)); // Light Blue
+		ColorRamp.Add(FLinearColor(0.0f, 1.0f, 0.0f)); // Green
+		ColorRamp.Add(FLinearColor(1.0f, 1.0f, 0.0f)); // Yellow
+		ColorRamp.Add(FLinearColor(1.0f, 0.65f, 0.0f)); // Orange
+		ColorRamp.Add(FLinearColor(1.0f, 0.0f, 0.0f, 0.0f)); // Red
+	}
 }
 
 TArray<FName> USkinWeightsPaintToolProperties::GetLODsFunc() const
@@ -1543,7 +1549,7 @@ FVector4f USkinWeightsPaintTool::GetColorOfVertex(VertexIndex InVertexIndex, Bon
 		{
 			if (InCurrentBoneIndex == INDEX_NONE)
 			{
-				return WeightToolProperties->MinColor; // with no bone selected, all vertices are drawn black
+				return FLinearColor::Black; // with no bone selected, all vertices are drawn black
 			}
 			const float Value = Weights.GetWeightOfBoneOnVertex(InCurrentBoneIndex, InVertexIndex, Weights.CurrentWeights);
 			return FMath::Lerp(FLinearColor::Black, FLinearColor::White, Value);
@@ -1552,39 +1558,57 @@ FVector4f USkinWeightsPaintTool::GetColorOfVertex(VertexIndex InVertexIndex, Bon
 		{
 			if (InCurrentBoneIndex == INDEX_NONE)
 			{
-				return WeightToolProperties->MinColor; // with no bone selected, all vertices are drawn black
+				return FLinearColor::Black; // with no bone selected, all vertices are drawn black
 			}
-			
+
+			// get user-specified colors
+			const TArray<FLinearColor>& Colors = WeightToolProperties->ColorRamp;
+			// get weight value
 			float Value = Weights.GetWeightOfBoneOnVertex(InCurrentBoneIndex, InVertexIndex, Weights.CurrentWeights);
 			Value = FMath::Clamp(Value, 0.0f, 1.0f);
-			
-			// early out zero weights to min color
-			if (Value <= MinimumWeightThreshold)
+
+			// ZERO user supplied colors, then revert to greyscale
+			if (Colors.IsEmpty())
 			{
-				return WeightToolProperties->MinColor;
+				return FMath::Lerp(FLinearColor::Black, FLinearColor::White, Value);
 			}
 
-			// early out full weights to max color
-			if (FMath::IsNearlyEqual(Value, 1.0f))
+			// ONE user defined color, blend it with black
+			if (Colors.Num() == 1)
 			{
-				return WeightToolProperties->MaxColor;
+				return FMath::Lerp(FLinearColor::Black, Colors[0], Value);
 			}
 
-			// get user-specified color ramp for intermediate colors
-			const TArray<FLinearColor>& Colors = WeightToolProperties->ColorRamp;
-	
-			// revert back to simple Lerp(min,max) if user supplied color ramp doesn't have enough colors
-			if (Colors.Num() < 2)
+			// TWO user defined color, simple LERP
+			if (Colors.Num() == 2)
 			{
-				return UE::Geometry::ToVector4<float>(FMath::Lerp(WeightToolProperties->MinColor, WeightToolProperties->MaxColor, Value));
+				return FMath::Lerp(Colors[0], Colors[1], Value);
+			}
+			
+			// blend colors between min and max value
+			constexpr float MinValue = 0.1f;
+			constexpr float MaxValue = 0.9f;
+			
+			// early out zero weights to min color
+			if (Value <= MinValue)
+			{
+				return Colors[0];
 			}
 
-			// otherwise, interpolate within two nearest ramp colors
+			// early out full weights to max color
+			if (Value >= MaxValue)
+			{
+				return Colors.Last();
+			}
+			
+			// remap from 0-1 to range of MinValue to MaxValue
+			const float ScaledValue = (Value - MinValue) * 1.0f / (MaxValue - MinValue);
+			// interpolate within two nearest ramp colors
 			const float PerColorRange = 1.0f / (Colors.Num() - 1);
-			const int ColorIndex = static_cast<int>(Value / PerColorRange);
+			const int ColorIndex = static_cast<int>(ScaledValue / PerColorRange);
 			const float RangeStart = ColorIndex * PerColorRange;
 			const float RangeEnd = (ColorIndex + 1) * PerColorRange;
-			const float Param = (Value - RangeStart) / (RangeEnd - RangeStart);
+			const float Param = (ScaledValue - RangeStart) / (RangeEnd - RangeStart);
 			const FLinearColor& StartColor = Colors[ColorIndex];
 			const FLinearColor& EndColor = Colors[ColorIndex+1];
 			return UE::Geometry::ToVector4<float>(FMath::Lerp(StartColor, EndColor, Param));
@@ -1644,10 +1668,13 @@ void USkinWeightsPaintTool::UpdateVertexColorForSubsetOfVertices()
 	
 	PreviewMesh->DeferredEditMesh([this](FDynamicMesh3& Mesh)
 		{
+			if (CurrentBone == NAME_None)
+			{
+				
+			}
 			TArray<int> ElementIds;
 			UE::Geometry::FDynamicMeshColorOverlay* ColorOverlay = Mesh.Attributes()->PrimaryColors();
-			const int32 CurrentBoneIndex = Weights.Deformer.BoneNameToIndexMap[CurrentBone];
-			const int32 NumVerticesInStamp = VerticesToUpdateColor.Num();
+			const int32 CurrentBoneIndex = GetBoneIndexFromName(CurrentBone);
 			for (const int32 VertexID : VerticesToUpdateColor)
 			{
 				FVector4f NewColor(GetColorOfVertex(VertexID, CurrentBoneIndex));
@@ -1711,7 +1738,7 @@ void USkinWeightsPaintTool::ApplyStamp(const FBrushStampData& Stamp)
 		{
 			// edit weight; either by "Add", "Remove", "Replace", "Multiply"
 			const float UseStrength = CalculateBrushStrengthToUse(WeightToolProperties->BrushMode);
-			const int32 CurrentBoneIndex = Weights.Deformer.BoneNameToIndexMap[CurrentBone];
+			const int32 CurrentBoneIndex = GetCurrentBoneIndex();
 			EditWeightOfBoneOnVertices(
 				WeightToolProperties->BrushMode,
 				CurrentBoneIndex,
@@ -3034,8 +3061,6 @@ void USkinWeightsPaintTool::OnPropertyModified(UObject* ModifiedObject, FPropert
 	const TArray<FString> ColorPropertyNames = {
 		GET_MEMBER_NAME_STRING_CHECKED(USkinWeightsPaintToolProperties, ColorMode),
 		GET_MEMBER_NAME_STRING_CHECKED(USkinWeightsPaintToolProperties, ColorRamp),
-		GET_MEMBER_NAME_STRING_CHECKED(USkinWeightsPaintToolProperties, MinColor),
-		GET_MEMBER_NAME_STRING_CHECKED(USkinWeightsPaintToolProperties, MaxColor),
 		GET_MEMBER_NAME_STRING_CHECKED(FLinearColor, R),
 		GET_MEMBER_NAME_STRING_CHECKED(FLinearColor, G),
 		GET_MEMBER_NAME_STRING_CHECKED(FLinearColor, B),
@@ -3045,8 +3070,6 @@ void USkinWeightsPaintTool::OnPropertyModified(UObject* ModifiedObject, FPropert
 		bVertexColorsNeedUpdated = true;
 
 		// force all colors to have Alpha = 1
-		WeightToolProperties->MinColor.A = 1.f;
-		WeightToolProperties->MaxColor.A = 1.f;
 		for (FLinearColor& Color : WeightToolProperties->ColorRamp)
 		{
 			Color.A = 1.f;
diff --git a/Engine/Plugins/Experimental/MeshModelingToolsetExp/Source/MeshModelingToolsEditorOnlyExp/Public/SkeletalMesh/SkinWeightsPaintTool.h b/Engine/Plugins/Experimental/MeshModelingToolsetExp/Source/MeshModelingToolsEditorOnlyExp/Public/SkeletalMesh/SkinWeightsPaintTool.h
--- a/Engine/Plugins/Experimental/MeshModelingToolsetExp/Source/MeshModelingToolsEditorOnlyExp/Public/SkeletalMesh/SkinWeightsPaintTool.h
+++ b/Engine/Plugins/Experimental/MeshModelingToolsetExp/Source/MeshModelingToolsEditorOnlyExp/Public/SkeletalMesh/SkinWeightsPaintTool.h
@@ -337,10 +337,6 @@ public:
 	EWeightColorMode ColorMode;
 	UPROPERTY(EditAnywhere, Config, Category = MeshDisplay)
 	TArray<FLinearColor> ColorRamp;
-	UPROPERTY(EditAnywhere, Config, Category = MeshDisplay)
-	FLinearColor MinColor;
-	UPROPERTY(EditAnywhere, Config, Category = MeshDisplay)
-	FLinearColor MaxColor;
 
 	// weight editing arguments
 	UPROPERTY(Config)
