Skin weight transfer tool
    - use a source mesh to transfer skin weights from
    - added the possibility of transferring to a subset only
    - added ability to create a new profile (which must be synchronized at the SkeletalMesh level)

diff --git a/Engine/Plugins/Experimental/Animation/SkeletalMeshModelingTools/Source/SkeletalMeshModelingTools/Private/DetailCustomization/SkinWeightDetailCustomization.cpp b/Engine/Plugins/Experimental/Animation/SkeletalMeshModelingTools/Source/SkeletalMeshModelingTools/Private/DetailCustomization/SkinWeightDetailCustomization.cpp
--- a/Engine/Plugins/Experimental/Animation/SkeletalMeshModelingTools/Source/SkeletalMeshModelingTools/Private/DetailCustomization/SkinWeightDetailCustomization.cpp
+++ b/Engine/Plugins/Experimental/Animation/SkeletalMeshModelingTools/Source/SkeletalMeshModelingTools/Private/DetailCustomization/SkinWeightDetailCustomization.cpp
@@ -140,6 +140,8 @@ void FSkinWeightDetailCustomization::CustomizeDetails(IDetailLayoutBuilder& Deta
 		]
 	];
 
+	AddTransferUI(DetailBuilder);
+
 	// hide all base brush properties that have been customized
 	const TSharedRef<IPropertyHandle> BrushModeHandle = DetailBuilder.GetProperty(GET_MEMBER_NAME_CHECKED(USkinWeightsPaintToolProperties, BrushMode));
 	DetailBuilder.HideProperty(BrushModeHandle);
@@ -159,7 +161,7 @@ void FSkinWeightDetailCustomization::CustomizeDetails(IDetailLayoutBuilder& Deta
 	DetailBuilder.HideProperty(ColorModePropHandle);
 }
 
-void FSkinWeightDetailCustomization::AddBrushUI(IDetailLayoutBuilder& DetailBuilder)
+void FSkinWeightDetailCustomization::AddBrushUI(IDetailLayoutBuilder& DetailBuilder) const
 {
 	// custom display of falloff mode as segmented toggle buttons
 	IDetailCategoryBuilder& BrushCategory = DetailBuilder.EditCategory("Brush", FText::GetEmpty(), ECategoryPriority::Important);
@@ -321,7 +323,7 @@ void FSkinWeightDetailCustomization::AddBrushUI(IDetailLayoutBuilder& DetailBuil
 	];
 }
 
-void FSkinWeightDetailCustomization::AddSelectionUI(IDetailLayoutBuilder& DetailBuilder)
+void FSkinWeightDetailCustomization::AddSelectionUI(IDetailLayoutBuilder& DetailBuilder) const
 {
 	// custom display of weight editing tools
 	IDetailCategoryBuilder& EditSelectionCategory = DetailBuilder.EditCategory("Edit Selection", FText::GetEmpty(), ECategoryPriority::Important);
@@ -964,6 +966,46 @@ void FSkinWeightDetailCustomization::AddSelectionUI(IDetailLayoutBuilder& Detail
 	];
 }
 
+void FSkinWeightDetailCustomization::AddTransferUI(IDetailLayoutBuilder& DetailBuilder) const
+{
+	if (!ensure(Tool.IsValid()))
+	{
+		return;
+	}
+	
+	IDetailCategoryBuilder& TransferWeightsCategory = DetailBuilder.EditCategory("WeightTransfer", FText::GetEmpty(), ECategoryPriority::Important);
+	TransferWeightsCategory.InitiallyCollapsed(true);
+
+	TransferWeightsCategory.AddCustomRow(LOCTEXT("TransferWeightsRow", "Transfer Weights"), false)
+	.WholeRowContent()
+	[
+		SNew(SVerticalBox)
+
+		+SVerticalBox::Slot()
+		.Padding(0.f, WeightEditVerticalPadding)
+		[
+			SNew(SBox)
+			[
+				SNew(SButton)
+				.HAlign(HAlign_Center)
+				.Text(LOCTEXT("TransferWeightsButtonLabel", "Transfer Weights"))
+				.ToolTipText(LOCTEXT("TransferButtonTooltip",
+					"Weights are transfered from the source skeletal mesh using inpainting.\n"
+					"This command can operate on selected components when in Mesh mode."))
+				.OnClicked_Lambda([this]()
+				{
+					Tool->TransferWeights();
+					return FReply::Handled();
+				})
+				.IsEnabled_Lambda([this]()
+				{
+					return Tool->GetSourceTarget() != nullptr;
+				})
+			]
+		]
+	];
+}
+
 void SVertexWeightItem::Construct(const FArguments& InArgs, const TSharedRef<STableViewBase>& OwnerTableView)
 {
 	Element = InArgs._Element;
diff --git a/Engine/Plugins/Experimental/Animation/SkeletalMeshModelingTools/Source/SkeletalMeshModelingTools/Private/DetailCustomization/SkinWeightDetailCustomization.h b/Engine/Plugins/Experimental/Animation/SkeletalMeshModelingTools/Source/SkeletalMeshModelingTools/Private/DetailCustomization/SkinWeightDetailCustomization.h
--- a/Engine/Plugins/Experimental/Animation/SkeletalMeshModelingTools/Source/SkeletalMeshModelingTools/Private/DetailCustomization/SkinWeightDetailCustomization.h
+++ b/Engine/Plugins/Experimental/Animation/SkeletalMeshModelingTools/Source/SkeletalMeshModelingTools/Private/DetailCustomization/SkinWeightDetailCustomization.h
@@ -33,10 +33,11 @@ public:
 
 private:
 
-	void AddBrushUI(IDetailLayoutBuilder& DetailBuilder);
-	void AddSelectionUI(IDetailLayoutBuilder& DetailBuilder);
+	void AddBrushUI(IDetailLayoutBuilder& DetailBuilder) const;
+	void AddSelectionUI(IDetailLayoutBuilder& DetailBuilder) const;
+	void AddTransferUI(IDetailLayoutBuilder& DetailBuilder) const;
 
-	IDetailLayoutBuilder* CurrentDetailBuilder;
+	IDetailLayoutBuilder* CurrentDetailBuilder = nullptr;
 	TWeakObjectPtr<USkinWeightsPaintToolProperties> ToolSettings;
 	TWeakObjectPtr<USkinWeightsPaintTool> Tool;
 
diff --git a/Engine/Plugins/Experimental/Animation/SkeletalMeshModelingTools/Source/SkeletalMeshModelingTools/Private/SkeletalMeshModelingToolsEditorMode.cpp b/Engine/Plugins/Experimental/Animation/SkeletalMeshModelingTools/Source/SkeletalMeshModelingTools/Private/SkeletalMeshModelingToolsEditorMode.cpp
--- a/Engine/Plugins/Experimental/Animation/SkeletalMeshModelingTools/Source/SkeletalMeshModelingTools/Private/SkeletalMeshModelingToolsEditorMode.cpp
+++ b/Engine/Plugins/Experimental/Animation/SkeletalMeshModelingTools/Source/SkeletalMeshModelingTools/Private/SkeletalMeshModelingToolsEditorMode.cpp
@@ -58,6 +58,7 @@
 #include "SkeletalMesh/SkeletonEditingTool.h"
 #include "SkeletalMesh/SkinWeightsBindingTool.h"
 #include "SkeletalMesh/SkinWeightsPaintTool.h"
+#include "ToolTargets/SkeletalMeshToolTarget.h"
 
 #define LOCTEXT_NAMESPACE "SkeletalMeshModelingToolsEditorMode"
 
@@ -186,8 +187,12 @@ void USkeletalMeshModelingToolsEditorMode::Enter()
 	UEdMode::Enter();
 
 	UEditorInteractiveToolsContext* InteractiveToolsContext = GetInteractiveToolsContext();
-	
-	InteractiveToolsContext->TargetManager->AddTargetFactory(NewObject<USkeletalMeshComponentToolTargetFactory>(InteractiveToolsContext->TargetManager));
+
+	if (TObjectPtr<UToolTargetManager> ToolTargetManager = InteractiveToolsContext->TargetManager)
+	{
+		ToolTargetManager->AddTargetFactory( NewObject<USkeletalMeshComponentToolTargetFactory>(ToolTargetManager) );
+		ToolTargetManager->AddTargetFactory( NewObject<USkeletalMeshReadOnlyToolTargetFactory>(ToolTargetManager) );
+	}
 
 #if ENABLE_STYLUS_SUPPORT
 	StylusStateTracker = MakeUnique<FStylusStateTracker>();
diff --git a/Engine/Plugins/Experimental/MeshModelingToolsetExp/Source/MeshModelingToolsEditorOnlyExp/Private/SkeletalMesh/SkinWeightsPaintTool.cpp b/Engine/Plugins/Experimental/MeshModelingToolsetExp/Source/MeshModelingToolsEditorOnlyExp/Private/SkeletalMesh/SkinWeightsPaintTool.cpp
--- a/Engine/Plugins/Experimental/MeshModelingToolsetExp/Source/MeshModelingToolsEditorOnlyExp/Private/SkeletalMesh/SkinWeightsPaintTool.cpp
+++ b/Engine/Plugins/Experimental/MeshModelingToolsetExp/Source/MeshModelingToolsEditorOnlyExp/Private/SkeletalMesh/SkinWeightsPaintTool.cpp
@@ -32,13 +32,17 @@
 #include "DynamicSubmesh3.h"
 #include "MeshDescriptionToDynamicMesh.h"
 #include "SkeletalDebugRendering.h"
+#include "ToolSetupUtil.h"
+#include "ToolTargetManager.h"
 #include "Editor/Persona/Public/IPersonaEditorModeManager.h"
 #include "Editor/Persona/Public/PersonaModule.h"
 #include "Preferences/PersonaOptions.h"
 #include "PreviewProfileController.h"
 #include "Animation/SkinWeightProfile.h"
 #include "AnimationRuntime.h"
+#include "DynamicMesh/DynamicVertexSkinWeightsAttribute.h"
 #include "TargetInterfaces/DynamicMeshCommitter.h"
+#include "Operations/TransferBoneWeights.h"
 
 #include UE_INLINE_GENERATED_CPP_BY_NAME(SkinWeightsPaintTool)
 
@@ -111,6 +115,36 @@ USkeletalMeshComponent* GetSkeletalMeshComponent(const UToolTarget* InTarget)
 	return Component;
 }
 
+const FName& CreateNewName()
+{
+	static const FName CreateNew("Create New...");
+	return CreateNew;
+}
+
+FSkinWeightsVertexAttributesRef GetOrCreateSkinWeightsAttribute(FMeshDescription& InMesh, const FName InProfileName)
+{
+	FSkeletalMeshAttributes MeshAttribs(InMesh);
+	const TArray<FName> Profiles = MeshAttribs.GetSkinWeightProfileNames();
+	if (!Profiles.Contains(InProfileName))
+	{
+		ensure( MeshAttribs.RegisterSkinWeightAttribute(InProfileName) );
+	}
+	return MeshAttribs.GetVertexSkinWeights(InProfileName);
+}
+
+bool RenameSkinWeightsAttribute(FMeshDescription& InMesh, const FName InOldName, const FName InNewName)
+{
+	FSkeletalMeshAttributes MeshAttribs(InMesh);
+	const TArray<FName> Profiles = MeshAttribs.GetSkinWeightProfileNames();
+	if (Profiles.Contains(InOldName))
+	{
+		FSkinWeightsVertexAttributesRef NewWeightsAttr = GetOrCreateSkinWeightsAttribute(InMesh, InNewName);
+		NewWeightsAttr = MeshAttribs.GetVertexSkinWeights(InOldName);
+		return MeshAttribs.UnregisterSkinWeightAttribute(InOldName);
+	}
+	return false;
+}
+	
 }
 
 // thread pool to use for async operations
@@ -221,55 +255,99 @@ USkinWeightsPaintToolProperties::USkinWeightsPaintToolProperties()
 	}
 }
 
+namespace SkinWeightLayer
+{
+
+TArray<FName> GetLODs(UToolTarget* InTarget)
+{
+    static TArray<FName> Dummy;
+
+    if (!ensure(InTarget))
+    {
+    	return Dummy;
+    }
+
+    bool bSupportsLODs = false;
+    const TArray<EMeshLODIdentifier> LODIDSs = UE::ToolTarget::GetMeshDescriptionLODs(InTarget, bSupportsLODs);
+    if (!ensure(bSupportsLODs))
+    {
+    	return Dummy;
+    }
+
+    TArray<FName> LODs;
+    LODs.Reserve(LODIDSs.Num());
+    for (const EMeshLODIdentifier LODId: LODIDSs)
+    {
+    	const FName LODName = GetLODName(LODId);
+    	if (LODName != NAME_None)
+    	{
+    		LODs.Add(LODName);
+    	}
+    }
+    ensure(!LODs.IsEmpty());
+    
+    return LODs;
+}
+
+TArray<FName> GetSkinWeightProfilesFunc(const FMeshDescription& InMeshDescription)
+{
+	const FSkeletalMeshConstAttributes MeshAttribs(InMeshDescription);
+	return MeshAttribs.GetSkinWeightProfileNames();
+}
+
+}
+
 TArray<FName> USkinWeightsPaintToolProperties::GetLODsFunc() const
 {
 	static TArray<FName> Dummy;
-
 	if (!ensure(WeightTool && WeightTool->GetTarget()))
 	{
 		return Dummy;
 	}
+	
+	return SkinWeightLayer::GetLODs(WeightTool->GetTarget());
+}
 
-	bool bSupportsLODs = false;
-	const TArray<EMeshLODIdentifier> LODIDSs = UE::ToolTarget::GetMeshDescriptionLODs(WeightTool->GetTarget(), bSupportsLODs);
-	if (!ensure(bSupportsLODs))
+TArray<FName> USkinWeightsPaintToolProperties::GetSkinWeightProfilesFunc() const
+{
+	if (const USkeletalMeshComponent* SkeletalMeshComponent = GetSkeletalMeshComponent(WeightTool->GetTarget()))
 	{
-		return Dummy;
+		const EMeshLODIdentifier LODId = GetLODId(ActiveLOD);
+		const FGetMeshParameters Params(true, LODId);
+		const FMeshDescription* MeshDescription = UE::ToolTarget::GetMeshDescription(WeightTool->GetTarget(), Params);
+		TArray Profiles = SkinWeightLayer::GetSkinWeightProfilesFunc(*MeshDescription);
+		Profiles.Add(CreateNewName());
+		return Profiles;
 	}
 
-	TArray<FName> LODs;
-	LODs.Reserve(LODIDSs.Num());
-	for (const EMeshLODIdentifier LODId: LODIDSs)
+	static const TArray Profiles = {FSkeletalMeshAttributesShared::DefaultSkinWeightProfileName, CreateNewName()};
+	return Profiles;
+}
+
+TArray<FName> USkinWeightsPaintToolProperties::GetSourceLODsFunc() const
+{
+	if (WeightTool->GetSourceTarget())
 	{
-		const FName LODName = GetLODName(LODId);
-		if (LODName != NAME_None)
-		{
-			LODs.Add(LODName);
-		}
+		return SkinWeightLayer::GetLODs(WeightTool->GetSourceTarget());	
 	}
-	ensure(!LODs.IsEmpty());
-	
-	return LODs;
+	return GetLODsFunc();	
 }
 
-TArray<FName> USkinWeightsPaintToolProperties::GetSkinWeightProfilesFunc() const
+TArray<FName> USkinWeightsPaintToolProperties::GetSourceSkinWeightProfilesFunc() const
 {
-	TArray Profiles = {FSkeletalMeshAttributesShared::DefaultSkinWeightProfileName};
-
-	if (const USkeletalMeshComponent* SkeletalMeshComponent = GetSkeletalMeshComponent(WeightTool->GetTarget()))
+	if (USkeletalMesh* SrcSkeletalMesh = SourceSkeletalMesh.Get())
 	{
-		const USkeletalMesh* SkeletalMesh = SkeletalMeshComponent->GetSkeletalMeshAsset();
-		const int32 LOD = static_cast<int32>(GetLODId(ActiveLOD));
-		for (const FSkinWeightProfileInfo& Profile : SkeletalMesh->GetSkinWeightProfiles())
-		{
-			if (Profile.PerLODSourceFiles.Contains(LOD)) 
-			{
-				Profiles.Add(Profile.Name);
-			}
-		}		
+		const EMeshLODIdentifier LODId = GetLODId(SourceLOD);
+		const FGetMeshParameters Params(true, LODId);
+		const FMeshDescription* MeshDescription = UE::ToolTarget::GetMeshDescription(WeightTool->GetSourceTarget(), Params);
+		return SkinWeightLayer::GetSkinWeightProfilesFunc(*MeshDescription);
 	}
-	
-	return Profiles;
+	return GetSkinWeightProfilesFunc();
+}
+
+FName USkinWeightsPaintToolProperties::GetActiveSkinWeightProfile() const
+{
+	return bShowNewProfileName ? NewSkinWeightProfile : ActiveSkinWeightProfile;
 }
 
 FSkinWeightBrushConfig& USkinWeightsPaintToolProperties::GetBrushConfig()
@@ -525,6 +603,9 @@ void FSkinToolWeights::InitializeSkinWeights(
 	const USkeletalMeshComponent* SkeletalMeshComponent,
 	const FMeshDescription* Mesh)
 {
+	static constexpr int32 RootBoneIndex = 0;
+	static constexpr float FullWeight = 1.f;
+
 	// initialize deformer data
 	Deformer.Initialize(SkeletalMeshComponent, Mesh);
 
@@ -553,6 +634,15 @@ void FSkinToolWeights::InitializeSkinWeights(
 			CurrentWeights[VertexIndex].Emplace(BoneIndex, BoneLocalPositionInRefPose, Weight);
 			++InfluenceIndex;
 		}
+
+		// if there are no bone weights, default to root bone 
+		if (InfluenceIndex == 0)
+		{
+			const FVector& RefPoseVertexPosition = Deformer.RefPoseVertexPositions[VertexIndex];
+			const FTransform& InvRefPoseTransform = Deformer.InvCSRefPoseTransforms[RootBoneIndex];
+			const FVector& BoneLocalPositionInRefPose = InvRefPoseTransform.TransformPosition(RefPoseVertexPosition);
+			CurrentWeights[VertexIndex].Emplace(RootBoneIndex, BoneLocalPositionInRefPose, FullWeight);
+		}
 	}
 	
 	// maintain duplicate weight map
@@ -926,8 +1016,8 @@ void USkinWeightsPaintTool::Init(const FToolBuilderState& InSceneState)
 {
 	const UContextObjectStore* ContextObjectStore = InSceneState.ToolManager->GetContextObjectStore();
 	EditorContext = ContextObjectStore->FindContext<USkeletalMeshEditorContextObjectBase>();
-
 	PersonaModeManagerContext = ContextObjectStore->FindContext<UPersonaEditorModeManagerContext>();
+	TargetManager = InSceneState.TargetManager;
 }
 
 void USkinWeightsPaintTool::Setup()
@@ -964,7 +1054,11 @@ void USkinWeightsPaintTool::Setup()
 	WeightToolProperties->ActiveSkinWeightProfile = FSkeletalMeshAttributesShared::DefaultSkinWeightProfileName;
 	WatcherIndex = WeightToolProperties->WatchProperty(WeightToolProperties->ActiveSkinWeightProfile, [this](FName) { OnActiveSkinWeightProfileChanged(); });
 	WeightToolProperties->SilentUpdateWatcherAtIndex(WatcherIndex);
-	
+	WatcherIndex = WeightToolProperties->WatchProperty(WeightToolProperties->NewSkinWeightProfile, [this](FName) { OnNewSkinWeightProfileChanged(); });
+	WeightToolProperties->SilentUpdateWatcherAtIndex(WatcherIndex);
+	WeightToolProperties->SourceSkeletalMesh = nullptr;
+    WeightToolProperties->SourcePreviewOffset = FTransform::Identity;
+		
 	// replace the base brush properties
 	ReplaceToolPropertySource(BrushProperties, WeightToolProperties);
 	BrushProperties = WeightToolProperties;
@@ -1163,7 +1257,18 @@ FBox USkinWeightsPaintTool::GetWorldSpaceFocusBox()
 	}
 
 	// 3. Finally, fallback on component bounds if nothing else is selected
-	return PreviewMesh->GetActor()->GetComponentsBoundingBox();
+	static constexpr bool bNonColliding = true;
+	FBox PreviewBox = PreviewMesh->GetActor()->GetComponentsBoundingBox(bNonColliding);
+	
+	if (WeightToolProperties->bShowSourcePreview && SourcePreviewMesh)
+	{
+		if (AActor* SourceActor = SourcePreviewMesh->GetActor())
+		{
+			PreviewBox += SourceActor->GetComponentsBoundingBox(bNonColliding);
+		}
+	}
+
+	return PreviewBox;
 }
 
 void USkinWeightsPaintTool::OnUpdateModifierState(int ModifierID, bool bIsOn)
@@ -1273,11 +1378,11 @@ void USkinWeightsPaintTool::PostEditMeshInitialization(
 
 	// update weights
 	Weights = FSkinToolWeights();
-	if (!IsProfileValid(WeightToolProperties->ActiveSkinWeightProfile))
+	if (!IsProfileValid(WeightToolProperties->GetActiveSkinWeightProfile()))
 	{
 		WeightToolProperties->ActiveSkinWeightProfile = FSkeletalMeshAttributesShared::DefaultSkinWeightProfileName;
 	}
-	Weights.Profile = WeightToolProperties->ActiveSkinWeightProfile;
+	Weights.Profile = WeightToolProperties->GetActiveSkinWeightProfile();
 	Weights.InitializeSkinWeights(InComponent, &InMeshDescription);
 	bVertexColorsNeedUpdated = true;
 	
@@ -2056,12 +2161,14 @@ void USkinWeightsPaintTool::OnShutdown(EToolShutdownType ShutdownType)
 	{
 		PersonaModeManagerContext->GetPersonaEditorModeManager()->DeactivateMode(FPersonaEditModes::SkeletonSelection);
 	}
+
+	ResetSourceForTransfer();
 }
 
 void USkinWeightsPaintTool::BeginChange()
 {
 	const EMeshLODIdentifier LOD = GetLODId(WeightToolProperties->ActiveLOD);
-	const FName SkinProfile = WeightToolProperties->ActiveSkinWeightProfile;
+	const FName SkinProfile = WeightToolProperties->GetActiveSkinWeightProfile();
 	ActiveChange = MakeUnique<FMeshSkinWeightsChange>(LOD, SkinProfile);
 }
 
@@ -2104,7 +2211,7 @@ void USkinWeightsPaintTool::ExternalUpdateSkinWeightLayer(const EMeshLODIdentifi
 		None
 	} State = ESkinWeightChangeState::None;
 	
-	if (InSkinWeightProfile != WeightToolProperties->ActiveSkinWeightProfile)
+	if (InSkinWeightProfile != WeightToolProperties->GetActiveSkinWeightProfile())
 	{
 		WeightToolProperties->ActiveSkinWeightProfile = InSkinWeightProfile;
 		State = ESkinWeightChangeState::SkinProfile;
@@ -2501,6 +2608,137 @@ void USkinWeightsPaintTool::NormalizeWeights()
 	ApplyWeightEditsAsTransaction(WeightEditsFromNormalization, TransactionLabel);
 }
 
+void USkinWeightsPaintTool::TransferWeights()
+{
+	using namespace UE::Geometry;
+	using namespace UE::AnimationCore;
+
+	if (!SourceTarget)
+	{
+		return;
+	}
+
+	const EMeshLODIdentifier TargetLODId = GetLODId(WeightToolProperties->ActiveLOD);
+    const FGetMeshParameters TargetParams(true, TargetLODId);
+	FDynamicMesh3 TargetMesh = UE::ToolTarget::GetDynamicMeshCopy(Target, TargetParams);
+
+	const EMeshLODIdentifier SourceLODId = GetLODId(WeightToolProperties->SourceLOD);
+	const FGetMeshParameters SourceParams(true, SourceLODId);
+	const FDynamicMesh3 SourceMesh = UE::ToolTarget::GetDynamicMeshCopy(SourceTarget, SourceParams);
+	
+	if (!SourceMesh.HasAttributes() || !SourceMesh.Attributes()->HasBones())
+	{
+		return;
+	}
+	
+	if (SourceMesh.Attributes()->GetNumBones() == 0)
+	{
+		return;
+	}
+	
+	FTransferBoneWeights TransferBoneWeights(&SourceMesh, WeightToolProperties->SourceSkinWeightProfile);
+	TransferBoneWeights.TransferMethod = FTransferBoneWeights::ETransferBoneWeightsMethod::InpaintWeights;
+	
+	if (!TargetMesh.HasAttributes())
+	{
+		TargetMesh.EnableAttributes();
+	}
+
+	FDynamicMeshAttributeSet* TargetAttributes = TargetMesh.Attributes();
+	if (!TargetAttributes->HasBones())
+	{
+		TargetAttributes->CopyBoneAttributes(*SourceMesh.Attributes());
+	}
+
+	// NOTE should we expose all the options?
+	// 	TransferBoneWeights.NormalThreshold;
+	// 	TransferBoneWeights.SearchRadius
+	// 	TransferBoneWeights.NumSmoothingIterations;
+	// 	TransferBoneWeights.SmoothingStrength;
+	// 	TransferBoneWeights.LayeredMeshSupport;
+	// 	TransferBoneWeights.ForceInpaintWeightMapName;
+	
+	if (TransferBoneWeights.Validate() != EOperationValidationResult::Ok)
+	{
+		return;
+	}
+
+	if (WeightToolProperties->EditingMode == EWeightEditMode::Mesh)
+	{
+		TransferBoneWeights.TargetVerticesSubset = GetSelectedVertices();
+	}
+
+	const FName TargetProfile = WeightToolProperties->GetActiveSkinWeightProfile();
+	if (TransferBoneWeights.TransferWeightsToMesh(TargetMesh, TargetProfile))
+	{
+		// store weight edits
+		BeginChange();
+		FMultiBoneWeightEdits WeightEdits;
+		
+		{	
+			FDynamicMeshVertexSkinWeightsAttribute* TransferedSkinWeights = TargetAttributes->GetSkinWeightsAttribute(TargetProfile);
+			check(TransferedSkinWeights);
+			
+			const bool bUseSubset = !TransferBoneWeights.TargetVerticesSubset.IsEmpty();
+
+			static constexpr float ZeroWeight = 0.f;
+			
+			const int32 NumVertices = bUseSubset ? TransferBoneWeights.TargetVerticesSubset.Num() : TargetMesh.VertexCount();
+			for (int32 VertexIndex = 0; VertexIndex < NumVertices; ++VertexIndex)
+			{
+				const int32 VertexID = bUseSubset ? TransferBoneWeights.TargetVerticesSubset[VertexIndex] : VertexIndex;
+				
+				// remove all weight on vertex
+				const VertexWeights& VertexBoneWeights = Weights.PreChangeWeights[VertexID];
+				if (!VertexBoneWeights.IsEmpty())
+				{
+					for (const FVertexBoneWeight& BoneWeight : VertexBoneWeights)
+					{
+						const float OldWeight = BoneWeight.Weight;
+						WeightEdits.MergeSingleEdit(BoneWeight.BoneIndex, VertexID, OldWeight, ZeroWeight);
+					}
+				}
+				else
+				{
+					WeightEdits.MergeSingleEdit(0, VertexID, 1.f, ZeroWeight);
+				}
+
+				// update with new weight
+				FBoneWeights TransferedBoneWeights;
+				TransferedSkinWeights->GetValue(VertexID, TransferedBoneWeights);
+				for (FBoneWeight BoneWeight: TransferedBoneWeights)
+				{
+					const int32 BoneIndex = BoneWeight.GetBoneIndex();					
+					const float OldWeight = Weights.GetWeightOfBoneOnVertex(BoneIndex, VertexID, Weights.PreChangeWeights);
+					const float NewWeight = BoneWeight.GetWeight();
+					WeightEdits.MergeSingleEdit(BoneIndex, VertexID, OldWeight, NewWeight);
+				}
+			}
+		}
+
+		// set new weights (we could probably use TransferedSkinWeights instead of a full convert) 
+		FDynamicMeshToMeshDescription Converter;
+		Converter.Convert(&TargetMesh, *EditedMesh);
+
+		// update weights
+		Weights = FSkinToolWeights();
+		Weights.Profile = TargetProfile;
+		Weights.InitializeSkinWeights(GetSkeletalMeshComponent(Target), EditedMesh);
+		bVertexColorsNeedUpdated = true;
+
+		// store weight edits in the active change & commit
+		{
+			for (const TTuple<int32, FSingleBoneWeightEdits>& BoneWeightEdits : WeightEdits.PerBoneWeightEdits)
+			{
+				ActiveChange->AddBoneWeightEdit(BoneWeightEdits.Value);
+			}
+
+			static const FText TransactionLabel = LOCTEXT("TransferWeightsChange", "Transfer skin weights.");
+			EndChange(TransactionLabel);
+		}
+	}
+}
+
 void USkinWeightsPaintTool::HandleSkeletalMeshModified(const TArray<FName>& InBoneNames, const ESkeletalMeshNotifyType InNotifyType)
 {
 	switch (InNotifyType)
@@ -2565,7 +2803,7 @@ void USkinWeightsPaintTool::OnActiveLODChanged()
 	}
 
 	// reinitialize all mesh data structures
-	const UE::Geometry::FDynamicMesh3 DynamicMesh = UE::ToolTarget::GetDynamicMeshCopy(Target, Params);
+	const FDynamicMesh3 DynamicMesh = UE::ToolTarget::GetDynamicMeshCopy(Target, Params);
 	PostEditMeshInitialization(Component, DynamicMesh, *EditedMesh);
 }
 
@@ -2577,17 +2815,28 @@ void USkinWeightsPaintTool::OnActiveSkinWeightProfileChanged()
 		return;
 	}
 
+	WeightToolProperties->bShowNewProfileName = WeightToolProperties->ActiveSkinWeightProfile == CreateNewName();
+
 	if (IsSelectionIsolated())
 	{
 		FinishIsolatedSelection();
 	}
 
-	if (!IsProfileValid(WeightToolProperties->ActiveSkinWeightProfile))
+	if (WeightToolProperties->bShowNewProfileName)
+	{
+		if (!IsProfileValid(WeightToolProperties->NewSkinWeightProfile))
+		{
+			GetOrCreateSkinWeightsAttribute(*EditedMesh, WeightToolProperties->NewSkinWeightProfile);
+		} 
+	}
+	
+	if (!IsProfileValid(WeightToolProperties->GetActiveSkinWeightProfile()))
 	{
 		WeightToolProperties->ActiveSkinWeightProfile = FSkeletalMeshAttributesShared::DefaultSkinWeightProfileName;
+		WeightToolProperties->bShowNewProfileName = false;
 	}
 	
-	if (WeightToolProperties->ActiveSkinWeightProfile == Weights.Profile)
+	if (WeightToolProperties->GetActiveSkinWeightProfile() == Weights.Profile)
 	{
 		return;
 	}
@@ -2597,26 +2846,27 @@ void USkinWeightsPaintTool::OnActiveSkinWeightProfileChanged()
 
 	// re-init Weights with new skin profile
 	Weights = FSkinToolWeights();
-	Weights.Profile = WeightToolProperties->ActiveSkinWeightProfile;
+	Weights.Profile = WeightToolProperties->GetActiveSkinWeightProfile();
 	Weights.InitializeSkinWeights(SkeletalMeshComponent, EditedMesh);
 	bVertexColorsNeedUpdated = true;
 }
 
-bool USkinWeightsPaintTool::IsProfileValid(const FName InProfileName) const
+void USkinWeightsPaintTool::OnNewSkinWeightProfileChanged()
 {
-	const USkeletalMeshComponent* SkeletalMeshComponent = GetSkeletalMeshComponent(Target);
-	if (!SkeletalMeshComponent)
+	if (WeightToolProperties->bShowNewProfileName && WeightToolProperties->NewSkinWeightProfile != Weights.Profile)
 	{
-		return false;
+		const bool bRenamed = RenameSkinWeightsAttribute(*EditedMesh, Weights.Profile, WeightToolProperties->NewSkinWeightProfile);
+		if (ensure(bRenamed))
+		{
+			Weights.Profile = WeightToolProperties->NewSkinWeightProfile;
+		}
 	}
+}
 
-	// check SkeletalMesh
-	const TArray<FSkinWeightProfileInfo>& SkeletalMeshProfiles = SkeletalMeshComponent->GetSkeletalMeshAsset()->GetSkinWeightProfiles();
-	bool bHasProfile = SkeletalMeshProfiles.ContainsByPredicate([InProfileName](const FSkinWeightProfileInfo& Info)
-	{
-		return Info.Name == InProfileName;
-	});
-	if (!bHasProfile)
+bool USkinWeightsPaintTool::IsProfileValid(const FName InProfileName) const
+{
+	const USkeletalMeshComponent* SkeletalMeshComponent = GetSkeletalMeshComponent(Target);
+	if (!SkeletalMeshComponent)
 	{
 		return false;
 	}
@@ -2624,7 +2874,7 @@ bool USkinWeightsPaintTool::IsProfileValid(const FName InProfileName) const
 	// check current MeshDescription
 	const FSkeletalMeshConstAttributes MeshAttribs(*EditedMesh);
 	const TArray<FName> MeshDescProfiles = MeshAttribs.GetSkinWeightProfileNames();
-	bHasProfile = MeshDescProfiles.ContainsByPredicate([InProfileName](const FName Name)
+	const bool bHasProfile = MeshDescProfiles.ContainsByPredicate([InProfileName](const FName Name)
 	{
 		return Name == InProfileName;
 	});
@@ -2806,7 +3056,7 @@ void USkinWeightsPaintTool::FinishIsolatedSelection()
 
 	// copy the remapped weights back to the full mesh
 	const FSkeletalMeshConstAttributes MeshAttribs(*PartialMeshDescription.Get());
-	const FSkinWeightsVertexAttributesConstRef AllVertexWeights = MeshAttribs.GetVertexSkinWeights(WeightToolProperties->ActiveSkinWeightProfile);
+	const FSkinWeightsVertexAttributesConstRef AllVertexWeights = MeshAttribs.GetVertexSkinWeights(WeightToolProperties->GetActiveSkinWeightProfile());
 	const int32 NumVerticesInPartialMesh = PartialMeshDescription.Get()->Vertices().Num();
 	for (int32 VertexIndexPartial = 0; VertexIndexPartial < NumVerticesInPartialMesh; VertexIndexPartial++)
 	{
@@ -3075,9 +3325,86 @@ void USkinWeightsPaintTool::OnPropertyModified(UObject* ModifiedObject, FPropert
 			Color.A = 1.f;
 		}
 	}
+
+	if (ModifiedProperty->GetName() == GET_MEMBER_NAME_STRING_CHECKED(USkinWeightsPaintToolProperties, SourceSkeletalMesh))
+	{
+		ResetSourceForTransfer(WeightToolProperties->SourceSkeletalMesh.Get());
+	}
+
+	if (ModifiedProperty->GetName() == GET_MEMBER_NAME_STRING_CHECKED(USkinWeightsPaintToolProperties, SourceLOD))
+	{
+		if (SourcePreviewMesh)
+		{
+			const EMeshLODIdentifier SourceLODId = GetLODId(WeightToolProperties->SourceLOD);
+			const FGetMeshParameters SourceParams(true, SourceLODId);
+			SourcePreviewMesh->ReplaceMesh(UE::ToolTarget::GetDynamicMeshCopy(SourceTarget, SourceParams));
+		}
+	}
+
+	if (ModifiedProperty->GetName() == GET_MEMBER_NAME_STRING_CHECKED(USkinWeightsPaintToolProperties, bShowSourcePreview))
+	{
+		if (SourcePreviewMesh)
+		{
+			SourcePreviewMesh->SetVisible(WeightToolProperties->bShowSourcePreview);
+		}
+	}
+
+	if (ModifiedProperty->GetName() == GET_MEMBER_NAME_STRING_CHECKED(USkinWeightsPaintToolProperties, SourcePreviewOffset))
+	{
+		if (SourcePreviewMesh)
+		{
+			SourcePreviewMesh->SetTransform(WeightToolProperties->SourcePreviewOffset);
+		}
+	}
 	
 	SetFocusInViewport();
 }
 
+void USkinWeightsPaintTool::ResetSourceForTransfer(USkeletalMesh* InSkeletalMesh)
+{
+	if (SourcePreviewMesh)
+	{
+		SourcePreviewMesh->SetVisible(false);
+		SourcePreviewMesh->Disconnect();
+		SourcePreviewMesh = nullptr;
+	}
+
+	if (SourceTarget)
+	{
+		SourceTarget = nullptr;
+	}
+
+	if (InSkeletalMesh)
+	{
+		SourceTarget = TargetManager->BuildTarget(InSkeletalMesh, FToolTargetTypeRequirements());
+
+		SourcePreviewMesh = NewObject<UPreviewMesh>(this);
+		SourcePreviewMesh->CreateInWorld(TargetWorld.Get(), FTransform::Identity);
+
+		if (USkeletalMeshComponent* SkeletalMeshComponent = GetSkeletalMeshComponent(Target))
+		{
+			const FBoxSphereBounds TargetBounds = SkeletalMeshComponent->GetSkeletalMeshAsset()->GetBounds();
+			const FBoxSphereBounds SourceBounds = InSkeletalMesh->GetBounds();
+
+			FTransform Transform = UE::ToolTarget::GetLocalToWorldTransform(Target);
+			FVector Location = Transform.GetLocation();
+			Location.X += TargetBounds.GetBoxExtrema(1).X;
+			Location.X += 1.1 * SourceBounds.GetBoxExtrema(1).X;
+			Transform.SetLocation(Location);
+			WeightToolProperties->SourcePreviewOffset = Transform;
+		}
+		
+		SourcePreviewMesh->SetTransform(WeightToolProperties->SourcePreviewOffset);
+
+		ToolSetupUtil::ApplyRenderingConfigurationToPreview(SourcePreviewMesh, SourceTarget);
+		SourcePreviewMesh->SetTangentsMode(EDynamicMeshComponentTangentsMode::AutoCalculated);
+		SourcePreviewMesh->ReplaceMesh(UE::ToolTarget::GetDynamicMeshCopy(SourceTarget));
+
+		const FComponentMaterialSet MaterialSet = UE::ToolTarget::GetMaterialSet(SourceTarget);
+		SourcePreviewMesh->SetMaterials(MaterialSet.Materials);
+
+		SourcePreviewMesh->SetVisible(WeightToolProperties->bShowSourcePreview);
+	}
+}
 
 #undef LOCTEXT_NAMESPACE
\ No newline at end of file
diff --git a/Engine/Plugins/Experimental/MeshModelingToolsetExp/Source/MeshModelingToolsEditorOnlyExp/Public/SkeletalMesh/SkinWeightsPaintTool.h b/Engine/Plugins/Experimental/MeshModelingToolsetExp/Source/MeshModelingToolsEditorOnlyExp/Public/SkeletalMesh/SkinWeightsPaintTool.h
--- a/Engine/Plugins/Experimental/MeshModelingToolsetExp/Source/MeshModelingToolsEditorOnlyExp/Public/SkeletalMesh/SkinWeightsPaintTool.h
+++ b/Engine/Plugins/Experimental/MeshModelingToolsetExp/Source/MeshModelingToolsEditorOnlyExp/Public/SkeletalMesh/SkinWeightsPaintTool.h
@@ -21,6 +21,7 @@
 #include "SkinWeightsPaintTool.generated.h"
 
 
+class USkeletalMeshComponentReadOnlyToolTarget;
 struct FMeshDescription;
 class USkinWeightsPaintTool;
 class UPolygonSelectionMechanic;
@@ -370,6 +371,15 @@ public:
 	UPROPERTY(EditAnywhere, Category = SkinWeightLayer, meta = (DisplayName = "Active Profile", GetOptions = GetSkinWeightProfilesFunc))
 	FName ActiveSkinWeightProfile = FSkeletalMeshAttributesShared::DefaultSkinWeightProfileName;
 	
+	// new profile properties
+	UPROPERTY(meta = (TransientToolProperty))
+	bool bShowNewProfileName = false;
+	UPROPERTY(EditAnywhere, Category = SkinWeightLayer, meta = (TransientToolProperty, DisplayName = "New Profile Name",
+		EditCondition = bShowNewProfileName, HideEditConditionToggle, NoResetToDefault))
+	FName NewSkinWeightProfile = "Profile";
+
+	FName GetActiveSkinWeightProfile() const;
+	
 	// pointer back to paint tool
 	TObjectPtr<USkinWeightsPaintTool> WeightTool;
 
@@ -378,12 +388,33 @@ public:
 	void SetColorMode(EWeightColorMode InColorMode);
 	void SetBrushMode(EWeightEditOperation InBrushMode);
 
+	// transfer
+	UPROPERTY(EditAnywhere, Transient, Category = WeightTransfer)
+	TWeakObjectPtr<USkeletalMesh> SourceSkeletalMesh;
+	
+	UPROPERTY(EditAnywhere, Category = "WeightTransfer|SkinWeightLayer", meta = (GetOptions = GetSourceLODsFunc))
+	FName SourceLOD = "LOD0";
+	
+	UPROPERTY(EditAnywhere, Category = "WeightTransfer|SkinWeightLayer", meta = (DisplayName = "Source Profile", GetOptions = GetSourceSkinWeightProfilesFunc))
+	FName SourceSkinWeightProfile = FSkeletalMeshAttributesShared::DefaultSkinWeightProfileName;
+	
+	UPROPERTY(EditAnywhere, Transient, Category = "WeightTransfer|Preview")
+	bool bShowSourcePreview = false;
+	
+	UPROPERTY(EditAnywhere, Transient, Category = "WeightTransfer|Preview")
+	FTransform SourcePreviewOffset = FTransform::Identity;
+	
 private:
 	
 	UFUNCTION()
 	TArray<FName> GetLODsFunc() const;
 	UFUNCTION()
 	TArray<FName> GetSkinWeightProfilesFunc() const;
+
+	UFUNCTION()
+	TArray<FName> GetSourceLODsFunc() const;
+	UFUNCTION()
+	TArray<FName> GetSourceSkinWeightProfilesFunc() const;
 };
 
 // An interactive tool for painting and editing skin weights.
@@ -429,6 +460,7 @@ public:
 	void PruneWeights(const float Threshold);
 	void AverageWeights();
 	void NormalizeWeights();
+	void TransferWeights();
 	
 	// method to set weights directly (numeric input, for example)
 	void EditWeightsOnVertices(
@@ -496,6 +528,9 @@ public:
 	DECLARE_MULTICAST_DELEGATE(FOnWeightsChanged);
 	FOnWeightsChanged OnWeightsChanged;
 
+	// gets the current source target
+	UToolTarget* GetSourceTarget() const { return SourceTarget; }
+
 protected:
 
 	virtual void ApplyStamp(const FBrushStampData& Stamp);
@@ -620,12 +655,24 @@ protected:
 	// skin weight layer
 	void OnActiveLODChanged();
 	void OnActiveSkinWeightProfileChanged();
+	void OnNewSkinWeightProfileChanged();
 	bool IsProfileValid(const FName InProfileName) const;
 
+	// global properties stored on initialization
 	UPROPERTY()
 	TWeakObjectPtr<USkeletalMeshEditorContextObjectBase> EditorContext = nullptr;
 	UPROPERTY()
 	TWeakObjectPtr<UPersonaEditorModeManagerContext> PersonaModeManagerContext = nullptr;
+	UPROPERTY()
+	TWeakObjectPtr<UToolTargetManager> TargetManager = nullptr;
+
+	// skin weights transfer properties
+	UPROPERTY()
+	TObjectPtr<UPreviewMesh> SourcePreviewMesh = nullptr;
+	UPROPERTY()
+	TObjectPtr<UToolTarget> SourceTarget = nullptr;
+
+	void ResetSourceForTransfer(USkeletalMesh* InSkeletalMesh = nullptr);
 	
 	// editor state to restore when exiting the paint tool
 	FString PreviewProfileToRestore;
diff --git a/Engine/Source/Runtime/SkeletalMeshDescription/Private/SkeletalMeshAttributes.cpp b/Engine/Source/Runtime/SkeletalMeshDescription/Private/SkeletalMeshAttributes.cpp
--- a/Engine/Source/Runtime/SkeletalMeshDescription/Private/SkeletalMeshAttributes.cpp
+++ b/Engine/Source/Runtime/SkeletalMeshDescription/Private/SkeletalMeshAttributes.cpp
@@ -121,6 +121,28 @@ bool FSkeletalMeshAttributes::RegisterSkinWeightAttribute(const FName InProfileN
 	return MeshDescription.VertexAttributes().RegisterAttribute<int32[]>(AttributeName, 1, 0, EMeshAttributeFlags::None).IsValid();
 }
 
+bool FSkeletalMeshAttributes::UnregisterSkinWeightAttribute(const FName InProfileName)
+{
+	if (!IsValidSkinWeightProfileName(InProfileName))
+	{
+		return false;
+	}
+
+	const FName AttributeName = CreateSkinWeightAttributeName(InProfileName);
+	if (!ensure(AttributeName.IsValid()))
+	{
+		return false;
+	}
+
+	// Attribute not there?
+	if (!MeshDescription.VertexAttributes().HasAttribute(AttributeName))
+	{
+		return false;
+	}
+
+	MeshDescription.VertexAttributes().UnregisterAttribute(AttributeName);
+	return true;
+}
 
 FSkinWeightsVertexAttributesRef FSkeletalMeshAttributes::GetVertexSkinWeights(const FName InProfileName)
 {
diff --git a/Engine/Source/Runtime/SkeletalMeshDescription/Public/SkeletalMeshAttributes.h b/Engine/Source/Runtime/SkeletalMeshDescription/Public/SkeletalMeshAttributes.h
--- a/Engine/Source/Runtime/SkeletalMeshDescription/Public/SkeletalMeshAttributes.h
+++ b/Engine/Source/Runtime/SkeletalMeshDescription/Public/SkeletalMeshAttributes.h
@@ -261,6 +261,11 @@ public:
 	/// Returns \c false if the attribute was already registered or if IsValidSkinWeightProfileName() returned false.
 	SKELETALMESHDESCRIPTION_API bool RegisterSkinWeightAttribute(const FName InProfileName);
 
+	/// Unregister an existing skin weight profile with the given name
+	/// Returns \c true if theprofile was successfully unregistered.
+	/// Returns \c false if the attribute wasn't registered or if \c InProfileName is empty.
+	SKELETALMESHDESCRIPTION_API bool UnregisterSkinWeightAttribute(const FName InProfileName);
+
 	/// Returns the skin weight profile given by its name. NAME_None corresponds to the default profile.
 	SKELETALMESHDESCRIPTION_API FSkinWeightsVertexAttributesRef GetVertexSkinWeights(const FName InProfileName = NAME_None);
 
