Fixes issue, skin paint brush has anisotropic falloff on some meshes.

diff --git a/Engine/Plugins/Experimental/MeshModelingToolsetExp/Source/MeshModelingToolsEditorOnlyExp/Private/SkeletalMesh/SkinWeightsPaintTool.cpp b/Engine/Plugins/Experimental/MeshModelingToolsetExp/Source/MeshModelingToolsEditorOnlyExp/Private/SkeletalMesh/SkinWeightsPaintTool.cpp
--- a/Engine/Plugins/Experimental/MeshModelingToolsetExp/Source/MeshModelingToolsEditorOnlyExp/Private/SkeletalMesh/SkinWeightsPaintTool.cpp
+++ b/Engine/Plugins/Experimental/MeshModelingToolsetExp/Source/MeshModelingToolsEditorOnlyExp/Private/SkeletalMesh/SkinWeightsPaintTool.cpp
@@ -1587,57 +1587,34 @@ void USkinWeightsPaintTool::CalculateVertexROI(
 
 	if (WeightToolProperties->GetBrushConfig().FalloffMode == EWeightBrushFalloffMode::Surface)
 	{
-		// get coordinate frame from stamp
-		auto GetFrameFromStamp = [](const FBrushStampData& InStamp) -> FFrame3d
-		{
-			const FVector3d Origin = InStamp.WorldPosition;
-			const FVector3d Normal = InStamp.WorldNormal;
-			FVector3d NonCollinear = Normal;
-			// get a guaranteed non collinear vector to the normal
-			// doesn't matter where in the plane, stamp is radially symmetric
-			do 
-			{
-				NonCollinear.X = FMath::RandRange(-1.0f, 1.0f);
-				NonCollinear.Y = FMath::RandRange(-1.0f, 1.0f);
-				NonCollinear.Z = FMath::RandRange(-1.0f, 1.0f);
-				NonCollinear.Normalize();
-					
-			} while (FMath::Abs(NonCollinear.Dot(Normal)) > 0.8f);
-
-			const FVector3d Plane = Normal.Cross(NonCollinear);
-			const FVector3d Cross = Plane.Cross(Normal);
-			return FFrame3d(Origin, Cross, Plane, Normal);
-				
-		};
-		const FFrame3d SeedFrame = GetFrameFromStamp(Stamp);
-			
 		// create the ExpMap generator, computes vertex polar coordinates in a plane tangent to the surface
 		const FDynamicMesh3* Mesh = PreviewMesh->GetPreviewDynamicMesh();
+		FFrame3d SeedFrame = Mesh->GetTriFrame(TriangleUnderStamp);
+		SeedFrame.Origin = Stamp.WorldPosition;
+		
 		TMeshLocalParam<FDynamicMesh3> Param(Mesh);
-		Param.ParamMode = ELocalParamTypes::ExponentialMapUpwindAvg;
+		Param.ParamMode = ELocalParamTypes::PlanarProjection;
 		const FIndex3i TriVerts = Mesh->GetTriangle(TriangleUnderStamp);
-		Param.ComputeToMaxDistance(SeedFrame, TriVerts, Stamp.Radius);
+		Param.ComputeToMaxDistance(SeedFrame, TriVerts, Stamp.Radius * 1.5f);
+		// store vertices under the brush and their distances from the stamp
+		const float StampRadSq = FMath::Pow(Stamp.Radius, 2);
+		for (int32 VertexID : Mesh->VertexIndicesItr())
 		{
-			// store vertices under the brush and their distances from the stamp
-			const float StampRadSq = FMath::Pow(Stamp.Radius, 2);
-			for (int32 VertexID : Mesh->VertexIndicesItr())
+			if (!Param.HasUV(VertexID))
 			{
-				if (!Param.HasUV(VertexID))
-				{
-					continue;
-				}
-				
-				FVector2d UV = Param.GetUV(VertexID);
-				const float DistSq = UV.SizeSquared();
-				if (DistSq >= StampRadSq)
-				{
-					continue;
-				}
-
-				
-				VertexFalloffs.Add(DistanceToFalloff(VertexID, DistSq));
-				VertexIDs.Add(VertexID);
+				continue;
+			}
+			
+			FVector2d UV = Param.GetUV(VertexID);
+			const float DistSq = UV.SizeSquared();
+			if (DistSq >= StampRadSq)
+			{
+				continue;
 			}
+
+			
+			VertexFalloffs.Add(DistanceToFalloff(VertexID, DistSq));
+			VertexIDs.Add(VertexID);
 		}
 		
 		return;
