Fixed skin weight mirror issue, consecutive operations give inconsistent results.

diff --git a/Engine/Plugins/Experimental/MeshModelingToolsetExp/Source/MeshModelingToolsEditorOnlyExp/Private/SkeletalMesh/SkinWeightsPaintTool.cpp b/Engine/Plugins/Experimental/MeshModelingToolsetExp/Source/MeshModelingToolsEditorOnlyExp/Private/SkeletalMesh/SkinWeightsPaintTool.cpp
--- a/Engine/Plugins/Experimental/MeshModelingToolsetExp/Source/MeshModelingToolsEditorOnlyExp/Private/SkeletalMesh/SkinWeightsPaintTool.cpp
+++ b/Engine/Plugins/Experimental/MeshModelingToolsetExp/Source/MeshModelingToolsEditorOnlyExp/Private/SkeletalMesh/SkinWeightsPaintTool.cpp
@@ -2231,6 +2231,8 @@ void FSkinMirrorData::RegenerateMirrorData(
 	bIsInitialized = false;
 	Axis = InMirrorAxis;
 	Direction = InMirrorDirection;
+	BoneMap.Reset();
+	VertexMap.Reset();
 	
 	// build bone map for mirroring
 	// TODO, provide some way to edit the mirror bone mapping, either by providing a UMirrorDataTable input or editing directly in the hierarchy view.
@@ -2347,8 +2349,8 @@ void USkinWeightsPaintTool::MirrorWeights(EAxis::Type Axis, EMirrorDirection Dir
 		Direction);
 
 	// get a reference to the mirror tables
-	const TMap<int32, int32>& BoneMap = MirrorData.GetBoneMap();
-	const TMap<int32, int32>& VertexMirrorMap = MirrorData.GetVertexMap(); // <Target, Source>
+	const TMap<BoneIndex, BoneIndex>& BoneMap = MirrorData.GetBoneMap();
+	const TMap<VertexIndex, VertexIndex>& VertexMirrorMap = MirrorData.GetVertexMap(); // <Target, Source>
 
 	// get set of vertices to mirror
 	TArray<VertexIndex> AllVerticesToEdit = GetVerticesToEdit();
@@ -2366,7 +2368,7 @@ void USkinWeightsPaintTool::MirrorWeights(EAxis::Type Axis, EMirrorDirection Dir
 		else
 		{
 			// vertex is located on the source side (to copy FROM), so we need to search for it's mirror target vertex
-			for (const TPair<int32, int32>& ToFromPair : VertexMirrorMap)
+			for (const TPair<VertexIndex, VertexIndex>& ToFromPair : VertexMirrorMap)
 			{
 				if (ToFromPair.Value != SelectedVertex)
 				{
@@ -2382,8 +2384,8 @@ void USkinWeightsPaintTool::MirrorWeights(EAxis::Type Axis, EMirrorDirection Dir
 	FMultiBoneWeightEdits WeightEditsFromMirroring;
 	for (const VertexIndex VertexToMirror : VerticesToMirror)
 	{
-		const int32 SourceVertexID = VertexMirrorMap[VertexToMirror];
-		const int32 TargetVertexID = VertexToMirror;
+		const VertexIndex SourceVertexID = VertexMirrorMap[VertexToMirror];
+		const VertexIndex TargetVertexID = VertexToMirror;
 
 		// remove all weight on vertex
 		for (const FVertexBoneWeight& TargetBoneWeight : Weights.PreChangeWeights[TargetVertexID])
@@ -2396,7 +2398,7 @@ void USkinWeightsPaintTool::MirrorWeights(EAxis::Type Axis, EMirrorDirection Dir
 		// copy source weights, but with mirrored bones
 		for (const FVertexBoneWeight& SourceBoneWeight : Weights.PreChangeWeights[SourceVertexID])
 		{
-			const int32 MirroredBoneIndex = BoneMap[SourceBoneWeight.BoneIndex];
+			const BoneIndex MirroredBoneIndex = BoneMap[SourceBoneWeight.BoneIndex];
 			const float OldWeight = Weights.GetWeightOfBoneOnVertex(MirroredBoneIndex, TargetVertexID, Weights.PreChangeWeights);
 			const float NewWeight = SourceBoneWeight.Weight;
 			WeightEditsFromMirroring.MergeSingleEdit(MirroredBoneIndex, TargetVertexID, OldWeight, NewWeight);
