Skin Weights Paint Tool: fixed crashes using non manifold vertices
    - as the tool uses a mix of FMeshDescription & FDynamicMesh3, we need to ensure that vertex ids represent the source ids where necessary.

#rnx

diff --git a/Engine/Plugins/Experimental/MeshModelingToolsetExp/Source/MeshModelingToolsEditorOnlyExp/Private/SkeletalMesh/SkinWeightsPaintTool.cpp b/Engine/Plugins/Experimental/MeshModelingToolsetExp/Source/MeshModelingToolsEditorOnlyExp/Private/SkeletalMesh/SkinWeightsPaintTool.cpp
--- a/Engine/Plugins/Experimental/MeshModelingToolsetExp/Source/MeshModelingToolsEditorOnlyExp/Private/SkeletalMesh/SkinWeightsPaintTool.cpp
+++ b/Engine/Plugins/Experimental/MeshModelingToolsetExp/Source/MeshModelingToolsEditorOnlyExp/Private/SkeletalMesh/SkinWeightsPaintTool.cpp
@@ -159,9 +159,9 @@ class FPaintToolWeightsDataSource : public UE::Geometry::TBoneWeightsDataSource<
 {
 public:
 
-	FPaintToolWeightsDataSource(const SkinPaintTool::FSkinToolWeights* InWeights)
-	:
-	Weights(InWeights) 
+	FPaintToolWeightsDataSource(const FSkinToolWeights* InWeights, const FDynamicMesh3& InDynaMesh)
+		: Weights(InWeights)
+		, NonManifoldMappingSupport(InDynaMesh)
 	{
 		checkSlow(Weights);
 	}
@@ -170,27 +170,37 @@ public:
 
 	virtual int32 GetBoneNum(const int32 VertexID) override
 	{
-		return Weights->PreChangeWeights[VertexID].Num();
+		const int32 SrcVertexId = GetSourceVertexId(VertexID);
+		return Weights->PreChangeWeights[SrcVertexId].Num();
 	}
 
 	virtual int32 GetBoneIndex(const int32 VertexID, const int32 Index) override
 	{
-		return Weights->PreChangeWeights[VertexID][Index].BoneIndex;
+		const int32 SrcVertexId = GetSourceVertexId(VertexID);
+		return Weights->PreChangeWeights[SrcVertexId][Index].BoneIndex;
 	}
 
 	virtual float GetBoneWeight(const int32 VertexID, const int32 Index) override
 	{
-		return Weights->PreChangeWeights[VertexID][Index].Weight;
+		const int32 SrcVertexId = GetSourceVertexId(VertexID);
+		return Weights->PreChangeWeights[SrcVertexId][Index].Weight;
 	}
 
 	virtual float GetWeightOfBoneOnVertex(const int32 VertexID, const int32 BoneIndex) override
 	{
-		return Weights->GetWeightOfBoneOnVertex(BoneIndex, VertexID, Weights->PreChangeWeights);
+		const int32 SrcVertexId = GetSourceVertexId(VertexID);
+		return Weights->GetWeightOfBoneOnVertex(BoneIndex, SrcVertexId, Weights->PreChangeWeights);
 	}
 
 protected:
 
-	const SkinPaintTool::FSkinToolWeights* Weights = nullptr;
+	int32 GetSourceVertexId(const int32 InVertexID) const
+	{
+		return NonManifoldMappingSupport.GetOriginalNonManifoldVertexID(InVertexID);
+	}
+	
+	const FSkinToolWeights* Weights = nullptr;
+	const UE::Geometry::FNonManifoldMappingSupport NonManifoldMappingSupport;
 };
 
 void FDirectEditWeightState::Reset()
@@ -1562,9 +1572,9 @@ double USkinWeightsPaintTool::EstimateMaximumTargetDimension()
 }
 
 void USkinWeightsPaintTool::CalculateVertexROI(
-	const FBrushStampData& Stamp,
-	TArray<VertexIndex>& VertexIDs,
-	TArray<float>& VertexFalloffs)
+	const FBrushStampData& InStamp,
+	TArray<VertexIndex>& OutVertexIDs,
+	TArray<float>& OutVertexFalloffs)
 {
 	using namespace UE::Geometry;
 
@@ -1581,19 +1591,27 @@ void USkinWeightsPaintTool::CalculateVertexROI(
 	{
 		const IPrimitiveComponentBackedTarget* TargetComponent = Cast<IPrimitiveComponentBackedTarget>(Target);
 		const FTransform3d Transform(TargetComponent->GetWorldTransform());
-		const FVector3d StampPosLocal = Transform.InverseTransformPosition((FVector3d)Stamp.WorldPosition);
+		const FVector3d StampPosLocal = Transform.InverseTransformPosition(InStamp.WorldPosition);
 		const float RadiusSqr = CurrentBrushRadius * CurrentBrushRadius;
 		const FDynamicMesh3* Mesh = PreviewMesh->GetPreviewDynamicMesh();
 		const FAxisAlignedBox3d QueryBox(StampPosLocal, CurrentBrushRadius);
 		VerticesOctree->RangeQuery(QueryBox,
 			[&](int32 VertexID) { return FVector3d::DistSquared(Mesh->GetVertex(VertexID), StampPosLocal) < RadiusSqr; },
-			VertexIDs);
+			OutVertexIDs);
 
-		for (const int32 VertexID : VertexIDs)
+		const FNonManifoldMappingSupport NonManifoldMappingSupport(*Mesh);
+		TArray<VertexIndex> SourceVertexIDs;
+		SourceVertexIDs.Reserve(OutVertexIDs.Num());
+		OutVertexFalloffs.Reserve(OutVertexIDs.Num());
+		for (const int32 VertexID : OutVertexIDs)
 		{
 			const float DistSq = FVector3d::DistSquared(Mesh->GetVertex(VertexID), StampPosLocal);
-			VertexFalloffs.Add(DistanceToFalloff(VertexID, DistSq));
+
+			const int32 SrcVertexId = NonManifoldMappingSupport.GetOriginalNonManifoldVertexID(VertexID);
+			SourceVertexIDs.Add(SrcVertexId);
+			OutVertexFalloffs.Add(DistanceToFalloff(SrcVertexId, DistSq));
 		}
+		OutVertexIDs = MoveTemp(SourceVertexIDs);
 		
 		return;
 	}
@@ -1603,14 +1621,15 @@ void USkinWeightsPaintTool::CalculateVertexROI(
 		// create the ExpMap generator, computes vertex polar coordinates in a plane tangent to the surface
 		const FDynamicMesh3* Mesh = PreviewMesh->GetPreviewDynamicMesh();
 		FFrame3d SeedFrame = Mesh->GetTriFrame(TriangleUnderStamp);
-		SeedFrame.Origin = Stamp.WorldPosition;
+		SeedFrame.Origin = InStamp.WorldPosition;
 		
 		TMeshLocalParam<FDynamicMesh3> Param(Mesh);
 		Param.ParamMode = ELocalParamTypes::PlanarProjection;
 		const FIndex3i TriVerts = Mesh->GetTriangle(TriangleUnderStamp);
-		Param.ComputeToMaxDistance(SeedFrame, TriVerts, Stamp.Radius * 1.5f);
+		Param.ComputeToMaxDistance(SeedFrame, TriVerts, InStamp.Radius * 1.5f);
 		// store vertices under the brush and their distances from the stamp
-		const float StampRadSq = FMath::Pow(Stamp.Radius, 2);
+		const float StampRadSq = FMath::Pow(InStamp.Radius, 2);
+		const FNonManifoldMappingSupport NonManifoldMappingSupport(*Mesh);
 		for (int32 VertexID : Mesh->VertexIndicesItr())
 		{
 			if (!Param.HasUV(VertexID))
@@ -1625,9 +1644,9 @@ void USkinWeightsPaintTool::CalculateVertexROI(
 				continue;
 			}
 
-			
-			VertexFalloffs.Add(DistanceToFalloff(VertexID, DistSq));
-			VertexIDs.Add(VertexID);
+			const int32 SrcVertexId = NonManifoldMappingSupport.GetOriginalNonManifoldVertexID(VertexID);
+			OutVertexFalloffs.Add(DistanceToFalloff(SrcVertexId, DistSq));
+			OutVertexIDs.Add(SrcVertexId);
 		}
 		
 		return;
@@ -2031,9 +2050,10 @@ void USkinWeightsPaintTool::InitializeSmoothWeightsOperator()
 		return;
 	}
 
-	// NOTE: this could probably be initialized lazily as it's only used with the relax brush 
-	SmoothWeightsDataSource = MakeUnique<FPaintToolWeightsDataSource>(&Weights);
-	SmoothWeightsOp = MakeUnique<UE::Geometry::TSmoothBoneWeights<int32, float>>(PreviewMesh->GetMesh(), SmoothWeightsDataSource.Get());
+	// NOTE: this could probably be initialized lazily as it's only used with the relax brush
+	const FDynamicMesh3* DynaMesh = PreviewMesh->GetMesh();
+	SmoothWeightsDataSource = MakeUnique<FPaintToolWeightsDataSource>(&Weights, *DynaMesh);
+	SmoothWeightsOp = MakeUnique<UE::Geometry::TSmoothBoneWeights<int32, float>>(DynaMesh, SmoothWeightsDataSource.Get());
 	SmoothWeightsOp->MinimumWeightThreshold = MinimumWeightThreshold;
 }
 
@@ -2676,23 +2696,26 @@ void USkinWeightsPaintTool::TransferWeights()
 			static constexpr float ZeroWeight = 0.f;
 			
 			const int32 NumVertices = bUseSubset ? TransferBoneWeights.TargetVerticesSubset.Num() : TargetMesh.VertexCount();
+			const FNonManifoldMappingSupport NonManifoldMappingSupport(TargetMesh);
+			
 			for (int32 VertexIndex = 0; VertexIndex < NumVertices; ++VertexIndex)
 			{
 				const int32 VertexID = bUseSubset ? TransferBoneWeights.TargetVerticesSubset[VertexIndex] : VertexIndex;
+				const int32 SrcVertexID = NonManifoldMappingSupport.GetOriginalNonManifoldVertexID(VertexID);
 				
 				// remove all weight on vertex
-				const VertexWeights& VertexBoneWeights = Weights.PreChangeWeights[VertexID];
+				const VertexWeights& VertexBoneWeights = Weights.PreChangeWeights[SrcVertexID];
 				if (!VertexBoneWeights.IsEmpty())
 				{
 					for (const FVertexBoneWeight& BoneWeight : VertexBoneWeights)
 					{
 						const float OldWeight = BoneWeight.Weight;
-						WeightEdits.MergeSingleEdit(BoneWeight.BoneIndex, VertexID, OldWeight, ZeroWeight);
+						WeightEdits.MergeSingleEdit(BoneWeight.BoneIndex, SrcVertexID, OldWeight, ZeroWeight);
 					}
 				}
 				else
 				{
-					WeightEdits.MergeSingleEdit(0, VertexID, 1.f, ZeroWeight);
+					WeightEdits.MergeSingleEdit(0, SrcVertexID, 1.f, ZeroWeight);
 				}
 
 				// update with new weight
@@ -2701,9 +2724,9 @@ void USkinWeightsPaintTool::TransferWeights()
 				for (FBoneWeight BoneWeight: TransferedBoneWeights)
 				{
 					const int32 BoneIndex = BoneWeight.GetBoneIndex();					
-					const float OldWeight = Weights.GetWeightOfBoneOnVertex(BoneIndex, VertexID, Weights.PreChangeWeights);
+					const float OldWeight = Weights.GetWeightOfBoneOnVertex(BoneIndex, SrcVertexID, Weights.PreChangeWeights);
 					const float NewWeight = BoneWeight.GetWeight();
-					WeightEdits.MergeSingleEdit(BoneIndex, VertexID, OldWeight, NewWeight);
+					WeightEdits.MergeSingleEdit(BoneIndex, SrcVertexID, OldWeight, NewWeight);
 				}
 			}
 		}
diff --git a/Engine/Plugins/Experimental/MeshModelingToolsetExp/Source/MeshModelingToolsEditorOnlyExp/Public/SkeletalMesh/SkinWeightsPaintTool.h b/Engine/Plugins/Experimental/MeshModelingToolsetExp/Source/MeshModelingToolsEditorOnlyExp/Public/SkeletalMesh/SkinWeightsPaintTool.h
--- a/Engine/Plugins/Experimental/MeshModelingToolsetExp/Source/MeshModelingToolsEditorOnlyExp/Public/SkeletalMesh/SkinWeightsPaintTool.h
+++ b/Engine/Plugins/Experimental/MeshModelingToolsetExp/Source/MeshModelingToolsEditorOnlyExp/Public/SkeletalMesh/SkinWeightsPaintTool.h
@@ -550,9 +550,9 @@ protected:
 	// stamp
 	float CalculateBrushFalloff(float Distance) const;
 	void CalculateVertexROI(
-		const FBrushStampData& Stamp,
-		TArray<VertexIndex>& VertexIDs,
-		TArray<float>& VertexFalloffs);
+		const FBrushStampData& InStamp,
+		TArray<VertexIndex>& OutVertexIDs,
+		TArray<float>& OutVertexFalloffs);
 	float CalculateBrushStrengthToUse(EWeightEditOperation EditMode) const;
 	bool bInvertStroke = false;
 	FBrushStampData StartStamp;
