expose missing ui toggle for 'apply thicken shells' setting in vox wrap tool

diff --git a/Engine/Plugins/Experimental/MeshModelingToolsetExp/Source/MeshModelingToolsExp/Public/VoxelSolidifyMeshesTool.h b/Engine/Plugins/Experimental/MeshModelingToolsetExp/Source/MeshModelingToolsExp/Public/VoxelSolidifyMeshesTool.h
--- a/Engine/Plugins/Experimental/MeshModelingToolsetExp/Source/MeshModelingToolsExp/Public/VoxelSolidifyMeshesTool.h
+++ b/Engine/Plugins/Experimental/MeshModelingToolsetExp/Source/MeshModelingToolsExp/Public/VoxelSolidifyMeshesTool.h
@@ -38,7 +38,7 @@ public:
 	bool bSolidAtBoundaries = true;
 
 	/** If true, uses the ThickenShells setting */
-	UPROPERTY()
+	UPROPERTY(EditAnywhere, Category = Options)
 	bool bApplyThickenShells = false;
 
 	/** Thicken open-boundary surfaces (extrude them inwards) to ensure they are captured in the VoxWrap output. Units are in world space. */
