ModelingTools: Change PatternTool to use new gizmo utility functions to create its pattern sub gizmo.

diff --git a/Engine/Plugins/Experimental/MeshModelingToolsetExp/Source/MeshModelingToolsExp/Private/PatternTool.cpp b/Engine/Plugins/Experimental/MeshModelingToolsetExp/Source/MeshModelingToolsExp/Private/PatternTool.cpp
--- a/Engine/Plugins/Experimental/MeshModelingToolsetExp/Source/MeshModelingToolsExp/Private/PatternTool.cpp
+++ b/Engine/Plugins/Experimental/MeshModelingToolsetExp/Source/MeshModelingToolsExp/Private/PatternTool.cpp
@@ -28,11 +28,15 @@
 #include "ContextObjectStore.h"
 #include "BaseGizmos/AxisPositionGizmo.h"
 #include "BaseGizmos/CombinedTransformGizmo.h"
+#include "BaseGizmos/ComponentBoundTransformProxy.h"
 #include "BaseGizmos/GizmoActor.h"
 #include "BaseGizmos/GizmoComponents.h"
+#include "BaseGizmos/GizmoUtil.h"
 #include "BaseGizmos/GizmoViewContext.h"
 #include "BaseGizmos/GizmoBoxComponent.h"
+#include "BaseGizmos/PlanePositionGizmo.h"
 #include "BaseGizmos/TransformGizmoUtil.h"
+#include "BaseGizmos/TransformSubGizmoUtil.h"
 
 #include "Engine/World.h"
 
@@ -99,45 +103,6 @@ const FToolTargetTypeRequirements& UPatternToolBuilder::GetTargetRequirements()
 	return TypeRequirements;
 }
 
-/*
- * Custom GizmoActor Factory
- */
-class FPatternToolGizmoActorFactory : public FCombinedTransformGizmoActorFactory
-{
-public:
-	FPatternToolGizmoActorFactory(UGizmoViewContext* GizmoViewContextIn)
-	: FCombinedTransformGizmoActorFactory(GizmoViewContextIn)
-	{
-		EnableElements = ETransformGizmoSubElements::None;
-	}
-
-	/**
-	 * @param World the UWorld to create the new Actor in
-	 * @return new ACombinedTransformGizmoActor instance with members initialized with Components suitable for a transformation Gizmo
-	 */
-	virtual ACombinedTransformGizmoActor* CreateNewGizmoActor(UWorld* World) const override
-	{
-		FActorSpawnParameters SpawnInfo;
-		ACombinedTransformGizmoActor* NewActor = World->SpawnActor<ACombinedTransformGizmoActor>(FVector::ZeroVector, FRotator::ZeroRotator, SpawnInfo);
-
-		UGizmoBoxComponent* Component = AGizmoActor::AddDefaultBoxComponent(World, NewActor, GizmoViewContext, FLinearColor::Red, FVector::ZeroVector);
-		Component->LineThickness = 5.0f;
-		Component->Dimensions = FVector(Component->LineThickness);
-		Component->NotifyExternalPropertyUpdates();
-		
-		if ((EnableElements & ETransformGizmoSubElements::TranslateAxisX) != ETransformGizmoSubElements::None)
-		{
-			NewActor->TranslateX = Component;
-		}
-		else if ((EnableElements & ETransformGizmoSubElements::TranslatePlaneXY) != ETransformGizmoSubElements::None)
-		{
-			NewActor->TranslateXY = Component;
-		}
-		
-		return NewActor;
-	}
-};
-
 /*
  * PatternGenerators
  */
@@ -729,24 +694,30 @@ void UPatternTool::Setup()
 	PlaneMechanic->Setup(this);
 	PlaneMechanic->Initialize( GetTargetWorld(), CurrentStartFrameWorld );
 	PlaneMechanic->OnPlaneChanged.AddLambda([this]() { OnMainFrameUpdated(); });
+
+	UGizmoViewContext* GizmoViewContext = nullptr;
+	UContextObjectStore* ContextObjectStore = GetToolManager()->GetContextObjectStore();
+	if (ContextObjectStore)
+	{
+		GizmoViewContext = ContextObjectStore->FindContext<UGizmoViewContext>();
+	}
 	
-	PatternGizmoProxy = NewObject<UTransformProxy>(this);
+	// The advantage of using a fully visually symmetric component here is that we can use the same component 
+	//  regardless of whether we're translating in a line or a plane, so we don't need to swap it out when that
+	//  changes.
+	// Note that this helper method already attaches the component to the parent gizmo actor.
+	UGizmoBoxComponent* BoxComponent = AGizmoActor::AddDefaultBoxComponent(GetWorld(), PlaneMechanic->PlaneTransformGizmo->GetGizmoActor(),
+		GizmoViewContext, FLinearColor::Red, FVector::ZeroVector);
+	BoxComponent->LineThickness = 5.0f;
+	BoxComponent->Dimensions = FVector(BoxComponent->LineThickness * 1.5);
+	BoxComponent->NotifyExternalPropertyUpdates();
+	PatternGizmoComponent = BoxComponent;
+
+	PatternGizmoProxy = NewObject<UComponentBoundTransformProxy>(this);
 	PatternGizmoProxy->OnTransformChanged.AddUObject(this, &UPatternTool::OnTransformGizmoUpdated);
-
-	// The gizmo used to define pattern extents/radius is a bit hacky because it uses a combined transform
-	// gizmo which only ever has a single component of the gizmo at a time (it will either allow movement in
-	// a single axis or a single plane) and is visually just a box. This made the code much simpler than
-	// creating an entirely new gizmo but probably shouldn't be used as a pattern for other custom gizmos.
-	// This tool registers its own builder and unregisters the builder at shutdown, the builder shouldn't be
-	// relied on elsewhere.
-	UCombinedTransformGizmoBuilder* CustomThreeAxisBuilder = NewObject<UCombinedTransformGizmoBuilder>();
-	GizmoActorBuilder = MakeShared<FPatternToolGizmoActorFactory>(GetToolManager()->GetContextObjectStore()->FindContext<UGizmoViewContext>());
-	CustomThreeAxisBuilder->AxisPositionBuilderIdentifier = UInteractiveGizmoManager::DefaultAxisPositionBuilderIdentifier;
-	CustomThreeAxisBuilder->PlanePositionBuilderIdentifier = UInteractiveGizmoManager::DefaultPlanePositionBuilderIdentifier;
-	CustomThreeAxisBuilder->AxisAngleBuilderIdentifier = UInteractiveGizmoManager::DefaultAxisAngleBuilderIdentifier;
-	CustomThreeAxisBuilder->GizmoActorBuilder = GizmoActorBuilder;
-	GetToolManager()->GetPairedGizmoManager()->RegisterGizmoType(PatternToolThreeAxisTransformBuilderIdentifier, CustomThreeAxisBuilder);
-	bPatternToolThreeAxisTransformGizmoRegistered = true;
+	PatternGizmoProxy->BindToComponent(BoxComponent, 
+		// bStoreScaleSeparately, doesn't matter because scale isn't used
+		false);
 
 	// Needs to be called before any of the watchers call ResetTransformGizmoPosition in order to set the
 	// proxy as the target of the gizmo, otherwise the gizmo will attempt to dereference an invalid StateTarget
@@ -805,10 +776,6 @@ void UPatternTool::OnShutdown(EToolShutdownType ShutdownType)
 	DragAlignmentMechanic->Shutdown();
 
 	GetToolManager()->GetPairedGizmoManager()->DestroyAllGizmosByOwner(this);
-
-	ensure(bPatternToolThreeAxisTransformGizmoRegistered);
-	GetToolManager()->GetPairedGizmoManager()->DeregisterGizmoType(PatternToolThreeAxisTransformBuilderIdentifier);
-	bPatternToolThreeAxisTransformGizmoRegistered = false;
 	
 	OnSourceVisibilityToggled(true);
 
@@ -992,7 +959,6 @@ void UPatternTool::OnMainFrameUpdated()
 	CurrentStartFrameWorld = PlaneMechanic->Plane;
 
 	MarkPatternDirty();
-	ResetTransformGizmoPosition();
 }
 
 
@@ -1117,6 +1083,17 @@ void UPatternTool::OnTransformGizmoUpdated(UTransformProxy* Proxy, FTransform Tr
 
 void UPatternTool::ResetTransformGizmoPosition()
 {
+	if (!PatternGizmo || !ensure(PatternGizmoProxy))
+	{
+		return;
+	}
+
+	// There are two ways we can deal with our transform, particularly the rotation. Either we can
+	//  keep the sub gizmo unrotated in parent gizmo space, and change which axes we use on the gizmo,
+	//  or we can have the gizmo operate on the same axis/axes and change the rotation here to align
+	//  that axis in the proper direction.
+	// For now we do the latter. 
+
 	FVector3d OffsetFromOrigin = FVector3d::ZeroVector;
 	double DistanceFromOriginX = 0.0f;
 	double DistanceFromOriginY = 0.0f;
@@ -1191,29 +1168,51 @@ void UPatternTool::ResetTransformGizmoPosition()
 
 	const FVector3d ProxyTranslation = CurrentStartFrameWorld.Origin + FRotator(CurrentStartFrameWorld.Rotation).RotateVector(OffsetFromOrigin);
 
-	PatternGizmo->ReinitializeGizmoTransform(FTransform(FQuat(CurrentStartFrameWorld.Rotation * FQuaterniond(RotationInLocalSpace)), ProxyTranslation));
+	TGuardValue<bool> SetPivotGuard(PatternGizmoProxy->bSetPivotMode, true);
+	PatternGizmoProxy->SetTransform(FTransform(RotationInLocalSpace, OffsetFromOrigin) * CurrentStartFrameWorld.ToFTransform());
 }
 
 void UPatternTool::ReconstructTransformGizmos()
 {
-	if (bPatternToolThreeAxisTransformGizmoRegistered)
+	UInteractiveGizmoManager* GizmoManager = GetToolManager()->GetPairedGizmoManager();
+	if (PatternGizmo)
 	{
-		// Determining which elements of the CombinedTransformGizmo will be used
-		GizmoActorBuilder->EnableElements = bUsingSingleAxis ? ETransformGizmoSubElements::TranslateAxisX : ETransformGizmoSubElements::TranslatePlaneXY;
+		GizmoManager->DestroyGizmo(PatternGizmo);
+	}
 
-		// Reconstructing gizmos with proper elements and transform
-		UInteractiveGizmoManager* GizmoManager = GetToolManager()->GetPairedGizmoManager();
-		GizmoManager->DestroyAllGizmosByOwner(this);
-		
-		PatternGizmo = Cast<UCombinedTransformGizmo>(GizmoManager->CreateGizmo(PatternToolThreeAxisTransformBuilderIdentifier, FString(), this));
+	// There are a couple ways to create a sub gizmo here. The proper way is to add the extra element individually,
+	//  as we do here, but it could have been done with a TRS subgizmo.
 
-		// Necessary to force underlying AxisSources to axes in local space to restrict gizmo movement to tool frame axes
-		PatternGizmo->bUseContextCoordinateSystem = false;
-		PatternGizmo->CurrentCoordinateSystem = EToolContextCoordinateSystem::Local;
+	UE::GizmoUtil::FTransformSubGizmoCommonParams Params;
+	Params.Component = PatternGizmoComponent;
+	Params.TransformProxy = PatternGizmoProxy;
+	Params.bManipulatesRootComponent = false;
+	Params.bAxisIsBasedOnRootComponent = false;
 
-		PatternGizmo->SetActiveTarget(PatternGizmoProxy);
-		ResetTransformGizmoPosition();
+	if (bUsingSingleAxis)
+	{
+		UAxisPositionGizmo* AxisGizmo = UE::GizmoUtil::CreateGizmoViaSimpleBuilder<UAxisPositionGizmo>(
+			GetToolManager()->GetPairedGizmoManager(), FString(), this);
+		Params.OuterForSubobjects = AxisGizmo;
+		Params.Axis = EAxis::X;
+		AxisGizmo->InitializeAsTranslateGizmo(Params,
+			// No shared state
+			nullptr);
+		PatternGizmo = AxisGizmo;
 	}
+	else
+	{
+		UPlanePositionGizmo* PlaneGizmo = UE::GizmoUtil::CreateGizmoViaSimpleBuilder<UPlanePositionGizmo>(
+			GetToolManager()->GetPairedGizmoManager(), FString(), this);
+		Params.OuterForSubobjects = PlaneGizmo;
+		Params.Axis = EAxis::Z;
+		PlaneGizmo->InitializeAsTranslateGizmo(Params,
+			// No shared state
+			nullptr);
+		PatternGizmo = PlaneGizmo;
+	}
+
+	ResetTransformGizmoPosition();
 }
 
 
diff --git a/Engine/Plugins/Experimental/MeshModelingToolsetExp/Source/MeshModelingToolsExp/Public/PatternTool.h b/Engine/Plugins/Experimental/MeshModelingToolsetExp/Source/MeshModelingToolsExp/Public/PatternTool.h
--- a/Engine/Plugins/Experimental/MeshModelingToolsetExp/Source/MeshModelingToolsExp/Public/PatternTool.h
+++ b/Engine/Plugins/Experimental/MeshModelingToolsetExp/Source/MeshModelingToolsExp/Public/PatternTool.h
@@ -13,7 +13,6 @@
 #include "TransformTypes.h"
 #include "PatternTool.generated.h"
 
-class FCombinedTransformGizmoActorFactory;
 class UBaseAxisTranslationGizmo;
 class UAxisAngleGizmo;
 class UDragAlignmentMechanic;
@@ -22,10 +21,12 @@ class UTransformProxy;
 class UPreviewGeometry;
 class UDynamicMesh;
 class UDynamicMeshComponent;
+class UInteractiveGizmo;
 class UStaticMesh;
 class UStaticMeshComponent;
 class UMaterialInterface;
 class UConstructionPlaneMechanic;
+class UComponentBoundTransformProxy;
 
 
 /**
@@ -490,17 +491,15 @@ protected:
 	/**
 	 * Pattern Gizmo:
 	 */
-	FString PatternToolThreeAxisTransformBuilderIdentifier = TEXT("PatternToolThreeAxisTransformBuilderIdentifier");
-	FString PatternToolAxisPositionBuilderIdentifier = TEXT("PatternToolAxisPositionBuilderIdentifier");
-	FString PatternToolPlanePositionBuilderIdentifier = TEXT("PatternToolPlanePositionBuilderIdentifier");
-	TSharedPtr<FCombinedTransformGizmoActorFactory> GizmoActorBuilder;
-	bool bPatternToolThreeAxisTransformGizmoRegistered = false;
-	
+
+	UPROPERTY()
+	TObjectPtr<UComponentBoundTransformProxy> PatternGizmoProxy = nullptr;
+
 	UPROPERTY()
-	TObjectPtr<UTransformProxy> PatternGizmoProxy = nullptr;
+	TObjectPtr<UInteractiveGizmo> PatternGizmo = nullptr;
 
 	UPROPERTY()
-	TObjectPtr<UCombinedTransformGizmo> PatternGizmo = nullptr;
+	TObjectPtr<UPrimitiveComponent> PatternGizmoComponent = nullptr;
 
 	// If true, Settings->SingleAxis is being used. If false, Settings->SinglePlane is being used
 	bool bUsingSingleAxis;
@@ -514,8 +513,6 @@ protected:
 	void ResetTransformGizmoPosition();
 	void ReconstructTransformGizmos();
 	
-
-	
 	UPROPERTY()
 	TObjectPtr<UDragAlignmentMechanic> DragAlignmentMechanic = nullptr;
 
