Multiple fixes/additions to direct skin weight editing toolset.

diff --git a/Engine/Plugins/Experimental/Animation/SkeletalMeshModelingTools/Source/SkeletalMeshModelingTools/Private/DetailCustomization/SkinWeightDetailCustomization.cpp b/Engine/Plugins/Experimental/Animation/SkeletalMeshModelingTools/Source/SkeletalMeshModelingTools/Private/DetailCustomization/SkinWeightDetailCustomization.cpp
--- a/Engine/Plugins/Experimental/Animation/SkeletalMeshModelingTools/Source/SkeletalMeshModelingTools/Private/DetailCustomization/SkinWeightDetailCustomization.cpp
+++ b/Engine/Plugins/Experimental/Animation/SkeletalMeshModelingTools/Source/SkeletalMeshModelingTools/Private/DetailCustomization/SkinWeightDetailCustomization.cpp
@@ -22,10 +22,9 @@
 float FSkinWeightDetailCustomization::WeightEditVerticalPadding = 4.0f;
 float FSkinWeightDetailCustomization::WeightEditHorizontalPadding = 2.0f;
 
-// add a buffer to the weight sliders to prevent ever fully getting a value to 1 or 0 using the slider alone
-// doing so will remove other influences and cause the slider to no longer function as all other influences
-// will be culled by normalization thus making the slider "stuck" at full value.
-constexpr float SliderBuffer = UE::AnimationCore::BoneWeightThreshold * 2.f;
+static FName ColumnName_Bone = "Bone";
+static FName ColumnName_Weight = "Weight";
+static FName ColumnName_Prune = "Prune";
 
 FSkinWeightDetailCustomization::~FSkinWeightDetailCustomization()
 {
@@ -50,10 +49,6 @@
 	ToolSettings = Cast<USkinWeightsPaintToolProperties>(DetailObjects[0]);
 	Tool = ToolSettings->WeightTool;
 	Tool->OnSelectionChanged.AddSP(this, &FSkinWeightDetailCustomization::OnSelectionChanged);
-
-	// Edit SkinWeightLayer category first 
-	IDetailCategoryBuilder& SkinWeightLayerCategory = DetailBuilder.EditCategory("SkinWeightLayer", FText::GetEmpty(), ECategoryPriority::Important);
-	SkinWeightLayerCategory.InitiallyCollapsed(true);
 	
 	// custom display of falloff mode as segmented toggle buttons
 	IDetailCategoryBuilder& EditModeCategory = DetailBuilder.EditCategory("Weight Editing Mode", FText::GetEmpty(), ECategoryPriority::Important);
@@ -142,6 +137,10 @@
 
 	AddTransferUI(DetailBuilder);
 
+	// Edit SkinWeightLayer category 
+	IDetailCategoryBuilder& SkinWeightLayerCategory = DetailBuilder.EditCategory("SkinWeightLayer", FText::GetEmpty(), ECategoryPriority::Important);
+	SkinWeightLayerCategory.InitiallyCollapsed(true);
+
 	// hide all base brush properties that have been customized
 	const TSharedRef<IPropertyHandle> BrushModeHandle = DetailBuilder.GetProperty(GET_MEMBER_NAME_CHECKED(USkinWeightsPaintToolProperties, BrushMode));
 	DetailBuilder.HideProperty(BrushModeHandle);
@@ -174,14 +173,12 @@
 		.Padding(2.0f)
 		[
 			SNew(SSegmentedControl<EWeightEditOperation>)
-			.ToolTipText(LOCTEXT("FloodTooltip",
-				"Add: applies the current weight plus the flood value to the new weight.\n"
-				"Replace: applies the current weight minus the strength value to the new weight.\n"
-				"Multiply: applies the current weight multiplied by the strength value to the new weight.\n"
-				"Relax: applies the average of the connected (by edge) vertex weights to the new vertex weight, blended by the strength.\n"
-				"This command operates on the selected bone(s) and selected vertices.\n"
-				"If no bones are selected, ALL bones are considered.\n"
-				"If no vertices are selected, ALL vertices are considered."))
+			.ToolTipText(LOCTEXT("BrushTooltip",
+				"Select the operation to apply when using the brush.\n"
+				"Add: applies the current weight plus the Strength value to the new weight.\n"
+				"Replace: applies the current weight minus the Strength value to the new weight.\n"
+				"Multiply: applies the current weight multiplied by the Strength value to the new weight.\n"
+				"Relax: applies the average of the connected (by edge) vertex weights to the new vertex weight, blended by the Strength."))
 			.Value_Lambda([this]()
 			{
 				return ToolSettings->BrushMode;
@@ -499,12 +496,64 @@
 				})
 			]
 		]
+
+		+SVerticalBox::Slot()
+		[
+			SNew(SHorizontalBox)
+		
+			+SHorizontalBox::Slot()
+			.Padding(2.f, WeightEditVerticalPadding)
+			[
+				SNew(SButton)
+				.HAlign(HAlign_Center)
+				.VAlign(VAlign_Center)
+				.Text(LOCTEXT("AffectedSelectionButtonLabel", "Affected"))
+				.ToolTipText(LOCTEXT("AffectedSelectionTooltip",
+						"Select vertices that are affected by the currently selected bone(s).\n"
+						"Holding Shift or Ctrl will add or subtract affected vertices from the current selection."))
+				.OnClicked_Lambda([this]()
+				{
+					Tool->SelectAffected();
+					return FReply::Handled();
+				})
+			]
+
+			+SHorizontalBox::Slot()
+			.Padding(2.f, WeightEditVerticalPadding)
+			[
+				SNew(SButton)
+				.HAlign(HAlign_Center)
+				.VAlign(VAlign_Center)
+				.Text(LOCTEXT("BorderSelectionButtonLabel", "Convert to Border"))
+				.ToolTipText(LOCTEXT("BorderSelectionTooltip", "Select vertices on the border of the current selection."))
+				.OnClicked_Lambda([this]()
+				{
+					Tool->SelectBorder();
+					return FReply::Handled();
+				})
+			]
+		]
 	];
 
 	// custom display of weight editing tools
 	IDetailCategoryBuilder& EditWeightsCategory = DetailBuilder.EditCategory("Edit Weights", FText::GetEmpty(), ECategoryPriority::Important);
 	EditWeightsCategory.InitiallyCollapsed(true);
 
+	EditWeightsCategory.AddCustomRow(LOCTEXT("SelectMessageRow", "Select Vertices"), false)
+	.WholeRowContent()
+	.VAlign(VAlign_Center)
+	[
+		SNew(STextBlock)
+		.Visibility_Lambda([this]()
+		{
+			return Tool->GetSelectedVertices().IsEmpty() ? EVisibility::Visible : EVisibility::Collapsed;
+		})
+		.Text_Lambda([this]()
+		{
+			return LOCTEXT("NothingSelectedLabel", "Select vertices to edit weights...");
+		})
+	];
+
 	// FLOOD WEIGHTS SLIDER category
 	ToolSettings->DirectEditState.Reset();
 	EditWeightsCategory.AddCustomRow(LOCTEXT("FloodWeightsRow", "Flood Weights Slider"), false)
@@ -515,9 +564,10 @@
 		+SHorizontalBox::Slot()
 		[
 			SNew(SSegmentedControl<EWeightEditOperation>)
-			.ToolTipText(LOCTEXT("EditModeTooltip",
+			.ToolTipText(LOCTEXT("InteractiveEditModeTooltip",
 				"Add: applies the current weight plus the flood value to the new weight.\n"
-				"Multiply: applies the current weight multiplied by the flood value to the new weight.\n"))
+				"Multiply: applies the current weight multiplied by the flood value to the new weight.\n"
+				"This operation applies interactively while dragging the slider. It operates on the currently selected vertices."))
 			.Value_Lambda([this]()
 			{
 				return ToolSettings->DirectEditState.EditMode;
@@ -549,11 +599,11 @@
 				})
 				.MinSliderValue_Lambda([this]()
 				{
-					return ToolSettings->DirectEditState.GetModeMinValue() + SliderBuffer;
+					return ToolSettings->DirectEditState.GetModeMinValue();
 				})
 				.MaxSliderValue_Lambda([this]()
 				{
-					return ToolSettings->DirectEditState.GetModeMaxValue() - SliderBuffer;
+					return ToolSettings->DirectEditState.GetModeMaxValue();
 				})
 				.MinValue_Lambda([this]()
 				{
@@ -585,8 +635,9 @@
 						Tool->EditWeightsOnVertices(
 							Tool->GetCurrentBoneIndex(),
 							Value,
+							0, // iterations
 							ToolSettings->DirectEditState.EditMode,
-							Tool->GetVerticesToEdit(),
+							Tool->GetSelectedVertices(),
 							bShouldTransact);
 					}
 				})
@@ -597,9 +648,10 @@
 						constexpr bool bShouldTransact = true;
 						Tool->EditWeightsOnVertices(
 							Tool->GetCurrentBoneIndex(),
-							NewValue, 
+							NewValue,
+							0, // iterations
 							ToolSettings->DirectEditState.EditMode,
-							Tool->GetVerticesToEdit(),
+							Tool->GetSelectedVertices(),
 							bShouldTransact);
 					}
 					ToolSettings->DirectEditState.bInTransaction = false;
@@ -616,20 +668,29 @@
 					Tool->EndChange(TransactionLabel);
 					ToolSettings->DirectEditState.bInTransaction = false;
 
-					// reset slider after each edit because multiplying operation is always relative to 1.0
+					// reset multiply slider
 					if (ToolSettings->DirectEditState.EditMode == EWeightEditOperation::Multiply)
 					{
+						// multiplying operation is always relative to 1.0
 						ToolSettings->DirectEditState.CurrentValue = 1.0f;
 						ToolSettings->DirectEditState.StartValue = 1.0f;
 					}
+
+					// reset add slider
+					if (ToolSettings->DirectEditState.EditMode == EWeightEditOperation::Add)
+					{
+						// add operation is always relative to 0.0
+						ToolSettings->DirectEditState.CurrentValue = 0.0f;
+						ToolSettings->DirectEditState.StartValue = 0.0f;
+					}
 				})
 				.ToolTipText(LOCTEXT("FloodWeightsToolTip", "Drag the slider to interactively adjust weights on the selected vertices."))
 			]	
 		]
 	];
 
-	// REPLACE WEIGHTS category
-	EditWeightsCategory.AddCustomRow(LOCTEXT("ReplaceWeightsRow", "Replace"), false)
+	// ADD WEIGHTS category
+	EditWeightsCategory.AddCustomRow(LOCTEXT("AddWeightsRow", "Add"), false)
 	.WholeRowContent()
 	[
 		SNew(SVerticalBox)
@@ -638,35 +699,63 @@
 		.Padding(WeightEditHorizontalPadding, WeightEditVerticalPadding)
 		[
 			SNew(SHorizontalBox)
-					
 			+SHorizontalBox::Slot()
-			.VAlign(VAlign_Center)
-			.FillWidth(WeightEditingLabelsPercent)
 			[
-				SNew(STextBlock)
-				.Text(LOCTEXT("ReplaceLabel", "Replace"))
-				.Font(FAppStyle::GetFontStyle(TEXT("PropertyWindow.NormalFont")))
-				.ToolTipText(LOCTEXT("ReplaceTooltip", "Replace the current weight."))
+				SNew(SBox)
+				[
+					SNew(SButton)
+					.HAlign(HAlign_Center)
+					.Text(LOCTEXT("AddWeightsButtonLabel", "Add"))
+					.ToolTipText(LOCTEXT("AddButtonTooltip",
+					"Add: applies the current weight plus the flood value to the new weight.\n"
+					"This operation applies to the currently selected vertices."))
+					.OnClicked_Lambda([this]()
+					{
+						constexpr bool bShouldTransact = true;
+						Tool->EditWeightsOnVertices(
+							Tool->GetCurrentBoneIndex(),
+							ToolSettings->AddStrength,
+							0, // iterations 
+							EWeightEditOperation::Add,
+							Tool->GetSelectedVertices(),
+							bShouldTransact);
+						return FReply::Handled();
+					})
+				]
 			]
 
 			+SHorizontalBox::Slot()
 			[
 				SNew(SSpinBox<float>)
-				.MinValue(0.f)
-				.MaxValue(1.f)
+				.MinValue(0)
+				.MaxValue(1)
+				.ToolTipText(LOCTEXT("AddWeightsSliderToolTip", "Adjust the value to Add to the selected vertices."))
 				.Value_Lambda([this]()
 				{
-					return ToolSettings->ReplaceValue;
+					return ToolSettings->AddStrength;
 				})
 				.OnValueChanged_Lambda([this](float NewValue)
 				{
-					ToolSettings->ReplaceValue = NewValue;
+					ToolSettings->AddStrength = NewValue;
 				})
 				.OnValueCommitted_Lambda([this](float NewValue, ETextCommit::Type CommitType)
 				{
 					ToolSettings->SaveConfig();
 				})
 			]
+		]
+	];
+
+	// REPLACE WEIGHTS category
+	EditWeightsCategory.AddCustomRow(LOCTEXT("ReplaceWeightsRow", "Replace"), false)
+	.WholeRowContent()
+	[
+		SNew(SVerticalBox)
+
+		+ SVerticalBox::Slot()
+		.Padding(WeightEditHorizontalPadding, WeightEditVerticalPadding)
+		[
+			SNew(SHorizontalBox)
 
 			+SHorizontalBox::Slot()
 			[
@@ -676,25 +765,47 @@
 					.HAlign(HAlign_Center)
 					.Text(LOCTEXT("ReplaceWeightsButtonLabel", "Replace"))
 					.ToolTipText(LOCTEXT("ReplaceButtonTooltip",
-						"Replace: the weight of selected vertices is replaced by the specified value.\n"))
+						"Replace: the weight of selected vertices is replaced by the specified value.\n"
+						"This operation applies to the currently selected vertices."))
 					.OnClicked_Lambda([this]()
 					{
 						constexpr bool bShouldTransact = true;
 						Tool->EditWeightsOnVertices(
 							Tool->GetCurrentBoneIndex(),
-							ToolSettings->ReplaceValue, 
+							ToolSettings->ReplaceValue,
+							0, // iterations
 							EWeightEditOperation::Replace,
-							Tool->GetVerticesToEdit(),
+							Tool->GetSelectedVertices(),
 							bShouldTransact);
 						return FReply::Handled();
 					})
 				]
 			]
+
+			+SHorizontalBox::Slot()
+			[
+				SNew(SSpinBox<float>)
+				.MinValue(0.f)
+				.MaxValue(1.f)
+				.ToolTipText(LOCTEXT("ReplaceWeightsSliderToolTip", "Adjust the value to Replace on the selected vertices."))
+				.Value_Lambda([this]()
+				{
+					return ToolSettings->ReplaceValue;
+				})
+				.OnValueChanged_Lambda([this](float NewValue)
+				{
+					ToolSettings->ReplaceValue = NewValue;
+				})
+				.OnValueCommitted_Lambda([this](float NewValue, ETextCommit::Type CommitType)
+				{
+					ToolSettings->SaveConfig();
+				})
+			]
 		]
 	];
 
-	// RELAX WEIGHTS category
-	EditWeightsCategory.AddCustomRow(LOCTEXT("RelaxWeightsRow", "Relax"), false)
+	// AVERAGE WEIGHTS category
+	EditWeightsCategory.AddCustomRow(LOCTEXT("AverageWeightsRow", "Average"), false)
 	.WholeRowContent()
 	[
 		SNew(SVerticalBox)
@@ -703,36 +814,56 @@
 		.Padding(WeightEditHorizontalPadding, WeightEditVerticalPadding)
 		[
 			SNew(SHorizontalBox)
-					
+
 			+SHorizontalBox::Slot()
-			.VAlign(VAlign_Center)
-			.FillWidth(WeightEditingLabelsPercent)
 			[
-				SNew(STextBlock)
-				.Text(LOCTEXT("RelaxIterationsLabel", "Iterations"))
-				.Font(FAppStyle::GetFontStyle(TEXT("PropertyWindow.NormalFont")))
-				.ToolTipText(LOCTEXT("RelaxIterationsTooltip", "The number of iterations of relaxation to apply."))
+				SNew(SBox)
+				[
+					SNew(SButton)
+					.HAlign(HAlign_Center)
+					.Text(LOCTEXT("AverageWeightsButtonLabel", "Average"))
+					.ToolTipText(LOCTEXT("AverageButtonTooltip",
+						"Weights on all selected vertices are set to the average of all selected vertices."))
+					.OnClicked_Lambda([this]()
+					{
+						ToolSettings->WeightTool->AverageWeights(ToolSettings->AverageStrength);
+						return FReply::Handled();
+					})
+				]
 			]
 
 			+SHorizontalBox::Slot()
 			[
-				SNew(SSpinBox<int32>)
-				.MinValue(1)
-				.MaxValue(10)
+				SNew(SSpinBox<float>)
+				.MinValue(0.f)
+				.MaxValue(1.f)
+				.ToolTipText(LOCTEXT("AverageWeightsSliderToolTip", "Blend the amount to Average the weights on the selected vertices."))
 				.Value_Lambda([this]()
 				{
-					return ToolSettings->RelaxIterations;
+					return ToolSettings->AverageStrength;
 				})
-				.OnValueChanged_Lambda([this](int32 NewValue)
+				.OnValueChanged_Lambda([this](float NewValue)
 				{
-					ToolSettings->RelaxIterations = NewValue;
+					ToolSettings->AverageStrength = NewValue;
 				})
-				.OnValueCommitted_Lambda([this](int32 NewValue, ETextCommit::Type CommitType)
+				.OnValueCommitted_Lambda([this](float NewValue, ETextCommit::Type CommitType)
 				{
 					ToolSettings->SaveConfig();
 				})
 			]
+		]
+	];
 
+	// RELAX WEIGHTS category
+	EditWeightsCategory.AddCustomRow(LOCTEXT("RelaxWeightsRow", "Relax"), false)
+	.WholeRowContent()
+	[
+		SNew(SVerticalBox)
+
+		+ SVerticalBox::Slot()
+		.Padding(WeightEditHorizontalPadding, WeightEditVerticalPadding)
+		[
+			SNew(SHorizontalBox)
 			+SHorizontalBox::Slot()
 			[
 				SNew(SBox)
@@ -741,133 +872,53 @@
 					.HAlign(HAlign_Center)
 					.Text(LOCTEXT("RelaxWeightsButtonLabel", "Relax"))
 					.ToolTipText(LOCTEXT("RelaxButtonTooltip",
-						"Relax: the weight of each vertex is replaced by the average of it's neighbors. This smooths weights across the mesh.\n"))
+						"Relax: the weight of each selected vertex is replaced by the average of it's neighbors. \n"
+						"This smooths weights across the mesh."))
 					.OnClicked_Lambda([this]()
 					{
 						constexpr bool bShouldTransact = true;
+						constexpr int32 DefaultRelaxIterations = 5; // provides a reasonable falloff distance 
 						Tool->EditWeightsOnVertices(
 							Tool->GetCurrentBoneIndex(),
-							ToolSettings->RelaxIterations, 
+							ToolSettings->RelaxStrength,
+							DefaultRelaxIterations, 
 							EWeightEditOperation::Relax,
-							Tool->GetVerticesToEdit(),
+							Tool->GetSelectedVertices(),
 							bShouldTransact);
 						return FReply::Handled();
 					})
 				]
 			]
-		]
-	];
-	
-	// PRUNE WEIGHTS category
-	EditWeightsCategory.AddCustomRow(LOCTEXT("PruneWeightsRow", "Prune"), false)
-	.WholeRowContent()
-	[
-		SNew(SVerticalBox)
-
-		+ SVerticalBox::Slot()
-		.Padding(WeightEditHorizontalPadding, WeightEditVerticalPadding)
-		[
-			SNew(SHorizontalBox)
-					
-			+SHorizontalBox::Slot()
-			.VAlign(VAlign_Center)
-			.FillWidth(WeightEditingLabelsPercent)
-			[
-				SNew(STextBlock)
-				.Text(LOCTEXT("PruneThresholdLabel", "Prune"))
-				.Font(FAppStyle::GetFontStyle(TEXT("PropertyWindow.NormalFont")))
-				.ToolTipText(LOCTEXT("PruneThresholdTooltip", "The threshold weight value to use with Prune operation."))
-			]
 
 			+SHorizontalBox::Slot()
 			[
 				SNew(SSpinBox<float>)
-				.MinValue(0.f)
-				.MaxValue(1.f)
+				.MinValue(0)
+				.MaxValue(1)
+				.ToolTipText(LOCTEXT("RelaxWeightsSliderToolTip", "Blend the amount to Relax the weights on the selected vertices."))
 				.Value_Lambda([this]()
 				{
-					return ToolSettings->PruneValue;
+					return ToolSettings->RelaxStrength;
 				})
 				.OnValueChanged_Lambda([this](float NewValue)
 				{
-					ToolSettings->PruneValue = NewValue;
+					ToolSettings->RelaxStrength = NewValue;
 				})
 				.OnValueCommitted_Lambda([this](float NewValue, ETextCommit::Type CommitType)
 				{
 					ToolSettings->SaveConfig();
 				})
 			]
-
-			+SHorizontalBox::Slot()
-			[
-				SNew(SBox)
-				[
-					SNew(SButton)
-					.HAlign(HAlign_Center)
-					.Text(LOCTEXT("PruneWeightsButtonLabel", "Prune"))
-					.ToolTipText(LOCTEXT("PruneButtonTooltip",
-						"Weights below the given threshold value are removed.\n"
-						"This command operates on the selected bone(s) and selected vertices.\n "
-						"If no bones are selected, ALL bone weights are considered.\n "
-						"If no vertices are selected, ALL vertices are considered."))
-					.OnClicked_Lambda([this]()
-					{
-						ToolSettings->WeightTool->PruneWeights(ToolSettings->PruneValue);
-						return FReply::Handled();
-					})
-				]
-			]
 		]
 	];
 
-	// AVERAGE/NORMALIZE WEIGHTS category
-	EditWeightsCategory.AddCustomRow(LOCTEXT("NormalizeWeightsRow", "Normalize"), false)
+	// VERTEX EDITOR category
+	EditWeightsCategory.AddCustomRow(LOCTEXT("VertexEditorRow", "Component Editor"), false)
 	.WholeRowContent()
 	[
-		SNew(SHorizontalBox)
-
-		+SHorizontalBox::Slot()
-		.Padding(2.f, WeightEditVerticalPadding)
-		[
-			SNew(SButton)
-			.HAlign(HAlign_Center)
-			.VAlign(VAlign_Center)
-			.Text(LOCTEXT("AverageWeightsButtonLabel", "Average"))
-			.ToolTipText(LOCTEXT("AverageWeightsTooltip",
-					"Takes the average of vertex weights and applies it to the selected vertices.\n"
-					"This command operates on the selected bone(s) and selected vertices.\n "
-					"If no bones are selected, ALL bone weights are considered.\n "
-					"If no vertices are selected, ALL vertices are considered."))
-			.OnClicked_Lambda([this]()
-			{
-				Tool->AverageWeights();
-				return FReply::Handled();
-			})
-		]
-
-		+SHorizontalBox::Slot()
-		.Padding(2.f, WeightEditVerticalPadding)
-		[
-			SNew(SButton)
-			.HAlign(HAlign_Center)
-			.VAlign(VAlign_Center)
-			.Text(LOCTEXT("NormalizeWeightsButtonLabel", "Normalize"))
-			.ToolTipText(LOCTEXT("NormalizeWeightsTooltip",
-					"Forces the weights on the selected vertices to sum to 1.\n"
-					"This command operates on the selected vertices.\n "
-					"If no vertices are selected, ALL vertices are considered."))
-			.IsEnabled_Lambda([this]()
-			{
-				return ToolSettings->EditingMode == EWeightEditMode::Mesh;
-			})
-			.OnClicked_Lambda([this]()
-			{
-				Tool->NormalizeWeights();
-				return FReply::Handled();
-			})
-		]
+		SNew(SVertexWeightEditor, ToolSettings->WeightTool)
 	];
-
+	
 	// MIRROR WEIGHTS category
 	EditWeightsCategory.AddCustomRow(LOCTEXT("MirrorWeightsRow", "Mirror"), false)
 	.WholeRowContent()
@@ -946,9 +997,7 @@
 				.Text(LOCTEXT("MirrorWeightsButtonLabel", "Mirror"))
 				.ToolTipText(LOCTEXT("MirrorButtonTooltip",
 					"Weights are copied across the given plane in the given direction.\n"
-					"This command operates on the selected bone(s) and selected vertices.\n "
-					"If no bones are selected, ALL bone weights are considered.\n "
-					"If no vertices are selected, ALL vertices are considered."))
+					"This command operates on the selected vertices."))
 				.OnClicked_Lambda([this]()
 				{
 					Tool->MirrorWeights(ToolSettings->MirrorAxis, ToolSettings->MirrorDirection);
@@ -958,11 +1007,84 @@
 		]
 	];
 
-	// VERTEX EDITOR category
-	EditWeightsCategory.AddCustomRow(LOCTEXT("VertexEditorRow", "Component Editor"), false)
+	// NORMALIZE WEIGHTS category
+	EditWeightsCategory.AddCustomRow(LOCTEXT("NormalizeWeightsRow", "Normalize"), false)
 	.WholeRowContent()
 	[
-		SNew(SVertexWeightEditor, ToolSettings->WeightTool)
+		SNew(SHorizontalBox)
+
+		+SHorizontalBox::Slot()
+		.Padding(2.f, WeightEditVerticalPadding)
+		[
+			SNew(SButton)
+			.HAlign(HAlign_Center)
+			.VAlign(VAlign_Center)
+			.Text(LOCTEXT("NormalizeWeightsButtonLabel", "Normalize"))
+			.ToolTipText(LOCTEXT("NormalizeWeightsTooltip",
+					"Forces the weights on the selected vertices to sum to 1.\n"
+					"This command operates on the selected vertices."))
+			.IsEnabled_Lambda([this]()
+			{
+				return ToolSettings->EditingMode == EWeightEditMode::Mesh;
+			})
+			.OnClicked_Lambda([this]()
+			{
+				Tool->NormalizeWeights();
+				return FReply::Handled();
+			})
+		]
 	];
+
+	// PRUNE WEIGHTS category
+	EditWeightsCategory.AddCustomRow(LOCTEXT("PruneWeightsRow", "Prune"), false)
+	.WholeRowContent()
+	[
+		SNew(SVerticalBox)
+
+		+ SVerticalBox::Slot()
+		.Padding(WeightEditHorizontalPadding, WeightEditVerticalPadding)
+		[
+			SNew(SHorizontalBox)
+
+			+SHorizontalBox::Slot()
+			[
+				SNew(SBox)
+				[
+					SNew(SButton)
+					.HAlign(HAlign_Center)
+					.Text(LOCTEXT("PruneWeightsButtonLabel", "Prune"))
+					.ToolTipText(LOCTEXT("PruneButtonTooltip",
+						"Removes influences with weights below the given threshold value.\n"
+						"Pruned bones are removed from the list of bones affecting the given vertex.\n"
+						"Pruned bones will no longer recieve weight when a vertex is normalized.\n"
+						"This command operates on the selected vertices."))
+					.OnClicked_Lambda([this]()
+					{
+						ToolSettings->WeightTool->PruneWeights(ToolSettings->PruneValue, TArray<BoneIndex>());
+						return FReply::Handled();
+					})
+				]
+			]
+
+			+SHorizontalBox::Slot()
+			[
+				SNew(SSpinBox<float>)
+				.MinValue(0.f)
+				.MaxValue(1.f)
+				.Value_Lambda([this]()
+				{
+					return ToolSettings->PruneValue;
+				})
+				.OnValueChanged_Lambda([this](float NewValue)
+				{
+					ToolSettings->PruneValue = NewValue;
+				})
+				.OnValueCommitted_Lambda([this](float NewValue, ETextCommit::Type CommitType)
+				{
+					ToolSettings->SaveConfig();
+				})
+			]
+		]
+	];
 }
 
@@ -1015,19 +1137,19 @@
 
 TSharedRef<SWidget> SVertexWeightItem::GenerateWidgetForColumn(const FName& ColumnName)
 {
-	if (ColumnName == "Bone")
+	if (ColumnName == ColumnName_Bone)
 	{
 		const FName BoneName = ParentTable->Tool->GetBoneNameFromIndex(Element->BoneIndex);
 		return SNew(STextBlock).Text(FText::FromName(BoneName));
 	}
 
-	if (ColumnName == "Weight")
+	if (ColumnName == ColumnName_Weight)
 	{
 		return SNew(SNumericEntryBox<float>)
 		.AllowSpin(true)
-		.MinSliderValue(SliderBuffer)
+		.MinSliderValue(0.f)
 		.MinValue(0.f)
-		.MaxSliderValue(1.0f-SliderBuffer)
+		.MaxSliderValue(1.0f)
 		.MaxValue(1.f)
 		.Value_Lambda([this]()
 		{
@@ -1036,17 +1158,26 @@
 				return ValueDuringSlide;
 			}
 			
-			return ParentTable->Tool->GetAverageWeightOnBone(Element->BoneIndex, ParentTable->Tool->GetVerticesToEdit());
+			return ParentTable->Tool->GetAverageWeightOnBone(Element->BoneIndex, ParentTable->Tool->GetSelectedVertices());
 		})
 		.OnValueChanged_Lambda([this](float NewValue)
 		{
 			if (bInTransaction)
 			{
 				ValueDuringSlide = NewValue;
-				
-				const float ScaleMultiplier = NewValue / FMath::Max(ValueAtStartOfSlide, SliderBuffer);
+
+				const bool bScalingUp = NewValue >= ValueAtStartOfSlide || FMath::IsNearlyEqual(ValueAtStartOfSlide, 0.f);
+				const float RangeEnd = bScalingUp ? 1.f : 0.f;
+				float RelativeScale = (NewValue - ValueAtStartOfSlide) * 1.f / (RangeEnd - ValueAtStartOfSlide);
+				RelativeScale *= bScalingUp ? 1.f : -1.f;
 				constexpr bool bShouldTransact = false;
-				ParentTable->Tool->EditWeightsOnVertices(Element->BoneIndex, ScaleMultiplier, EWeightEditOperation::Multiply, ParentTable->Tool->GetVerticesToEdit(), bShouldTransact);
+				ParentTable->Tool->EditWeightsOnVertices(
+					Element->BoneIndex,
+					RelativeScale,
+					0, // iterations
+					EWeightEditOperation::RelativeScale,
+					ParentTable->Tool->GetSelectedVertices(),
+					bShouldTransact);
 			}
 		})
 		.OnValueCommitted_Lambda([this](float NewValue, ETextCommit::Type CommitType)
@@ -1054,14 +1185,20 @@
 			if (!bInTransaction)
 			{
 				constexpr bool bShouldTransact = true;
-				ParentTable->Tool->EditWeightsOnVertices(Element->BoneIndex, NewValue, EWeightEditOperation::Replace, ParentTable->Tool->GetVerticesToEdit(), bShouldTransact);
+				ParentTable->Tool->EditWeightsOnVertices(
+					Element->BoneIndex,
+					NewValue,
+					0, // iterations,
+					EWeightEditOperation::Replace,
+					ParentTable->Tool->GetSelectedVertices(),
+					bShouldTransact);
 			}
 			bInTransaction = false;
 		})
 		.OnBeginSliderMovement_Lambda([this]()
 		{
 			ParentTable->Tool->BeginChange();
-			ValueAtStartOfSlide = ParentTable->Tool->GetAverageWeightOnBone(Element->BoneIndex, ParentTable->Tool->GetVerticesToEdit());
+			ValueAtStartOfSlide = ParentTable->Tool->GetAverageWeightOnBone(Element->BoneIndex, ParentTable->Tool->GetSelectedVertices());
 			ValueDuringSlide = ValueAtStartOfSlide;
 			bInTransaction = true;
 		})
@@ -1074,6 +1211,35 @@
 		.ToolTipText(LOCTEXT("WeightSliderToolTip", "Set the weight on this bone for the selected vertices."));
 	}
 
+	if (ColumnName == ColumnName_Prune)
+	{
+		return SNew(SHorizontalBox)
+		
+		+SHorizontalBox::Slot()
+		.FillWidth(1.f)
+		.HAlign(HAlign_Center)
+		.VAlign(VAlign_Center)
+		[
+			SNew(SButton)
+			.HAlign(HAlign_Right)
+			.ToolTipText(LOCTEXT("PruneInfluence", "Prune the influence from the selected vertices."))
+			.OnClicked_Lambda([this]() -> FReply
+			{
+				// use a negative threshold weight that no weight will ever
+				// be below because we only want to prune based on the bone, regardless of the weight
+				constexpr float NegativeThreshold = -1.0f;
+				ParentTable->Tool->PruneWeights(NegativeThreshold, {Element->BoneIndex});
+				return FReply::Handled();
+			})
+			.Content()
+			[
+				SNew(SImage)
+				.Image(FAppStyle::Get().GetBrush("Icons.Delete"))
+				.ColorAndOpacity(FSlateColor::UseForeground())
+			]
+		];
+	}
+
 	checkNoEntry();
 	return SNullWidget::NullWidget;
 }
@@ -1097,6 +1263,10 @@
 		SNew(SBox)
 		[
 			SAssignNew( ListView, SWeightEditorListViewType )
+			.Visibility_Lambda([this]()
+			{
+				return Tool->GetSelectedVertices().IsEmpty() ? EVisibility::Collapsed : EVisibility::Visible;
+			})
 			.SelectionMode(ESelectionMode::Single)
 			.ListItemsSource( &ListViewItems )
 			.OnGenerateRow_Lambda([this](TSharedPtr<FWeightEditorElement> Element, const TSharedRef<STableViewBase>& OwnerTableView)
@@ -1106,8 +1276,16 @@
 			.HeaderRow
 			(
 				SNew(SHeaderRow)
-				+ SHeaderRow::Column("Bone").DefaultLabel(NSLOCTEXT("WeightEditorBoneColumn", "Bone", "Bone"))
-				+ SHeaderRow::Column("Weight").DefaultLabel(NSLOCTEXT("WeightEditorWeightColumn", "Weight (Average)", "Weight (Average)"))
+				+ SHeaderRow::Column(ColumnName_Bone)
+				.HAlignHeader(HAlign_Center)
+				.DefaultLabel(LOCTEXT("WeightEditorBoneColumn", "Bone"))
+				+ SHeaderRow::Column(ColumnName_Weight)
+				.HAlignHeader(HAlign_Center)
+				.DefaultLabel(LOCTEXT("WeightEditorWeightColumn", "Weight (Average)"))
+				+ SHeaderRow::Column(ColumnName_Prune)
+				.HAlignHeader(HAlign_Center)
+				.DefaultLabel(LOCTEXT("WeightEditorPruneColumn", "Prune"))
+				.FixedWidth(60.f)
 			)
 		]
 	];
@@ -1127,7 +1305,7 @@
 	
 	// get all bones affecting the selected vertices
 	TArray<int32> Influences;
-	Tool->GetInfluences(Tool->GetVerticesToEdit(), Influences);
+	Tool->GetInfluences(Tool->GetSelectedVertices(), Influences);
 
 	// generate list view items
 	ListViewItems.Reset();
diff --git a/Engine/Plugins/Experimental/MeshModelingToolsetExp/Source/MeshModelingToolsEditorOnlyExp/Private/SkeletalMesh/SkinWeightsPaintTool.cpp b/Engine/Plugins/Experimental/MeshModelingToolsetExp/Source/MeshModelingToolsEditorOnlyExp/Private/SkeletalMesh/SkinWeightsPaintTool.cpp
--- a/Engine/Plugins/Experimental/MeshModelingToolsetExp/Source/MeshModelingToolsEditorOnlyExp/Private/SkeletalMesh/SkinWeightsPaintTool.cpp
+++ b/Engine/Plugins/Experimental/MeshModelingToolsetExp/Source/MeshModelingToolsEditorOnlyExp/Private/SkeletalMesh/SkinWeightsPaintTool.cpp
@@ -177,7 +177,7 @@ public:
 	virtual int32 GetBoneIndex(const int32 VertexID, const int32 Index) override
 	{
 		const int32 SrcVertexId = GetSourceVertexId(VertexID);
-		return Weights->PreChangeWeights[SrcVertexId][Index].BoneIndex;
+		return Weights->PreChangeWeights[SrcVertexId][Index].BoneID;
 	}
 
 	virtual float GetBoneWeight(const int32 VertexID, const int32 Index) override
@@ -415,6 +415,13 @@ void FMultiBoneWeightEdits::GetEditedVertexIndices(TSet<int32>& OutVerticesToEdi
 	}
 }
 
+void FMultiBoneWeightEdits::AddPruneBoneEdit(
+	const VertexIndex VertexToPruneFrom,
+	const BoneIndex BoneToPrune)
+{
+	PrunedInfluences.Emplace(VertexToPruneFrom, BoneToPrune);
+}
+
 void FSkinToolDeformer::Initialize(const USkeletalMeshComponent* InSkelMeshComponent, const FMeshDescription* InMeshDescription)
 {
 	// get all bone transforms in the reference pose store a copy in component space
@@ -525,11 +532,11 @@ void FSkinToolDeformer::UpdateVertexDeformation(
 			const VertexWeights& VertexPerBoneData = CurrentWeights[VertexID];
 			for (const FVertexBoneWeight& VertexData : VertexPerBoneData)
 			{
-				if (VertexData.BoneIndex == INDEX_NONE)
+				if (VertexData.BoneID == INDEX_NONE)
 				{
 					continue;
 				}
-				const FTransform& CurrentTransform = PoseComponentSpace[VertexData.BoneIndex];
+				const FTransform& CurrentTransform = PoseComponentSpace[VertexData.BoneID];
 				VertexNewPosition += CurrentTransform.TransformPosition(VertexData.VertexInBoneSpace) * VertexData.Weight;
 			}
 			
@@ -671,7 +678,7 @@ void FSkinToolWeights::InitializeSkinWeights(
 		{
 			if (VertexBoneData.Weight > UE::AnimationCore::BoneWeightThreshold)
 			{
-				IsBoneWeighted[VertexBoneData.BoneIndex] = true;
+				IsBoneWeighted[VertexBoneData.BoneID] = true;
 			}
 		}
 	}
@@ -685,97 +692,121 @@ void FSkinToolWeights::EditVertexWeightAndNormalize(
 {
 	// clamp new weight
 	NewWeightValue = FMath::Clamp(NewWeightValue, 0.0f, 1.0f);
-	
+
 	// calculate the sum of all the weights on this vertex (not including the one we currently applied)
-	TArray<int32> BonesAffectingVertex;
+	TArray<int32> RecordedBonesOnVertex;
 	TArray<float> ValuesToNormalize;
 	float Total = 0.0f;
 	const VertexWeights& VertexData = PreChangeWeights[VertexID];
 	for (const FVertexBoneWeight& VertexBoneData : VertexData)
 	{
-		if (VertexBoneData.BoneIndex == BoneToHoldIndex)
-		{
-			continue;
-		}
-		
-		if (VertexBoneData.Weight < MinimumWeightThreshold)
+		if (VertexBoneData.BoneID == BoneToHoldIndex)
 		{
 			continue;
 		}
 		
-		BonesAffectingVertex.Add(VertexBoneData.BoneIndex);
+		RecordedBonesOnVertex.Add(VertexBoneData.BoneID);
 		ValuesToNormalize.Add(VertexBoneData.Weight);
 		Total += VertexBoneData.Weight;
 	}
 
-	// are there no OTHER influences on this vertex?
-	const bool bVertexHasNoOtherInfluences = Total <= MinimumWeightThreshold;
-	
-	// if user applied any weight to this vertex AND there's no other weights of any significance,
-	// then simply apply full weight to this vertex, set all other influences to zero and return
-	const bool bApplyWeightToThisVertex = NewWeightValue >= MinimumWeightThreshold;
-	if (bApplyWeightToThisVertex && bVertexHasNoOtherInfluences)
+	// assigning full weight to this vertex?
+	if (FMath::IsNearlyEqual(NewWeightValue, 1.f))
 	{
-		// set all other influences to 0.0f
+		// in this case normalization is trivial, just assign the full weight directly and zero all others
+		const float PrevWeight = GetWeightOfBoneOnVertex(BoneToHoldIndex, VertexID, PreChangeWeights);
+		constexpr float FullWeight = 1.0f;
+		WeightEdits.MergeSingleEdit(
+			BoneToHoldIndex,
+			VertexID,
+			PrevWeight,
+			FullWeight);
+
+		// zero all others
 		for (int32 i=0; i<ValuesToNormalize.Num(); ++i)
 		{
-			const int32 BoneIndex = BonesAffectingVertex[i];
+			const int32 BoneIndex = RecordedBonesOnVertex[i];
 			const float OldWeight = ValuesToNormalize[i];
-			constexpr float NewWeight = 0.0f;
+			constexpr float NewWeight = 0.f;
 			WeightEdits.MergeSingleEdit(BoneIndex, VertexID, OldWeight, NewWeight);
 		}
 
-		// set current bone value to 1.0f
-		const float PrevWeight = GetWeightOfBoneOnVertex(BoneToHoldIndex, VertexID, PreChangeWeights);
-		WeightEdits.MergeSingleEdit(
-			BoneToHoldIndex,
-			VertexID,
-			PrevWeight,
-			1.0f);
-		
 		return;
 	}
 
-	// is the user trying to prune ALL weight from this vertex, AND all other weights are equal to zero?
-	// in this case, we have two options:
-	// 1. if there are other influences recorded for this vertex, then split remaining influence among them
-	// 2. if there are NO other influences recorded for this vertex, then move the weight to the parent as last ditch effort
-	const bool bApplyZeroWeightToThisVertex = NewWeightValue <= MinimumWeightThreshold;
-	if (bApplyZeroWeightToThisVertex && bVertexHasNoOtherInfluences)
+	// do any other influences have any weight on this vertex?
+	//
+	// In the case that:
+	// 1. user applied any weight < 1 to this vertex AND
+	// 2. there are NO other weights on this vertex
+	// then we need to decide where to put the remaining influence...
+	//
+	// the logic here attempts to find a reasonable and "least surprising" place to put the remaining weight based on artist feedback
+	const bool bVertexHasNoOtherWeightedInfluences = Total <= MinimumWeightThreshold;
+	if (bVertexHasNoOtherWeightedInfluences)
 	{
-		if (ValuesToNormalize.IsEmpty())
+		// does this vertex have any other recorded influences on it?
+		// a "recorded" influence here is one that used to have weight, but no longer does
+		if (!RecordedBonesOnVertex.IsEmpty())
 		{
-			// assign all weight to the parent
-			int32 BoneIndex = 0;
-			if (const USkeletalMesh* SkeletalMesh = Deformer.Component->GetSkeletalMeshAsset())
+			// this vertex:
+			// 1. was previously weighted to other influences
+			// 2. has subsequently had all other weight removed
+			// In this case, we evenly split the remaining weight among the recorded influences
+
+			// distribute remaining weight evenly over other recorded influences
+			const float WeightToDistribute = (1.0f - NewWeightValue) / RecordedBonesOnVertex.Num();
+			for (int32 i=0; i<ValuesToNormalize.Num(); ++i)
 			{
-				BoneIndex = SkeletalMesh->GetRefSkeleton().GetParentIndex(BoneToHoldIndex);
-				BoneIndex = BoneIndex == INDEX_NONE ? 0 : BoneIndex;
+				const int32 BoneIndex = RecordedBonesOnVertex[i];
+				const float OldWeight = ValuesToNormalize[i];
+				const float NewWeight = WeightToDistribute;
+				WeightEdits.MergeSingleEdit(BoneIndex, VertexID, OldWeight, NewWeight);
 			}
-			
-			constexpr float OldWeight = 0.f;
-			constexpr float NewWeight = 1.f;
-			WeightEdits.MergeSingleEdit(BoneIndex, VertexID, OldWeight, NewWeight);
+
+			// set current bone value to user assigned weight
+			const float PrevWeight = GetWeightOfBoneOnVertex(BoneToHoldIndex, VertexID, PreChangeWeights);
+			WeightEdits.MergeSingleEdit(
+				BoneToHoldIndex,
+				VertexID,
+				PrevWeight,
+				NewWeightValue);
 		}
 		else
 		{
-			// evenly distribute weight among other influences
-			const float NewWeight = 1.f / ValuesToNormalize.Num();
-			for (int32 i=0; i<ValuesToNormalize.Num(); ++i)
+			// this vertex:
+			// 1. has no other recorded influences
+			// 2. user is assigning PARTIAL weight to it (less than 1.0)
+			// so in this case we push the remaining weight onto the PARENT bone
+				
+			// assign remaining weight to the parent
+			const int32 ParentBoneIndex = GetParentBoneToWeightTo(BoneToHoldIndex);
+			if (ParentBoneIndex == BoneToHoldIndex)
 			{
-				const int32 BoneIndex = BonesAffectingVertex[i];
-				const float OldWeight = ValuesToNormalize[i];
-				WeightEdits.MergeSingleEdit(BoneIndex, VertexID, OldWeight, NewWeight);
+				// was unable to find parent OR child bone!
+				// this could only happen if user is trying to remove weight from the ONLY bone in the whole skeleton
+				// in this case just assign the full weight to the bone (there's no other valid configuration)
+				// this is a "do nothing" operation, but at least it generates an undo transaction to let user know the input was received
+				const float PrevWeight = GetWeightOfBoneOnVertex(BoneToHoldIndex, VertexID, PreChangeWeights);
+				constexpr float FullWeight = 1.0f;
+				WeightEdits.MergeSingleEdit(
+					BoneToHoldIndex,
+					VertexID,
+					PrevWeight,
+					FullWeight);
+			}
+			else
+			{
+				// assign remaining weight to parent
+				constexpr float OldParentWeight = 0.f;
+				const float NewParentWeight = 1.0f - NewWeightValue;
+				WeightEdits.MergeSingleEdit(ParentBoneIndex, VertexID, OldParentWeight, NewParentWeight);
+				// and assign user requested weight to the current bone
+				const float OldWeight = GetWeightOfBoneOnVertex(BoneToHoldIndex, VertexID, PreChangeWeights);
+				const float NewWeight = NewWeightValue;
+				WeightEdits.MergeSingleEdit(BoneToHoldIndex, VertexID, OldWeight, NewWeight);
 			}
 		}
-
-		// set current bone value to 0.0f
-		const float PrevWeight = GetWeightOfBoneOnVertex(BoneToHoldIndex, VertexID, PreChangeWeights);
-		WeightEdits.MergeSingleEdit(
-			BoneToHoldIndex,
-			VertexID,
-			PrevWeight,
-			0.f);
 		
 		return;
 	}
@@ -791,7 +822,7 @@ void FSkinToolWeights::EditVertexWeightAndNormalize(
 		{
 			NormalizedValue = (ValuesToNormalize[i] / Total) * AvailableTotal;	
 		}
-		const int32 BoneIndex = BonesAffectingVertex[i];
+		const int32 BoneIndex = RecordedBonesOnVertex[i];
 		const float OldWeight = ValuesToNormalize[i];
 		const float NewWeight = NormalizedValue;
 		WeightEdits.MergeSingleEdit(BoneIndex, VertexID, OldWeight, NewWeight);
@@ -831,7 +862,7 @@ void FSkinToolWeights::ApplyCurrentWeightsToMeshDescription(FMeshDescription* Me
 		const VertexWeights& VertexWeights = CurrentWeights[VertexIndex];
 		for (const FVertexBoneWeight& SingleBoneWeight : VertexWeights)
 		{
-			SourceBoneWeights.Add(UE::AnimationCore::FBoneWeight(SingleBoneWeight.BoneIndex, SingleBoneWeight.Weight));
+			SourceBoneWeights.Add(UE::AnimationCore::FBoneWeight(SingleBoneWeight.BoneID, SingleBoneWeight.Weight));
 		}
 
 		VertexSkinWeights.Set(FVertexID(VertexIndex), UE::AnimationCore::FBoneWeights::Create(SourceBoneWeights, Settings));
@@ -846,7 +877,7 @@ float FSkinToolWeights::GetWeightOfBoneOnVertex(
 	const VertexWeights& VertexWeights = InVertexWeights[VertexID];
 	for (const FVertexBoneWeight& BoneWeight : VertexWeights)
 	{
-		if (BoneWeight.BoneIndex == BoneIndex)
+		if (BoneWeight.BoneID == BoneIndex)
 		{
 			return BoneWeight.Weight;
 		}
@@ -867,7 +898,7 @@ void FSkinToolWeights::SetWeightOfBoneOnVertex(
 	VertexWeights& VertexWeights = InOutVertexWeights[VertexID];
 	for (FVertexBoneWeight& BoneWeight : VertexWeights)
 	{
-		if (BoneWeight.BoneIndex == BoneIndex)
+		if (BoneWeight.BoneID == BoneIndex)
 		{
 			BoneWeight.Weight = Weight;
 			return;
@@ -876,12 +907,18 @@ void FSkinToolWeights::SetWeightOfBoneOnVertex(
 
 	// bone not already an influence on this vertex, so we need to add it..
 
+	// if the weight was pruned, it won't be recorded in the VertexWeights array,
+	// but we also don't want to add it back
+	if (FMath::IsNearlyEqual(Weight, 0.f))
+	{
+		return;
+	}
+
 	// if vertex has room for more influences, then simply add it
 	if (VertexWeights.Num() < UE::AnimationCore::MaxInlineBoneWeightCount)
 	{
 		// add a new influence to this vertex
-		const FVector PosLocalToBone = Deformer.InvCSRefPoseTransforms[BoneIndex].TransformPosition(Deformer.RefPoseVertexPositions[VertexID]);
-		VertexWeights.Emplace(BoneIndex, PosLocalToBone, Weight);
+		AddNewInfluenceToVertex(VertexID, BoneIndex, Weight,InOutVertexWeights);
 		return;
 	}
 
@@ -905,7 +942,7 @@ void FSkinToolWeights::SetWeightOfBoneOnVertex(
 	// replace smallest influence
 	FVertexBoneWeight& BoneWeightToReplace = VertexWeights[SmallestInfluenceIndex];
 	BoneWeightToReplace.Weight = Weight;
-	BoneWeightToReplace.BoneIndex = BoneIndex;
+	BoneWeightToReplace.BoneID = BoneIndex;
 	BoneWeightToReplace.VertexInBoneSpace = Deformer.InvCSRefPoseTransforms[BoneIndex].TransformPosition(Deformer.RefPoseVertexPositions[VertexID]);
 
 	// now we need to re-normalize because the stamp does not handle maximum influences
@@ -920,6 +957,75 @@ void FSkinToolWeights::SetWeightOfBoneOnVertex(
 	}
 }
 
+void FSkinToolWeights::RemoveInfluenceFromVertex(
+	const VertexIndex InVertexID,
+	const BoneIndex InBoneID,
+	TArray<VertexWeights>& InOutVertexWeights)
+{
+	// should never be pruning a vertex that doesn't exist
+	if (!ensure(InOutVertexWeights.IsValidIndex(InVertexID)))
+	{
+		return;
+	}
+	
+	VertexWeights& SingleVertexWeights = InOutVertexWeights[InVertexID];
+	const int32 IndexOfBoneInVertex = SingleVertexWeights.IndexOfByPredicate([&InBoneID](const FVertexBoneWeight& CurrentVertexWeight)
+	{
+		return CurrentVertexWeight.BoneID == InBoneID;
+	});
+	// can't prune an influence that doesn't exist on a vertex
+	// this may happen if the calling code already pruned the influence to avoid normalization weights
+	if (IndexOfBoneInVertex == INDEX_NONE)
+	{
+		return;
+	}
+	
+	SingleVertexWeights.RemoveAt(IndexOfBoneInVertex);
+}
+
+void FSkinToolWeights::AddNewInfluenceToVertex(
+	const VertexIndex InVertexID,
+	const BoneIndex InBoneID,
+	const float Weight,
+	TArray<VertexWeights>& InOutVertexWeights)
+{		
+	// should never be adding an influence to a vertex that doesn't exist
+	if (!ensure(InOutVertexWeights.IsValidIndex(InVertexID)))
+	{
+		return;
+	}
+
+	// get list of weights on this single vertex
+	VertexWeights& SingleVertexWeights = InOutVertexWeights[InVertexID];
+
+	// should never be trying to add more influences beyond the max per-vertex limit
+	if (!ensure(SingleVertexWeights.Num() < UE::AnimationCore::MaxInlineBoneWeightCount))
+	{
+		return;
+	}
+
+	const int32 IndexOfBoneInVertex = SingleVertexWeights.IndexOfByPredicate([&InBoneID](const FVertexBoneWeight& CurrentVertexWeight)
+	{
+		return CurrentVertexWeight.BoneID == InBoneID;
+	});
+
+	// should never be adding an influence that already exists on a vertex
+	if (!ensure(IndexOfBoneInVertex == INDEX_NONE))
+	{
+		return;
+	}
+
+	// should never be adding an influence that doesn't exist in the skeleton
+	if (!ensure(Deformer.InvCSRefPoseTransforms.IsValidIndex(InBoneID)))
+	{
+		return;
+	}
+
+	// add a new influence to this vertex
+	const FVector PosLocalToBone = Deformer.InvCSRefPoseTransforms[InBoneID].TransformPosition(Deformer.RefPoseVertexPositions[InVertexID]);
+	SingleVertexWeights.Emplace(InBoneID, PosLocalToBone, Weight);
+}
+
 void FSkinToolWeights::SwapAfterChange()
 {
 	PreChangeWeights = CurrentWeights;
@@ -970,7 +1076,7 @@ void FSkinToolWeights::UpdateIsBoneWeighted(BoneIndex BoneToUpdate)
 	{
 		for (const FVertexBoneWeight& VertexBoneData : VertexData)
 		{
-			if (VertexBoneData.BoneIndex == BoneToUpdate && VertexBoneData.Weight > UE::AnimationCore::BoneWeightThreshold)
+			if (VertexBoneData.BoneID == BoneToUpdate && VertexBoneData.Weight > UE::AnimationCore::BoneWeightThreshold)
 			{
 				IsBoneWeighted[BoneToUpdate] = true;
 				break;
@@ -983,22 +1089,63 @@ void FSkinToolWeights::UpdateIsBoneWeighted(BoneIndex BoneToUpdate)
 	}
 }
 
+BoneIndex FSkinToolWeights::GetParentBoneToWeightTo(BoneIndex ChildBone)
+{
+	int32 ParentBoneIndex = 0;
+	if (const USkeletalMesh* SkeletalMesh = Deformer.Component->GetSkeletalMeshAsset())
+	{
+		ParentBoneIndex = SkeletalMesh->GetRefSkeleton().GetParentIndex(ChildBone);
+		
+	}
+
+	// are we at the root? (no parent)
+	if (ParentBoneIndex == INDEX_NONE)
+	{
+		ParentBoneIndex = 0; // fallback to root
+		
+		// in this case return the first child bone, if there is one
+		// NOTE: this allows the user to forcibly remove all weight on the root bone, without having another recorded influence on it
+		if (const USkeletalMesh* SkeletalMesh = Deformer.Component->GetSkeletalMeshAsset())
+		{
+			TArray<int32> RootsChildren;
+			SkeletalMesh->GetRefSkeleton().GetDirectChildBones(0, RootsChildren);
+			if (!RootsChildren.IsEmpty())
+			{
+				ParentBoneIndex = RootsChildren[0];
+			}
+		}
+		
+	}
+
+	return ParentBoneIndex;
+}
+
 void FMeshSkinWeightsChange::Apply(UObject* Object)
 {
 	USkinWeightsPaintTool* Tool = CastChecked<USkinWeightsPaintTool>(Object);
 
+	// apply weight edits
 	Tool->ExternalUpdateSkinWeightLayer(LOD, SkinWeightProfile);
 	for (TTuple<int32, FSingleBoneWeightEdits>& Pair : AllWeightEdits.PerBoneWeightEdits)
 	{
 		Tool->ExternalUpdateWeights(Pair.Key, Pair.Value.NewWeights);
 	}
+
+	// remove pruned influences (if there are any)
+	Tool->ExternalRemoveInfluences(AllWeightEdits.PrunedInfluences);
 }
 
 void FMeshSkinWeightsChange::Revert(UObject* Object)
 {
 	USkinWeightsPaintTool* Tool = CastChecked<USkinWeightsPaintTool>(Object);
 
+	// update the skin weight profile
 	Tool->ExternalUpdateSkinWeightLayer(LOD, SkinWeightProfile);
+
+	// apply prune edits (restores pruned influences if there are any)
+	Tool->ExternalAddInfluences(AllWeightEdits.PrunedInfluences);
+	
+	// apply weight edits
 	for (TTuple<int32, FSingleBoneWeightEdits>& Pair : AllWeightEdits.PerBoneWeightEdits)
 	{
 		Tool->ExternalUpdateWeights(Pair.Key, Pair.Value.OldWeights);
@@ -1013,6 +1160,11 @@ void FMeshSkinWeightsChange::AddBoneWeightEdit(const FSingleBoneWeightEdits& Bon
 	AllWeightEdits.MergeEdits(BoneWeightEdit);
 }
 
+void FMeshSkinWeightsChange::AddPruneBoneEdit(const VertexIndex VertexToPruneFrom, const BoneIndex BoneToPrune)
+{
+	AllWeightEdits.PrunedInfluences.Emplace(VertexToPruneFrom, BoneToPrune);
+}
+
 /*
  * ToolBuilder
  */
@@ -1738,10 +1890,10 @@ FVector4f USkinWeightsPaintTool::GetColorOfVertex(VertexIndex InVertexIndex, Bon
 					continue;
 				}
 				
-				const float Value = InCurrentBoneIndex == BoneWeight.BoneIndex ? 1.0f: 0.6f;
+				const float Value = InCurrentBoneIndex == BoneWeight.BoneID ? 1.0f: 0.6f;
 				constexpr float Saturation = 0.75f;
-				const FLinearColor BoneColor = SkeletalDebugRendering::GetSemiRandomColorForBone(BoneWeight.BoneIndex, Value, Saturation);
-				Color = FLinearColor::LerpUsingHSV(Color, BoneColor, BoneWeight.Weight);
+				const FLinearColor BoneColor = SkeletalDebugRendering::GetSemiRandomColorForBone(BoneWeight.BoneID, Value, Saturation);
+				Color = FMath::Lerp(Color, BoneColor, BoneWeight.Weight);
 			}
 			return Color;
 		}
@@ -1846,7 +1998,7 @@ void USkinWeightsPaintTool::ApplyStamp(const FBrushStampData& Stamp)
 			// use mesh topology to iteratively smooth weights across neighboring vertices
 			const float UseStrength = CalculateBrushStrengthToUse(EWeightEditOperation::Relax);
 			constexpr int32 RelaxIterationsPerStamp = 3;
-			RelaxWeightOnVertices(VerticesInStamp, VertexFalloffs, RelaxIterationsPerStamp, UseStrength, WeightEditsFromStamp);
+			RelaxWeightOnVertices(VerticesInStamp, VertexFalloffs, UseStrength, RelaxIterationsPerStamp, WeightEditsFromStamp);
 		}
 		else
 		{
@@ -1906,7 +2058,7 @@ void USkinWeightsPaintTool::EditWeightOfBoneOnVertices(
 	const BoneIndex Bone,
 	const TArray<int32>& VertexIndices,
 	const TArray<float>& VertexFalloffs,
-	const float UseStrength,
+	const float InValue,
 	FMultiBoneWeightEdits& InOutWeightEdits)
 {
 	// spin through the vertices in the stamp and store new weight values in NewValuesFromStamp
@@ -1924,20 +2076,33 @@ void USkinWeightsPaintTool::EditWeightOfBoneOnVertices(
 		{
 		case EWeightEditOperation::Add:
 			{
-				NewValueAfterStamp = ValueBeforeStroke + (UseStrength * UseFalloff);
+				NewValueAfterStamp = ValueBeforeStroke + (InValue * UseFalloff);
 				break;
 			}
 		case EWeightEditOperation::Replace:
 			{
-				NewValueAfterStamp = FMath::Lerp(ValueBeforeStroke, UseStrength, UseFalloff);
+				NewValueAfterStamp = FMath::Lerp(ValueBeforeStroke, InValue, UseFalloff);
 				break;
 			}
 		case EWeightEditOperation::Multiply:
 			{
-				const float DeltaFromThisStamp = ((ValueBeforeStroke * UseStrength) - ValueBeforeStroke) * UseFalloff;
+				const float DeltaFromThisStamp = ((ValueBeforeStroke * InValue) - ValueBeforeStroke) * UseFalloff;
 				NewValueAfterStamp = ValueBeforeStroke + DeltaFromThisStamp;
 				break;
 			}
+		case EWeightEditOperation::RelativeScale:
+			{
+				// LERP the weight from it's current value towards 1 (for positive values) or towards 0 (for negative values)
+				if (InValue >= 0.f)
+				{
+					NewValueAfterStamp = FMath::Lerp(ValueBeforeStroke, 1.0f, FMath::Abs(InValue) * UseFalloff);
+				}
+				else
+				{
+					NewValueAfterStamp = FMath::Lerp(ValueBeforeStroke, 0.0f, FMath::Abs(InValue) * UseFalloff);
+				}
+				break;
+			}
 		default:
 			// relax operation not supported by this function, use RelaxWeightOnVertices()
 			checkNoEntry();
@@ -1956,8 +2121,8 @@ void USkinWeightsPaintTool::EditWeightOfBoneOnVertices(
 void USkinWeightsPaintTool::RelaxWeightOnVertices(
 	TArray<int32> VertexIndices,
 	TArray<float> VertexFalloffs,
-	int32 Iterations,
-	const float UseStrength,
+	const float Strength,
+	const int32 Iterations,
 	FMultiBoneWeightEdits& InOutWeightEdits)
 {
 	if (!ensure(SmoothWeightsOp))
@@ -1965,13 +2130,13 @@ void USkinWeightsPaintTool::RelaxWeightOnVertices(
 		return;
 	}
 	
-	constexpr float PercentPerIteration = 0.95f;
-	for (int32 Iteration=0; Iteration<Iterations; ++Iteration)
+	for (int32 Iteration=0; Iteration < Iterations; ++Iteration)
 	{
 		for (int32 VertexIndex = 0; VertexIndex < VertexIndices.Num(); ++VertexIndex)
 		{
 			const int32 VertexID = VertexIndices[VertexIndex];
-			const float UseFalloff = (VertexFalloffs.IsValidIndex(VertexIndex) ? VertexFalloffs[VertexIndex] * UseStrength : UseStrength) * PercentPerIteration;
+			constexpr float PercentPerIteration = 0.95f;
+			const float UseFalloff = (VertexFalloffs.IsValidIndex(VertexIndex) ? VertexFalloffs[VertexIndex] * Strength : Strength) * PercentPerIteration;
 
 			TMap<int32, float> FinalWeights;
 			const bool bSmoothSuccess = SmoothWeightsOp->SmoothWeightsAtVertex(VertexID, UseFalloff, FinalWeights);
@@ -2088,6 +2253,12 @@ void USkinWeightsPaintTool::ApplyWeightEditsAsTransaction(const SkinPaintTool::F
 	{
 		ActiveChange->AddBoneWeightEdit(BoneWeightEdits.Value);
 	}
+
+	// store pruned influences
+	for (const TPair<VertexIndex, BoneIndex>& PrunedInfluence : WeightEdits.PrunedInfluences)
+	{
+		ActiveChange->AddPruneBoneEdit(PrunedInfluence.Key, PrunedInfluence.Value);
+	}
 	
 	// apply the weight edits of the active change to the actual mesh
 	// - copies weight modifications to the tool's weight data structure and updates the vertex colors
@@ -2246,6 +2417,25 @@ void USkinWeightsPaintTool::ExternalUpdateSkinWeightLayer(const EMeshLODIdentifi
 	}
 }
 
+void USkinWeightsPaintTool::ExternalAddInfluences(const TArray<TPair<VertexIndex, BoneIndex>>& InfluencesToAdd)
+{
+	for (const TPair<VertexIndex, BoneIndex>& ToAdd : InfluencesToAdd)
+	{
+		constexpr float DefaultWeight = 0.f;
+		Weights.AddNewInfluenceToVertex(ToAdd.Key, ToAdd.Value, DefaultWeight, Weights.CurrentWeights);
+		Weights.AddNewInfluenceToVertex(ToAdd.Key, ToAdd.Value, DefaultWeight, Weights.PreChangeWeights);
+	}
+}
+
+void USkinWeightsPaintTool::ExternalRemoveInfluences(const TArray<TPair<VertexIndex, BoneIndex>>& InfluencesToRemove)
+{
+	for (const TPair<VertexIndex, BoneIndex>& ToRemove : InfluencesToRemove)
+	{
+		Weights.RemoveInfluenceFromVertex(ToRemove.Key, ToRemove.Value, Weights.CurrentWeights);
+		Weights.RemoveInfluenceFromVertex(ToRemove.Key, ToRemove.Value, Weights.PreChangeWeights);
+	}
+}
+
 void FSkinMirrorData::RegenerateMirrorData(
     const TArray<FName>& BoneNames,
     const TMap<FName, BoneIndex>& BoneNameToIndexMap,
@@ -2386,7 +2576,7 @@ void USkinWeightsPaintTool::MirrorWeights(EAxis::Type Axis, EMirrorDirection Dir
 	const TMap<VertexIndex, VertexIndex>& VertexMirrorMap = MirrorData.GetVertexMap(); // <Target, Source>
 
 	// get set of vertices to mirror
-	TArray<VertexIndex> AllVerticesToEdit = GetVerticesToEdit();
+	TArray<VertexIndex> AllVerticesToEdit = GetSelectedVertices();
 
 	// convert all vertex indices to the target side of the mirror plane
 	TSet<VertexIndex> VerticesToMirror;
@@ -2425,13 +2615,13 @@ void USkinWeightsPaintTool::MirrorWeights(EAxis::Type Axis, EMirrorDirection Dir
 		{
 			const float OldWeight = TargetBoneWeight.Weight;
 			constexpr float NewWeight = 0.f;
-			WeightEditsFromMirroring.MergeSingleEdit(TargetBoneWeight.BoneIndex, TargetVertexID, OldWeight, NewWeight);
+			WeightEditsFromMirroring.MergeSingleEdit(TargetBoneWeight.BoneID, TargetVertexID, OldWeight, NewWeight);
 		}
 
 		// copy source weights, but with mirrored bones
 		for (const FVertexBoneWeight& SourceBoneWeight : Weights.PreChangeWeights[SourceVertexID])
 		{
-			const BoneIndex MirroredBoneIndex = BoneMap[SourceBoneWeight.BoneIndex];
+			const BoneIndex MirroredBoneIndex = BoneMap[SourceBoneWeight.BoneID];
 			const float OldWeight = Weights.GetWeightOfBoneOnVertex(MirroredBoneIndex, TargetVertexID, Weights.PreChangeWeights);
 			const float NewWeight = SourceBoneWeight.Weight;
 			WeightEditsFromMirroring.MergeSingleEdit(MirroredBoneIndex, TargetVertexID, OldWeight, NewWeight);
@@ -2452,6 +2642,7 @@ void USkinWeightsPaintTool::MirrorWeights(EAxis::Type Axis, EMirrorDirection Dir
 void USkinWeightsPaintTool::EditWeightsOnVertices(
 	BoneIndex Bone,
 	const float Value,
+	const int32 Iterations,
 	EWeightEditOperation EditOperation,
 	const TArray<VertexIndex>& VertexIndices,
 	const bool bShouldTransact)
@@ -2462,9 +2653,7 @@ void USkinWeightsPaintTool::EditWeightsOnVertices(
 
 	if (EditOperation == EWeightEditOperation::Relax)
 	{
-		const int32 Iterations = static_cast<int32>(Value);
-		constexpr float Strength = 1.0f; // always use max strength, and interpret value as iterations for relaxation
-		RelaxWeightOnVertices(GetVerticesToEdit(), VertexFalloffs, Iterations, Strength, DirectWeightEdits);
+		RelaxWeightOnVertices(GetSelectedVertices(), VertexFalloffs, Value, Iterations, DirectWeightEdits);
 	}
 	else
 	{
@@ -2489,32 +2678,65 @@ void USkinWeightsPaintTool::EditWeightsOnVertices(
 	}
 }
 
-void USkinWeightsPaintTool::PruneWeights(float Threshold)
+void USkinWeightsPaintTool::PruneWeights(float Threshold, const TArray<BoneIndex>& BonesToPrune)
 {
-	// set weights below the given threshold to zero
+	// set weights below the given threshold to zero AND remove them as a recorded influence on that vertex
 	FMultiBoneWeightEdits WeightEditsFromPrune;
-	const TArray<VertexIndex>& VerticesToPrune = GetVerticesToEdit();
-	const bool bPruningAllBones = SelectedBoneIndices.IsEmpty();
+	const TArray<VertexIndex>& VerticesToPrune = GetSelectedVertices();
 	for (const VertexIndex VertexID : VerticesToPrune)
 	{
+		TArray<BoneIndex> InfluencesToPrune;
 		const VertexWeights& VertexWeights = Weights.CurrentWeights[VertexID];
 		for (const FVertexBoneWeight& BoneWeight : VertexWeights)
 		{
-			// are we pruning weights only on selected bones?
-			if (!bPruningAllBones && !SelectedBoneIndices.Contains(BoneWeight.BoneIndex))
+			if (BoneWeight.Weight < Threshold || BonesToPrune.Contains(BoneWeight.BoneID))
 			{
-				// not the bone we're pruning...
-				continue;
+				InfluencesToPrune.Add(BoneWeight.BoneID);
+
+				// store a weight edit to remove this weight
+				WeightEditsFromPrune.MergeSingleEdit(BoneWeight.BoneID, VertexID, BoneWeight.Weight, 0.f);
 			}
+		}
+
+		// actually prune the influences from the vert
+		for (const BoneIndex InfluenceToPrune : InfluencesToPrune)
+		{
+			// store this in the transaction
+			WeightEditsFromPrune.AddPruneBoneEdit(VertexID, InfluenceToPrune);
 			
-			if (BoneWeight.Weight > Threshold)
+			// remove the influence from the vertex to prevent subsequent weight editing from normalizing weight back onto it
+			Weights.RemoveInfluenceFromVertex(VertexID, InfluenceToPrune, Weights.CurrentWeights);
+		}
+
+		// at this point, influences are pruned but this may leave the vertex non-normalized
+		if (VertexWeights.IsEmpty())
+		{
+			// we pruned ALL influences from a vertex, so dump all weight on root
+			constexpr BoneIndex RootBoneIndex = 0;
+			const float OldWeight = Weights.GetWeightOfBoneOnVertex(RootBoneIndex, VertexID, Weights.PreChangeWeights);
+			constexpr float NewWeight = 1.f;
+			WeightEditsFromPrune.MergeSingleEdit(RootBoneIndex, VertexID, OldWeight, NewWeight);
+		}
+		else
+		{
+			// re-normalize all existing weights
+			float TotalWeight = 0.f;
+			for (const FVertexBoneWeight& BoneWeight : VertexWeights)
 			{
-				continue;
+				TotalWeight += BoneWeight.Weight;
 			}
 
-			// note: removing ALL weight from a vertex disqualifies it as a candidate to receive weight during normalization
-			// so there is no need to remove the influence from the per-vertex influence array to "prune" it.
-			Weights.EditVertexWeightAndNormalize(BoneWeight.BoneIndex, VertexID,0.f,WeightEditsFromPrune);
+			// if there were no other weights to normalize (all zero), then simply evenly distribute the weight on the recorded influences
+			const bool bNoOtherWeights = FMath::IsNearlyEqual(TotalWeight, 0.f);
+			const float EvenlySplitWeight = 1.0f / VertexWeights.Num();
+
+			// record weight edits to normalize the weight across the remaining influences
+			for (const FVertexBoneWeight& BoneWeight : VertexWeights)
+			{
+				const float OldWeight = BoneWeight.Weight;
+				const float NewWeight = bNoOtherWeights ? EvenlySplitWeight : BoneWeight.Weight / TotalWeight;
+				WeightEditsFromPrune.MergeSingleEdit(BoneWeight.BoneID, VertexID, OldWeight, NewWeight);
+			}
 		}
 	}
 
@@ -2523,75 +2745,146 @@ void USkinWeightsPaintTool::PruneWeights(float Threshold)
 	ApplyWeightEditsAsTransaction(WeightEditsFromPrune, TransactionLabel);
 }
 
-void USkinWeightsPaintTool::AverageWeights()
+void USkinWeightsPaintTool::AverageWeights(const float Strength)
 {
-	// get vertices to edit weights on
-	const TArray<VertexIndex>& VerticesToAverage = GetVerticesToEdit();
-
-	// accumulate ALL the weights on the vertices
-	TMap<BoneIndex, float> AccumulatedWeightMap;
-	for (const VertexIndex VertexID : VerticesToAverage)
+	// if strength is zero, don't do anything
+	if (FMath::IsNearlyEqual(Strength, 0.0f))
 	{
-		for (const FVertexBoneWeight& BoneWeight : Weights.PreChangeWeights[VertexID])
-		{
-			float& AccumulatedWeight = AccumulatedWeightMap.FindOrAdd(BoneWeight.BoneIndex);
-			AccumulatedWeight += BoneWeight.Weight;
-		}
+		return;
 	}
-
-	// sort influences by total weight
-	AccumulatedWeightMap.ValueSort([](const float& A, const float& B)
+	
+	// remove smallest weights values from a weight map to fit in MAX_TOTAL_INFLUENCES
+	auto TruncateWeights = [](TMap<BoneIndex, float>& InOutWeights)
 	{
-		return A > B;
-	});
+		// sort influences by total weight
+		InOutWeights.ValueSort([](const float& A, const float& B)
+		{
+			return A > B;
+		});
 
-	// truncate to MaxInfluences
-	int32 Index = 0;
-	for (TMap<BoneIndex, float>::TIterator It(AccumulatedWeightMap); It; ++It)
+		// truncate to MaxInfluences
+		int32 Index = 0;
+		for (TMap<BoneIndex, float>::TIterator It(InOutWeights); It; ++It)
+		{
+			if (Index >= MAX_TOTAL_INFLUENCES)
+			{
+				It.RemoveCurrent();
+			}
+			else
+			{
+				++Index;
+			}
+		}
+	};
+
+	// normalize weights values to sum to 1.0
+	auto NormalizeWeights = [](TMap<BoneIndex, float>& InOutWeights)
 	{
-		if (Index >= MAX_TOTAL_INFLUENCES)
+		// normalize remaining influences
+		float TotalWeight = 0.f;
+		for (const TTuple<BoneIndex, float>& Weight : InOutWeights)
 		{
-			It.RemoveCurrent();
+			TotalWeight += Weight.Value;
 		}
-		else
+		for (TTuple<BoneIndex, float>& Weight : InOutWeights)
 		{
-			++Index;
+			Weight.Value /= TotalWeight > SMALL_NUMBER ? TotalWeight : 1.f;
 		}
-	}
+	};
 
-	// normalize remaining influences
-	float TotalWeight = 0.f;
-	for (const TTuple<BoneIndex, float>& AccumulatedWeight : AccumulatedWeightMap)
+	// sum up all weight on the given vertices
+	auto AccumulateWeights = [](
+		TMap<BoneIndex, float>& OutWeights,
+		TArray<VertexWeights>& AllWeights,
+		const TArray<VertexIndex>& VerticesToAccumulate)
 	{
-		TotalWeight += AccumulatedWeight.Value;
-	}
-	for (TTuple<BoneIndex, float>& AccumulatedWeight : AccumulatedWeightMap)
-	{
-		AccumulatedWeight.Value /= TotalWeight;
-	}
+		for (const VertexIndex VertexID : VerticesToAccumulate)
+		{
+			for (const FVertexBoneWeight& BoneWeight : AllWeights[VertexID])
+			{
+				float& AccumulatedWeight = OutWeights.FindOrAdd(BoneWeight.BoneID);
+				AccumulatedWeight += BoneWeight.Weight;
+			}
+		}
+	};
+	
+	// get vertices to edit weights on
+	const TArray<VertexIndex>& VerticesToAverage = GetSelectedVertices();
+	TMap<BoneIndex, float> AveragedWeights;
+	AccumulateWeights(AveragedWeights, Weights.PreChangeWeights, VerticesToAverage);
+	TruncateWeights(AveragedWeights);
+	NormalizeWeights(AveragedWeights);
 
-	// apply averaged weights to vertices
+	// store weight edits to apply averaging to selected vertices
 	FMultiBoneWeightEdits WeightEditsFromAveraging;
-	for (const VertexIndex VertexID : VerticesToAverage)
+	
+	// FULLY apply averaged weights to vertices if strength is 1.0
+	if (FMath::IsNearlyEqual(Strength, 1.0f))
 	{
-		// remove influences not a part of the average results
-		for (const FVertexBoneWeight& BoneWeight : Weights.PreChangeWeights[VertexID])
+		for (const VertexIndex VertexID : VerticesToAverage)
 		{
-			if (!AccumulatedWeightMap.Contains(BoneWeight.BoneIndex))
+			// remove influences not a part of the average results
+			for (const FVertexBoneWeight& BoneWeight : Weights.PreChangeWeights[VertexID])
 			{
-				const float OldWeight = BoneWeight.Weight;
-				constexpr float NewWeight = 0.f;
-				WeightEditsFromAveraging.MergeSingleEdit(BoneWeight.BoneIndex, VertexID, OldWeight, NewWeight);
+				if (!AveragedWeights.Contains(BoneWeight.BoneID))
+				{
+					const float OldWeight = BoneWeight.Weight;
+					constexpr float NewWeight = 0.f;
+					WeightEditsFromAveraging.MergeSingleEdit(BoneWeight.BoneID, VertexID, OldWeight, NewWeight);
+				}
 			}
-		}
 
-		// add influences from the averaging results
-		for (const TTuple<BoneIndex, float>& AccumulatedWeight : AccumulatedWeightMap)
+			// add influences from the averaging results
+			for (const TTuple<BoneIndex, float>& AveragedWeight : AveragedWeights)
+			{
+				const BoneIndex IndexOfBone = AveragedWeight.Key;
+				const float OldWeight = Weights.GetWeightOfBoneOnVertex(IndexOfBone, VertexID, Weights.PreChangeWeights);
+				const float NewWeight = AveragedWeight.Value;
+				WeightEditsFromAveraging.MergeSingleEdit(IndexOfBone, VertexID, OldWeight, NewWeight);
+			}
+		}
+	}
+	else
+	{
+		// blend averaged weight with the existing weight based on the strength value
+		const float OldWeightStrength = 1.0f - Strength;
+		const float NewWeightStrength = Strength;
+		for (const VertexIndex VertexID : VerticesToAverage)
 		{
-			const BoneIndex IndexOfBone = AccumulatedWeight.Key;
-			const float OldWeight = Weights.GetWeightOfBoneOnVertex(IndexOfBone, VertexID, Weights.PreChangeWeights);
-			const float NewWeight = AccumulatedWeight.Value;
-			WeightEditsFromAveraging.MergeSingleEdit(IndexOfBone, VertexID, OldWeight, NewWeight);
+			// storage for final blended weights on this vertex
+			TMap<BoneIndex, float> BlendedWeights;
+
+			// scale the existing weights by OldWeightStrength
+			for (const FVertexBoneWeight& BoneWeight : Weights.PreChangeWeights[VertexID])
+			{
+				BlendedWeights.Add(BoneWeight.BoneID, BoneWeight.Weight * OldWeightStrength);
+			}
+			
+			// accumulate existing weights with the scaled averaged weights
+			for (const TTuple<BoneIndex, float>& AveragedWeight : AveragedWeights)
+			{
+				if (BlendedWeights.Contains(AveragedWeight.Key))
+				{
+					BlendedWeights[AveragedWeight.Key] += AveragedWeight.Value * NewWeightStrength;
+				}
+				else
+				{
+					BlendedWeights.Add(AveragedWeight.Key, AveragedWeight.Value * NewWeightStrength);
+				}
+			}
+
+			// enforce max influences and normalize
+			TruncateWeights(BlendedWeights);
+			NormalizeWeights(BlendedWeights);
+			
+			// apply blended weights to this vertex
+			for (const TTuple<BoneIndex, float>& BlendedWeight : BlendedWeights)
+			{
+				const BoneIndex BoneID = BlendedWeight.Key;
+				const float OldWeight = Weights.GetWeightOfBoneOnVertex(BoneID, VertexID, Weights.PreChangeWeights);
+				const float NewWeight = BlendedWeight.Value;
+				WeightEditsFromAveraging.MergeSingleEdit(BoneID, VertexID, OldWeight, NewWeight);
+			}
 		}
 	}
 
@@ -2604,14 +2897,22 @@ void USkinWeightsPaintTool::NormalizeWeights()
 {
 	// re-set a weight on each vertex to force normalization
 	FMultiBoneWeightEdits WeightEditsFromNormalization;
-	const TArray<VertexIndex> VerticesToNormalize = GetVerticesToEdit();
+	const TArray<VertexIndex> VerticesToNormalize = GetSelectedVertices();
 	for (const VertexIndex VertexID : VerticesToNormalize)
 	{
 		const VertexWeights& VertexWeights = Weights.CurrentWeights[VertexID];
-		for (const FVertexBoneWeight& BoneWeight : VertexWeights)
+		if (VertexWeights.IsEmpty())
 		{
-			// set weight to current value, just to force re-normalization
-			Weights.EditVertexWeightAndNormalize(BoneWeight.BoneIndex, VertexID, BoneWeight.Weight,WeightEditsFromNormalization);
+			// ALL influences have been pruned from vertex, so assign it to the root
+			constexpr BoneIndex RootBoneIndex = 0;
+			constexpr float FullWeight = 1.f;
+			Weights.EditVertexWeightAndNormalize(RootBoneIndex, VertexID, FullWeight,WeightEditsFromNormalization);
+		}
+		else
+		{
+			// set first weight to current value, just to force re-normalization
+			const FVertexBoneWeight& BoneWeight = VertexWeights[0];
+			Weights.EditVertexWeightAndNormalize(BoneWeight.BoneID, VertexID, BoneWeight.Weight,WeightEditsFromNormalization);
 		}
 	}
 
@@ -2719,7 +3020,7 @@ void USkinWeightsPaintTool::TransferWeights()
 					for (const FVertexBoneWeight& BoneWeight : VertexBoneWeights)
 					{
 						const float OldWeight = BoneWeight.Weight;
-						WeightEdits.MergeSingleEdit(BoneWeight.BoneIndex, SrcVertexID, OldWeight, ZeroWeight);
+						WeightEdits.MergeSingleEdit(BoneWeight.BoneID, SrcVertexID, OldWeight, ZeroWeight);
 					}
 				}
 				else
@@ -2965,6 +3266,81 @@ void USkinWeightsPaintTool::FloodSelection() const
 	PolygonSelectionMechanic->FloodSelection();
 }
 
+void USkinWeightsPaintTool::SelectAffected() const
+{
+	GetToolManager()->BeginUndoTransaction(LOCTEXT("AffectedSelectionChange", "Select Affected"));
+	PolygonSelectionMechanic->BeginChange();
+	
+	// get all vertices affected by all selected bones
+	TSet<int32> AffectedVertices;
+	for (const BoneIndex SelectedBone : SelectedBoneIndices)
+	{
+		GetVerticesAffectedByBone(SelectedBone, AffectedVertices);
+	}
+	
+	// create selection set
+	FGroupTopologySelection Selection;
+
+	// optionally add/remove/replace selection based on modifier key state
+	const FGroupTopologySelection& CurrentSelection = PolygonSelectionMechanic->GetActiveSelection();
+	if (bShiftToggle)
+	{
+		// ADD to selection
+		Selection.SelectedCornerIDs.Append(CurrentSelection.SelectedCornerIDs);
+		Selection.SelectedCornerIDs.Append(AffectedVertices);
+	}
+	else if (bCtrlToggle)
+	{
+		// REMOVE from selection
+		Selection.SelectedCornerIDs = CurrentSelection.SelectedCornerIDs.Difference(AffectedVertices);
+	}
+	else
+	{
+		// REPLACE selection
+		Selection.SelectedCornerIDs.Append(AffectedVertices);
+	}
+	
+	// select vertices
+	constexpr bool bBroadcast = true;
+	PolygonSelectionMechanic->SetSelection(Selection, bBroadcast);
+	PolygonSelectionMechanic->EndChangeAndEmitIfModified();
+	GetToolManager()->EndUndoTransaction();
+}
+
+void USkinWeightsPaintTool::SelectBorder() const
+{
+	const FDynamicMesh3* Mesh = PreviewMesh->GetMesh();
+
+	const FGroupTopologySelection& CurrentSelection = PolygonSelectionMechanic->GetActiveSelection();
+	
+	// find all border vertices
+	// a "border" vertex is one that has a 1-ring neighbor that is not in the selection set
+	TSet<int32> BorderVertices;
+	for (const int32 SelectedVertex : CurrentSelection.SelectedCornerIDs)
+	{
+		for (const int32 NeighborIndex : Mesh->VtxVerticesItr(SelectedVertex))
+		{
+			if (!CurrentSelection.SelectedCornerIDs.Contains(NeighborIndex))
+			{
+				BorderVertices.Add(SelectedVertex);
+			}
+		}
+	}
+
+	GetToolManager()->BeginUndoTransaction(LOCTEXT("BorderSelectionChange", "Select Border"));
+	PolygonSelectionMechanic->BeginChange();
+
+	// create selection set
+	FGroupTopologySelection Selection;
+	Selection.SelectedCornerIDs.Append(BorderVertices);
+	
+	// select vertices
+	constexpr bool bBroadcast = true;
+	PolygonSelectionMechanic->SetSelection(Selection, bBroadcast);
+	PolygonSelectionMechanic->EndChangeAndEmitIfModified();
+	GetToolManager()->EndUndoTransaction();
+}
+
 bool USkinWeightsPaintTool::IsAnyComponentSelected() const
 {
 	if (!PolygonSelectionMechanic)
@@ -3094,9 +3470,9 @@ void USkinWeightsPaintTool::FinishIsolatedSelection()
 		{
 			const UE::AnimationCore::FBoneWeight& SingleBoneWeight = VertexWeightsPartial[InfluenceIndex];
 			FVertexBoneWeight& VertexBoneWeight = Weights.CurrentWeights[VertexIndexFull][InfluenceIndex];
-			VertexBoneWeight.BoneIndex = SingleBoneWeight.GetBoneIndex();
+			VertexBoneWeight.BoneID = SingleBoneWeight.GetBoneIndex();
 			VertexBoneWeight.Weight = SingleBoneWeight.GetWeight();
-			VertexBoneWeight.VertexInBoneSpace = Weights.Deformer.InvCSRefPoseTransforms[VertexBoneWeight.BoneIndex].TransformPosition(Weights.Deformer.RefPoseVertexPositions[VertexIndexFull]);
+			VertexBoneWeight.VertexInBoneSpace = Weights.Deformer.InvCSRefPoseTransforms[VertexBoneWeight.BoneID].TransformPosition(Weights.Deformer.RefPoseVertexPositions[VertexIndexFull]);
 		}
 	}
 	// sync both weight buffers
@@ -3172,33 +3548,28 @@ const TArray<int32>& USkinWeightsPaintTool::GetSelectedVertices() const
 	return SelectedVertices;
 }
 
-const TArray<int32>& USkinWeightsPaintTool::GetVerticesToEdit()
+void USkinWeightsPaintTool::GetVerticesAffectedByBone(BoneIndex IndexOfBone, TSet<int32>& OutVertexIndices) const
 {
-	// prioritize re-using selected vertices array (if there are any)
-	if (!SelectedVertices.IsEmpty())
+	VertexIndex VertexID = 0;
+	for (const VertexWeights& VertWeights : Weights.PreChangeWeights)
 	{
-		return SelectedVertices;
-	}
-	
-	// fallback on vertices weighted to selected bones
-	VerticesToEdit.Reset();
-	if (!SelectedBoneIndices.IsEmpty())
-	{
-		VertexIndex VertexID = 0;
-		for (const VertexWeights& VertWeights : Weights.PreChangeWeights)
+		for (const FVertexBoneWeight& BoneWeight : VertWeights)
 		{
-			for (const FVertexBoneWeight& BoneWeight : VertWeights)
+			if (BoneWeight.BoneID != IndexOfBone)
 			{
-				if (SelectedBoneIndices.Contains(BoneWeight.BoneIndex) && BoneWeight.Weight > MinimumWeightThreshold)
-				{
-					VerticesToEdit.Add(VertexID);	
-				}
+				continue;
+			}
+
+			if (BoneWeight.Weight < MinimumWeightThreshold)
+			{
+				continue;
 			}
-			++VertexID;
+			
+			OutVertexIndices.Add(VertexID);
 		}
+		
+		++VertexID;
 	}
-
-	return VerticesToEdit;
 }
 
 void USkinWeightsPaintTool::GetSelectedTriangles(TArray<int32>& OutTriangleIndices) const
@@ -3243,10 +3614,7 @@ void USkinWeightsPaintTool::GetInfluences(const TArray<int32>& VertexIndices, TA
 	{
 		for (const FVertexBoneWeight& VertexBoneData : Weights.CurrentWeights[SelectedVertex])
 		{
-			if (VertexBoneData.Weight >= MinimumWeightThreshold)
-			{
-				OutBoneIndices.AddUnique(VertexBoneData.BoneIndex);
-			}
+			OutBoneIndices.AddUnique(VertexBoneData.BoneID);
 		}
 	}
 	
@@ -3270,7 +3638,7 @@ float USkinWeightsPaintTool::GetAverageWeightOnBone(
 		
 		for (const FVertexBoneWeight& VertexBoneData : Weights.CurrentWeights[SelectedVertex])
 		{
-			if (VertexBoneData.BoneIndex == InBoneIndex)
+			if (VertexBoneData.BoneID == InBoneIndex)
 			{
 				++NumVerticesInfluencedByBone;
 				TotalWeight += VertexBoneData.Weight;
diff --git a/Engine/Plugins/Experimental/MeshModelingToolsetExp/Source/MeshModelingToolsEditorOnlyExp/Public/SkeletalMesh/SkinWeightsPaintTool.h b/Engine/Plugins/Experimental/MeshModelingToolsetExp/Source/MeshModelingToolsEditorOnlyExp/Public/SkeletalMesh/SkinWeightsPaintTool.h
--- a/Engine/Plugins/Experimental/MeshModelingToolsetExp/Source/MeshModelingToolsEditorOnlyExp/Public/SkeletalMesh/SkinWeightsPaintTool.h
+++ b/Engine/Plugins/Experimental/MeshModelingToolsetExp/Source/MeshModelingToolsEditorOnlyExp/Public/SkeletalMesh/SkinWeightsPaintTool.h
@@ -80,7 +80,8 @@ enum class EWeightEditOperation : uint8
 	Add,
 	Replace,
 	Multiply,
-	Relax
+	Relax,
+	RelativeScale
 };
 
 // mirror direction mode
@@ -97,11 +98,11 @@ namespace SkinPaintTool
 
 	struct FVertexBoneWeight
 	{
-		FVertexBoneWeight() : BoneIndex(INDEX_NONE), VertexInBoneSpace(FVector::ZeroVector), Weight(0.0f) {}
+		FVertexBoneWeight() : BoneID(INDEX_NONE), VertexInBoneSpace(FVector::ZeroVector), Weight(0.0f) {}
 		FVertexBoneWeight(BoneIndex InBoneIndex, const FVector& InPosInRefPose, float InWeight) :
-			BoneIndex(InBoneIndex), VertexInBoneSpace(InPosInRefPose), Weight(InWeight){}
+			BoneID(InBoneIndex), VertexInBoneSpace(InPosInRefPose), Weight(InWeight){}
 		
-		BoneIndex BoneIndex;
+		BoneIndex BoneID;
 		FVector VertexInBoneSpace;
 		float Weight;
 	};
@@ -154,11 +155,15 @@ namespace SkinPaintTool
 		void MergeEdits(const FSingleBoneWeightEdits& BoneWeightEdits);
 		float GetVertexDeltaFromEdits(const int32 BoneIndex, const int32 VertexIndex);
 		void GetEditedVertexIndices(TSet<int32>& OutVerticesToEdit) const;
+		void AddPruneBoneEdit(const VertexIndex VertexToPruneFrom, const BoneIndex BoneToPrune);
 
 		// map of bone indices to weight edits made to that bone
 		TMap<BoneIndex, FSingleBoneWeightEdits> PerBoneWeightEdits;
-	};
 
+		// influences to prune as part of these edits
+		TArray<TPair<VertexIndex, BoneIndex>> PrunedInfluences;
+	};
+	
 	class FMeshSkinWeightsChange : public FToolCommandChange
 	{
 	public:
@@ -170,7 +175,7 @@ namespace SkinPaintTool
 
 		virtual FString ToString() const override
 		{
-			return FString(TEXT("Paint Skin Weights"));
+			return FString(TEXT("Edit Skin Weights"));
 		}
 
 		virtual void Apply(UObject* Object) override;
@@ -179,6 +184,8 @@ namespace SkinPaintTool
 
 		void AddBoneWeightEdit(const FSingleBoneWeightEdits& BoneWeightEdit);
 
+		void AddPruneBoneEdit(const VertexIndex VertexToPruneFrom, const BoneIndex BoneToPrune);
+
 	private:
 		FMultiBoneWeightEdits AllWeightEdits;
 		EMeshLODIdentifier LOD = EMeshLODIdentifier::Default;
@@ -214,6 +221,17 @@ namespace SkinPaintTool
 			const float Weight,
 			TArray<VertexWeights>& InOutVertexData);
 
+		void RemoveInfluenceFromVertex(
+			const VertexIndex VertexID,
+			const BoneIndex BoneID,
+			TArray<VertexWeights>& InOutVertexWeights);
+
+		void AddNewInfluenceToVertex(
+			const VertexIndex VertexID,
+			const BoneIndex BoneIndex,
+			const float Weight,
+			TArray<VertexWeights>& InOutVertexWeights);
+
 		void SwapAfterChange();
 
 		float SetCurrentFalloffAndGetMaxFalloffThisStroke(int32 VertexID, float CurrentStrength);
@@ -221,6 +239,8 @@ namespace SkinPaintTool
 		void ApplyEditsToCurrentWeights(const FMultiBoneWeightEdits& Edits);
 
 		void UpdateIsBoneWeighted(BoneIndex BoneToUpdate);
+
+		BoneIndex GetParentBoneToWeightTo(BoneIndex ChildBone);
 		
 		// double-buffer of the entire weight matrix (stored sparsely for fast deformation)
 		// "Pre" is state of weights at stroke start
@@ -350,9 +370,13 @@ public:
 	UPROPERTY(Config)
 	float PruneValue = 0.01;
 	UPROPERTY(Config)
+	float AddStrength = 1.0;
+	UPROPERTY(Config)
 	float ReplaceValue = 1.0;
 	UPROPERTY(Config)
-	int32 RelaxIterations = 5;
+	float RelaxStrength = 0.5;
+	UPROPERTY(Config)
+	float AverageStrength = 1.0;
 	// the state of the direct weight editing tools (mode buttons + slider)
 	FDirectEditWeightState DirectEditState;
 
@@ -457,11 +481,13 @@ public:
 	// using when ToolChange is applied via Undo/Redo
 	void ExternalUpdateWeights(const int32 BoneIndex, const TMap<int32, float>& IndexValues);
 	void ExternalUpdateSkinWeightLayer(const EMeshLODIdentifier InLOD, const FName InSkinWeightProfile);
+	void ExternalAddInfluences(const TArray<TPair<VertexIndex, BoneIndex>>& InfluencesToAdd);
+	void ExternalRemoveInfluences(const TArray<TPair<VertexIndex, BoneIndex>>& InfluencesToRemove);
 
 	// weight editing operations (selection based)
 	void MirrorWeights(EAxis::Type Axis, EMirrorDirection Direction);
-	void PruneWeights(const float Threshold);
-	void AverageWeights();
+	void PruneWeights(const float Threshold, const TArray<BoneIndex>& BonesToPrune);
+	void AverageWeights(const float Strength);
 	void NormalizeWeights();
 	void TransferWeights();
 	
@@ -469,6 +495,7 @@ public:
 	void EditWeightsOnVertices(
 		BoneIndex Bone,
 		const float Value,
+		const int32 Iterations,
 		EWeightEditOperation EditOperation,
 		const TArray<VertexIndex>& VerticesToEdit,
 		const bool bShouldTransact);
@@ -481,6 +508,8 @@ public:
 	void GrowSelection() const;
 	void ShrinkSelection() const;
 	void FloodSelection() const;
+	void SelectAffected() const;
+	void SelectBorder() const;
 	// isolate selection
 	bool IsAnyComponentSelected() const;
 	bool IsSelectionIsolated() const;
@@ -488,7 +517,7 @@ public:
 
 	// get a list of currently selected vertices (converting edges and faces to vertices)
 	const TArray<int32>& GetSelectedVertices() const;
-	const TArray<int32>& GetVerticesToEdit();
+	void GetVerticesAffectedByBone(BoneIndex IndexOfBone, TSet<int32>& OutVertexIndices) const;
 	void GetSelectedTriangles(TArray<int32>& OutTriangleIndices) const;
 
 	// get the average weight value of each influence on the given vertices
@@ -568,14 +597,14 @@ protected:
 		const BoneIndex Bone,
 		const TArray<int32>& VerticesToEdit,
 		const TArray<float>& VertexFalloffs,
-		const float UseStrength,
+		const float InValue,
 		SkinPaintTool::FMultiBoneWeightEdits& InOutWeightEdits);
 	// same as EditWeightOfVertices() but specific to relaxation (topology aware operation)
 	void RelaxWeightOnVertices(
 		TArray<int32> VerticesToEdit,
 		TArray<float> VertexFalloffs,
-		int32 Iterations,
-		const float UseStrength,
+		const float Strength,
+		const int32 Iterations,
 		SkinPaintTool::FMultiBoneWeightEdits& InOutWeightEdits);
 
 	// used to accelerate mesh queries
@@ -649,7 +678,6 @@ protected:
 	TUniquePtr<UE::Geometry::FTriangleGroupTopology> SelectionTopology = nullptr;
 	void InitializeSelectionMechanic();
 	TArray<VertexIndex> SelectedVertices;
-	TArray<VertexIndex> VerticesToEdit;
 
 	// isolate selection sub-meshes
 	UE::Geometry::FDynamicSubmesh3 PartialSubMesh;
