Added "Hammer" weights operation to skin weight tool.

diff --git a/Engine/Plugins/Experimental/Animation/SkeletalMeshModelingTools/Source/SkeletalMeshModelingTools/Private/DetailCustomization/SkinWeightDetailCustomization.cpp b/Engine/Plugins/Experimental/Animation/SkeletalMeshModelingTools/Source/SkeletalMeshModelingTools/Private/DetailCustomization/SkinWeightDetailCustomization.cpp
--- a/Engine/Plugins/Experimental/Animation/SkeletalMeshModelingTools/Source/SkeletalMeshModelingTools/Private/DetailCustomization/SkinWeightDetailCustomization.cpp
+++ b/Engine/Plugins/Experimental/Animation/SkeletalMeshModelingTools/Source/SkeletalMeshModelingTools/Private/DetailCustomization/SkinWeightDetailCustomization.cpp
@@ -1033,6 +1033,27 @@ void FSkinWeightDetailCustomization::AddSelectionUI(IDetailLayoutBuilder& Detail
 				return FReply::Handled();
 			})
 		]
+
+		+SHorizontalBox::Slot()
+		.Padding(2.f, WeightEditVerticalPadding)
+		[
+			SNew(SButton)
+			.HAlign(HAlign_Center)
+			.VAlign(VAlign_Center)
+			.Text(LOCTEXT("HammerWeightsButtonLabel", "Hammer"))
+			.ToolTipText(LOCTEXT("HammerWeightsTooltip",
+					"Copies the weight of the nearest non-selected vertex.\n"
+					"This command operates on the selected vertices."))
+			.IsEnabled_Lambda([this]()
+			{
+				return ToolSettings->EditingMode == EWeightEditMode::Mesh;
+			})
+			.OnClicked_Lambda([this]()
+			{
+				Tool->HammerWeights();
+				return FReply::Handled();
+			})
+		]
 	];
 
 	// PRUNE WEIGHTS category
diff --git a/Engine/Plugins/Experimental/MeshModelingToolsetExp/Source/MeshModelingToolsEditorOnlyExp/Private/SkeletalMesh/SkinWeightsPaintTool.cpp b/Engine/Plugins/Experimental/MeshModelingToolsetExp/Source/MeshModelingToolsEditorOnlyExp/Private/SkeletalMesh/SkinWeightsPaintTool.cpp
--- a/Engine/Plugins/Experimental/MeshModelingToolsetExp/Source/MeshModelingToolsEditorOnlyExp/Private/SkeletalMesh/SkinWeightsPaintTool.cpp
+++ b/Engine/Plugins/Experimental/MeshModelingToolsetExp/Source/MeshModelingToolsEditorOnlyExp/Private/SkeletalMesh/SkinWeightsPaintTool.cpp
@@ -45,6 +45,7 @@
 #include "TargetInterfaces/MeshDescriptionCommitter.h"
 #include "TargetInterfaces/MeshDescriptionProvider.h"
 #include "Operations/TransferBoneWeights.h"
+#include "Parameterization/MeshDijkstra.h"
 
 #include UE_INLINE_GENERATED_CPP_BY_NAME(SkinWeightsPaintTool)
 
@@ -2921,6 +2922,80 @@ void USkinWeightsPaintTool::NormalizeWeights()
 	ApplyWeightEditsAsTransaction(WeightEditsFromNormalization, TransactionLabel);
 }
 
+void USkinWeightsPaintTool::HammerWeights()
+{
+	// get selected vertices
+	const TArray<VertexIndex> SelectedVerts = GetSelectedVertices();
+	if (SelectedVerts.IsEmpty())
+	{
+		return;
+	}
+
+	// reset mesh to ref pose so that Dijkstra path lengths are not deformed
+	Weights.Deformer.SetToRefPose(this);
+	
+	// find 1-ring neighbors of the current selection, lets call these "Surrounding" vertices
+	const FDynamicMesh3* Mesh = PreviewMesh->GetMesh();
+	TSet<int32> SurroundingVertices;
+	for (const int32 SelectedVertex : SelectedVerts)
+	{
+		for (const int32 NeighborIndex : Mesh->VtxVerticesItr(SelectedVertex))
+		{
+			if (!SelectedVerts.Contains(NeighborIndex))
+			{
+				SurroundingVertices.Add(NeighborIndex);
+			}
+		}
+	}
+
+	// seed a Dijkstra path finder with the surrounding vertices
+	UE::Geometry::TMeshDijkstra<FDynamicMesh3> PathFinder(Mesh);
+	TArray<UE::Geometry::TMeshDijkstra<FDynamicMesh3>::FSeedPoint> SeedPoints;
+	for (const int32 SurroundingVertex : SurroundingVertices)
+	{
+		SeedPoints.Add({ SurroundingVertex, SurroundingVertex, 0 });
+	}
+	PathFinder.ComputeToMaxDistance(SeedPoints, TNumericLimits<double>::Max());
+
+	// create set of weight edits that hammer the weights
+	FMultiBoneWeightEdits HammerWeightEdits;
+	
+	// for each selected vertex, find the nearest surrounding vertex and copy it's weights
+	TArray<int32> VertexPath;
+	for (const int32 SelectedVertex : SelectedVerts)
+	{
+		// find the closest surrounding vertex to this selected vertex
+		if (!PathFinder.FindPathToNearestSeed(SelectedVertex, VertexPath))
+		{
+			continue;
+		}
+		const int32 ClosestVertex = VertexPath.Last();
+
+		// remove all current weights
+		for (const FVertexBoneWeight& BoneWeight : Weights.PreChangeWeights[SelectedVertex])
+		{
+			const float OldWeight = BoneWeight.Weight;
+			constexpr float NewWeight = 0.f;
+			HammerWeightEdits.MergeSingleEdit(BoneWeight.BoneID, SelectedVertex, OldWeight, NewWeight);
+		}
+
+		// add weights from closest vertex
+		for (const FVertexBoneWeight& BoneWeight : Weights.PreChangeWeights[ClosestVertex])
+		{
+			const float OldWeight = Weights.GetWeightOfBoneOnVertex(BoneWeight.BoneID, SelectedVertex, Weights.PreChangeWeights);
+			const float NewWeight = BoneWeight.Weight;
+			HammerWeightEdits.MergeSingleEdit(BoneWeight.BoneID, SelectedVertex, OldWeight, NewWeight);
+		}
+	}
+	
+	// apply the changes
+	const FText TransactionLabel = LOCTEXT("HammerWeightsChange", "Hammer skin weights.");
+	ApplyWeightEditsAsTransaction(HammerWeightEdits, TransactionLabel);
+
+	// put the mesh back in it's current pose
+	Weights.Deformer.SetAllVerticesToBeUpdated();
+}
+
 void USkinWeightsPaintTool::TransferWeights()
 {
 	using namespace UE::Geometry;
diff --git a/Engine/Plugins/Experimental/MeshModelingToolsetExp/Source/MeshModelingToolsEditorOnlyExp/Public/SkeletalMesh/SkinWeightsPaintTool.h b/Engine/Plugins/Experimental/MeshModelingToolsetExp/Source/MeshModelingToolsEditorOnlyExp/Public/SkeletalMesh/SkinWeightsPaintTool.h
--- a/Engine/Plugins/Experimental/MeshModelingToolsetExp/Source/MeshModelingToolsEditorOnlyExp/Public/SkeletalMesh/SkinWeightsPaintTool.h
+++ b/Engine/Plugins/Experimental/MeshModelingToolsetExp/Source/MeshModelingToolsEditorOnlyExp/Public/SkeletalMesh/SkinWeightsPaintTool.h
@@ -489,6 +489,7 @@ public:
 	void PruneWeights(const float Threshold, const TArray<BoneIndex>& BonesToPrune);
 	void AverageWeights(const float Strength);
 	void NormalizeWeights();
+	void HammerWeights();
 	void TransferWeights();
 	
 	// method to set weights directly (numeric input, for example)
