ModelingTools: Added UVShell BakeTexture option

diff --git a/Engine/Plugins/Editor/ModelingToolsEditorMode/Source/ModelingToolsEditorMode/Private/DetailsCustomizations/BakeMeshAttributeToolCustomizations.cpp b/Engine/Plugins/Editor/ModelingToolsEditorMode/Source/ModelingToolsEditorMode/Private/DetailsCustomizations/BakeMeshAttributeToolCustomizations.cpp
--- a/Engine/Plugins/Editor/ModelingToolsEditorMode/Source/ModelingToolsEditorMode/Private/DetailsCustomizations/BakeMeshAttributeToolCustomizations.cpp
+++ b/Engine/Plugins/Editor/ModelingToolsEditorMode/Source/ModelingToolsEditorMode/Private/DetailsCustomizations/BakeMeshAttributeToolCustomizations.cpp
@@ -65,6 +65,7 @@ namespace BakeCustomizationLocals
 				EBakeMapType::VertexColor,
 				EBakeMapType::MaterialID,
 				EBakeMapType::PolyGroupID,
+				EBakeMapType::UVShell,
 				EBakeMapType::One,
 				EBakeMapType::Zero
 			}
diff --git a/Engine/Plugins/Experimental/MeshModelingToolsetExp/Source/MeshModelingToolsExp/Private/BakeMeshAttributeMapsTool.cpp b/Engine/Plugins/Experimental/MeshModelingToolsetExp/Source/MeshModelingToolsExp/Private/BakeMeshAttributeMapsTool.cpp
--- a/Engine/Plugins/Experimental/MeshModelingToolsetExp/Source/MeshModelingToolsExp/Private/BakeMeshAttributeMapsTool.cpp
+++ b/Engine/Plugins/Experimental/MeshModelingToolsetExp/Source/MeshModelingToolsExp/Private/BakeMeshAttributeMapsTool.cpp
@@ -13,6 +13,7 @@
 #include "Sampling/MeshOcclusionMapEvaluator.h"
 #include "Sampling/MeshCurvatureMapEvaluator.h"
 #include "Sampling/MeshPropertyMapEvaluator.h"
+#include "Sampling/MeshUVShellMapEvaluator.h"
 #include "Sampling/MeshResampleImageEvaluator.h"
 
 #include "ImageUtils.h"
@@ -115,6 +116,7 @@ public:
 	FOcclusionMapSettings OcclusionSettings;
 	FCurvatureMapSettings CurvatureSettings;
 	FMeshPropertyMapSettings PropertySettings;
+	FUVShellMapSettings UVShellSettings;
 	FTexture2DSettings TextureSettings;
 	FTexture2DSettings MultiTextureSettings;
 
@@ -190,7 +192,7 @@ public:
 			{
 			case EBakeMapType::TangentSpaceNormal:
 			{
-				TSharedPtr<FMeshNormalMapEvaluator, ESPMode::ThreadSafe> NormalEval = MakeShared<FMeshNormalMapEvaluator, ESPMode::ThreadSafe>();
+				TSharedPtr<FMeshNormalMapEvaluator> NormalEval = MakeShared<FMeshNormalMapEvaluator>();
 				DetailSampler.SetNormalTextureMap(DetailMesh.Get(), IMeshBakerDetailSampler::FBakeDetailNormalTexture(DetailMeshNormalMap.Get(), DetailMeshNormalUVLayer, DetailMeshNormalSpace));
 				Baker->AddEvaluator(NormalEval);
 				break;
@@ -209,7 +211,7 @@ public:
 			}
 			case EBakeMapType::Curvature:
 			{
-				TSharedPtr<FMeshCurvatureMapEvaluator, ESPMode::ThreadSafe> CurvatureEval = MakeShared<FMeshCurvatureMapEvaluator, ESPMode::ThreadSafe>();
+				TSharedPtr<FMeshCurvatureMapEvaluator> CurvatureEval = MakeShared<FMeshCurvatureMapEvaluator>();
 				CurvatureEval->RangeScale = FMathd::Clamp(CurvatureSettings.RangeMultiplier, 0.0001, 1000.0);
 				CurvatureEval->MinRangeScale = FMathd::Clamp(CurvatureSettings.MinRangeMultiplier, 0.0, 1.0);
 				CurvatureEval->UseCurvatureType = (FMeshCurvatureMapEvaluator::ECurvatureType)CurvatureSettings.CurvatureType;
@@ -220,7 +222,7 @@ public:
 			}
 			case EBakeMapType::ObjectSpaceNormal:
 			{
-				TSharedPtr<FMeshPropertyMapEvaluator, ESPMode::ThreadSafe> PropertyEval = MakeShared<FMeshPropertyMapEvaluator, ESPMode::ThreadSafe>();
+				TSharedPtr<FMeshPropertyMapEvaluator> PropertyEval = MakeShared<FMeshPropertyMapEvaluator>();
 				PropertyEval->Property = EMeshPropertyMapType::Normal;
 				DetailSampler.SetNormalTextureMap(DetailMesh.Get(), IMeshBakerDetailSampler::FBakeDetailNormalTexture(DetailMeshNormalMap.Get(), DetailMeshNormalUVLayer, DetailMeshNormalSpace));
 				Baker->AddEvaluator(PropertyEval);
@@ -228,49 +230,61 @@ public:
 			}
 			case EBakeMapType::FaceNormal:
 			{
-				TSharedPtr<FMeshPropertyMapEvaluator, ESPMode::ThreadSafe> PropertyEval = MakeShared<FMeshPropertyMapEvaluator, ESPMode::ThreadSafe>();
+				TSharedPtr<FMeshPropertyMapEvaluator> PropertyEval = MakeShared<FMeshPropertyMapEvaluator>();
 				PropertyEval->Property = EMeshPropertyMapType::FacetNormal;
 				Baker->AddEvaluator(PropertyEval);
 				break;
 			}
 			case EBakeMapType::Position:
 			{
-				TSharedPtr<FMeshPropertyMapEvaluator, ESPMode::ThreadSafe> PropertyEval = MakeShared<FMeshPropertyMapEvaluator, ESPMode::ThreadSafe>();
+				TSharedPtr<FMeshPropertyMapEvaluator> PropertyEval = MakeShared<FMeshPropertyMapEvaluator>();
 				PropertyEval->Property = EMeshPropertyMapType::Position;
 				Baker->AddEvaluator(PropertyEval);
 				break;
 			}
 			case EBakeMapType::MaterialID:
 			{
-				TSharedPtr<FMeshPropertyMapEvaluator, ESPMode::ThreadSafe> PropertyEval = MakeShared<FMeshPropertyMapEvaluator, ESPMode::ThreadSafe>();
+				TSharedPtr<FMeshPropertyMapEvaluator> PropertyEval = MakeShared<FMeshPropertyMapEvaluator>();
 				PropertyEval->Property = EMeshPropertyMapType::MaterialID;
 				Baker->AddEvaluator(PropertyEval);
 				break;
 			}
 			case EBakeMapType::PolyGroupID:
 			{
-				TSharedPtr<FMeshPropertyMapEvaluator, ESPMode::ThreadSafe> PropertyEval = MakeShared<FMeshPropertyMapEvaluator, ESPMode::ThreadSafe>();
+				TSharedPtr<FMeshPropertyMapEvaluator> PropertyEval = MakeShared<FMeshPropertyMapEvaluator>();
 				PropertyEval->Property = EMeshPropertyMapType::PolyGroupID;
 				Baker->AddEvaluator(PropertyEval);
 				break;
 			}
+			case EBakeMapType::UVShell:
+			{
+				TSharedPtr<FMeshUVShellMapEvaluator> UVShellEval = MakeShared<FMeshUVShellMapEvaluator>();
+				UVShellEval->TexelSize = BakeSettings.Dimensions.GetTexelSize();
+				UVShellEval->UVLayer = UVShellSettings.UVLayer;
+				UVShellEval->WireframeThickness = UVShellSettings.WireframeThickness;
+				UVShellEval->WireframeColor = UVShellSettings.WireframeColor;
+				UVShellEval->ShellColor = UVShellSettings.ShellColor;
+				UVShellEval->BackgroundColor = UVShellSettings.BackgroundColor;
+				Baker->AddEvaluator(UVShellEval);
+				break;
+			}
 			case EBakeMapType::VertexColor:
 			{
-				TSharedPtr<FMeshPropertyMapEvaluator, ESPMode::ThreadSafe> PropertyEval = MakeShared<FMeshPropertyMapEvaluator, ESPMode::ThreadSafe>();
+				TSharedPtr<FMeshPropertyMapEvaluator> PropertyEval = MakeShared<FMeshPropertyMapEvaluator>();
 				PropertyEval->Property = EMeshPropertyMapType::VertexColor;
 				Baker->AddEvaluator(PropertyEval);
 				break;
 			}
 			case EBakeMapType::Texture:
 			{
-				TSharedPtr<FMeshResampleImageEvaluator, ESPMode::ThreadSafe> TextureEval = MakeShared<FMeshResampleImageEvaluator, ESPMode::ThreadSafe>();
+				TSharedPtr<FMeshResampleImageEvaluator> TextureEval = MakeShared<FMeshResampleImageEvaluator>();
 				DetailSampler.SetTextureMap(DetailMesh.Get(), IMeshBakerDetailSampler::FBakeDetailTexture(TextureImage.Get(), TextureSettings.UVLayer));
 				Baker->AddEvaluator(TextureEval);
 				break;
 			}
 			case EBakeMapType::MultiTexture:
 			{
-				TSharedPtr<FMeshMultiResampleImageEvaluator, ESPMode::ThreadSafe> TextureEval = MakeShared<FMeshMultiResampleImageEvaluator, ESPMode::ThreadSafe>();
+				TSharedPtr<FMeshMultiResampleImageEvaluator> TextureEval = MakeShared<FMeshMultiResampleImageEvaluator>();
 				TextureEval->DetailUVLayer = MultiTextureSettings.UVLayer;
 				TextureEval->MultiTextures = MaterialIDTextures;
 				Baker->AddEvaluator(TextureEval);
@@ -392,6 +406,17 @@ void UBakeMeshAttributeMapsTool::Setup()
 	CurvatureSettings->WatchProperty(CurvatureSettings->Clamping, [this](EBakeCurvatureClampMode) { OpState |= EBakeOpState::Evaluate; });
 
 
+	UVShellSettings = NewObject<UBakeUVShellMapToolProperties>(this);
+	UVShellSettings->RestoreProperties(this);
+	AddToolPropertySource(UVShellSettings);
+	SetToolPropertySourceEnabled(UVShellSettings, false);
+	UVShellSettings->WatchProperty(UVShellSettings->UVLayer, [this](int) { OpState |= EBakeOpState::Evaluate; });
+	UVShellSettings->WatchProperty(UVShellSettings->WireframeThickness, [this](float) { OpState |= EBakeOpState::Evaluate; });
+	UVShellSettings->WatchProperty(UVShellSettings->WireframeColor, [this](FLinearColor) { OpState |= EBakeOpState::Evaluate; });
+	UVShellSettings->WatchProperty(UVShellSettings->ShellColor, [this](FLinearColor) { OpState |= EBakeOpState::Evaluate; });
+	UVShellSettings->WatchProperty(UVShellSettings->BackgroundColor, [this](FLinearColor) { OpState |= EBakeOpState::Evaluate; });
+
+
 	TextureSettings = NewObject<UBakeTexture2DProperties>(this);
 	TextureSettings->RestoreProperties(this);
 	AddToolPropertySource(TextureSettings);
@@ -492,6 +517,11 @@ TUniquePtr<UE::Geometry::TGenericDataOperator<FMeshMapBaker>> UBakeMeshAttribute
 		Op->PropertySettings = CachedMeshPropertyMapSettings;
 	}
 
+	if ((bool)(CachedBakeSettings.BakeMapTypes & EBakeMapType::UVShell))
+	{
+		Op->UVShellSettings = CachedUVShellMapSettings;
+	}
+
 	if ((bool)(CachedBakeSettings.BakeMapTypes & EBakeMapType::Texture))
 	{
 		Op->TextureSettings = CachedTexture2DSettings;
@@ -679,6 +709,10 @@ void UBakeMeshAttributeMapsTool::UpdateResult()
 	{
 		OpState |= UpdateResult_MeshProperty(CachedBakeSettings.Dimensions);
 	}
+	if ((bool)(CachedBakeSettings.BakeMapTypes & EBakeMapType::UVShell))
+	{
+		OpState |= UpdateResult_UVShellMap(CachedBakeSettings.Dimensions);
+	}
 	if ((bool)(CachedBakeSettings.BakeMapTypes & EBakeMapType::Texture))
 	{
 		OpState |= UpdateResult_Texture2DImage(CachedBakeSettings.Dimensions, DetailMesh.Get());
@@ -756,6 +790,25 @@ EBakeOpState UBakeMeshAttributeMapsTool::UpdateResult_DetailNormalMap()
 }
 
 
+EBakeOpState UBakeMeshAttributeMapsTool::UpdateResult_UVShellMap(const FImageDimensions& Dimensions)
+{
+	EBakeOpState ResultState = EBakeOpState::Clean;
+
+	FUVShellMapSettings UVShellMapSettings;
+	UVShellMapSettings.UVLayer = UVShellSettings->UVLayer;
+	UVShellMapSettings.WireframeThickness = UVShellSettings->WireframeThickness;
+	UVShellMapSettings.WireframeColor = UVShellSettings->WireframeColor;
+	UVShellMapSettings.ShellColor = UVShellSettings->ShellColor;
+	UVShellMapSettings.BackgroundColor = UVShellSettings->BackgroundColor;
+
+	if (!(CachedUVShellMapSettings == UVShellMapSettings))
+	{
+		CachedUVShellMapSettings = UVShellMapSettings;
+		ResultState |= EBakeOpState::Evaluate;
+	}
+	return ResultState;
+}
+
 void UBakeMeshAttributeMapsTool::UpdateVisualization()
 {
 	PreviewMesh->SetOverrideRenderMaterial(PreviewMaterial);
@@ -785,6 +838,7 @@ void UBakeMeshAttributeMapsTool::UpdateOnModeChange()
 	// Update tool property sets.
 	SetToolPropertySourceEnabled(OcclusionSettings, false);
 	SetToolPropertySourceEnabled(CurvatureSettings, false);
+	SetToolPropertySourceEnabled(UVShellSettings, false);
 	SetToolPropertySourceEnabled(TextureSettings, false);
 	SetToolPropertySourceEnabled(MultiTextureSettings, false);
 
@@ -808,6 +862,9 @@ void UBakeMeshAttributeMapsTool::UpdateOnModeChange()
 		case EBakeMapType::PolyGroupID:
 		case EBakeMapType::VertexColor:
 			break;
+		case EBakeMapType::UVShell:
+			SetToolPropertySourceEnabled(UVShellSettings, true);
+			break;
 		case EBakeMapType::Texture:
 			SetToolPropertySourceEnabled(TextureSettings, true);
 			break;
diff --git a/Engine/Plugins/Experimental/MeshModelingToolsetExp/Source/MeshModelingToolsExp/Private/BakeMeshAttributeMapsToolBase.cpp b/Engine/Plugins/Experimental/MeshModelingToolsetExp/Source/MeshModelingToolsExp/Private/BakeMeshAttributeMapsToolBase.cpp
--- a/Engine/Plugins/Experimental/MeshModelingToolsetExp/Source/MeshModelingToolsExp/Private/BakeMeshAttributeMapsToolBase.cpp
+++ b/Engine/Plugins/Experimental/MeshModelingToolsetExp/Source/MeshModelingToolsExp/Private/BakeMeshAttributeMapsToolBase.cpp
@@ -264,6 +264,7 @@ void UBakeMeshAttributeMapsToolBase::UpdatePreview(const EBakeMapType PreviewMap
 		case EBakeMapType::Texture:
 		case EBakeMapType::MultiTexture:
 		case EBakeMapType::VertexColor:
+		case EBakeMapType::UVShell:
 			PreviewMaterial->SetTextureParameterValue(TEXT("NormalMap"), EmptyNormalMap);
 			PreviewMaterial->SetTextureParameterValue(TEXT("OcclusionMap"), EmptyColorMapWhite);
 			PreviewMaterial->SetTextureParameterValue(TEXT("ColorMap"), PreviewMap);
@@ -454,6 +455,12 @@ void UBakeMeshAttributeMapsToolBase::OnMapsUpdated(const TUniquePtr<UE::Geometry
 			UpdateCachedMap(MapType, Format, 0);
 			break;
 		}
+		case EMeshMapEvaluatorType::UVShell:
+		{
+			constexpr EBakeMapType MapType = EBakeMapType::UVShell;
+			UpdateCachedMap(MapType, Format, 0);
+			break;
+		}
 		case EMeshMapEvaluatorType::ResampleImage:
 		{
 			constexpr EBakeMapType MapType = EBakeMapType::Texture;
@@ -515,6 +522,8 @@ FTexture2DBuilder::ETextureType UBakeMeshAttributeMapsToolBase::GetTextureType(c
 	case EBakeMapType::MaterialID:
 	case EBakeMapType::PolyGroupID:
 	case EBakeMapType::VertexColor:
+	case EBakeMapType::UVShell:
+		TexType = FTexture2DBuilder::ETextureType::Color;
 		break;
 	case EBakeMapType::Texture:
 	case EBakeMapType::MultiTexture:
@@ -566,6 +575,9 @@ void UBakeMeshAttributeMapsToolBase::GetTextureName(const EBakeMapType MapType,
 	case EBakeMapType::Position:
 		TexName = FString::Printf(TEXT("%s_Position"), *BaseName);
 		break;
+	case EBakeMapType::UVShell:
+		TexName = FString::Printf(TEXT("%s_UVShell"), *BaseName);
+		break;
 	case EBakeMapType::Texture:
 		TexName = FString::Printf(TEXT("%s_Texture"), *BaseName);
 		break;
@@ -735,6 +747,9 @@ void UBakeMeshAttributeMapsToolBase::RecordAnalytics(const FBakeAnalytics& Data,
 	const bool bVertexColor = static_cast<bool>(Data.BakeSettings.SourceBakeMapTypes & EBakeMapType::VertexColor);
 	Attributes.Add(FAnalyticsEventAttribute(TEXT("Settings.VertexColor.Enabled"), bVertexColor));
 
+	const bool bUVShell = static_cast<bool>(Data.BakeSettings.SourceBakeMapTypes & EBakeMapType::UVShell);
+	Attributes.Add(FAnalyticsEventAttribute(TEXT("Settings.UVShell.Enabled"), bUVShell));
+
 	FEngineAnalytics::GetProvider().RecordEvent(FString(TEXT("Editor.Usage.MeshModelingMode.")) + EventName, Attributes);
 
 	constexpr bool bLogAnalytics = false; 
diff --git a/Engine/Plugins/Experimental/MeshModelingToolsetExp/Source/MeshModelingToolsExp/Private/BakeMultiMeshAttributeMapsTool.cpp b/Engine/Plugins/Experimental/MeshModelingToolsetExp/Source/MeshModelingToolsExp/Private/BakeMultiMeshAttributeMapsTool.cpp
--- a/Engine/Plugins/Experimental/MeshModelingToolsetExp/Source/MeshModelingToolsExp/Private/BakeMultiMeshAttributeMapsTool.cpp
+++ b/Engine/Plugins/Experimental/MeshModelingToolsetExp/Source/MeshModelingToolsExp/Private/BakeMultiMeshAttributeMapsTool.cpp
@@ -224,6 +224,12 @@ public:
 		return FVector3d::Zero();
 	}
 
+	virtual bool GetTriUVs(const void* Mesh, int TriId, int UVLayer, FVector2f& UV0, FVector2f& UV1, FVector2f& UV2) const override
+	{
+		checkSlow(false);
+		return false;
+	}
+
 	virtual int32 GetMaterialID(const void* Mesh, const int TriId) const override
 	{
 		// TODO
diff --git a/Engine/Plugins/Experimental/MeshModelingToolsetExp/Source/MeshModelingToolsExp/Public/BakeMeshAttributeMapsTool.h b/Engine/Plugins/Experimental/MeshModelingToolsetExp/Source/MeshModelingToolsExp/Public/BakeMeshAttributeMapsTool.h
--- a/Engine/Plugins/Experimental/MeshModelingToolsetExp/Source/MeshModelingToolsExp/Public/BakeMeshAttributeMapsTool.h
+++ b/Engine/Plugins/Experimental/MeshModelingToolsetExp/Source/MeshModelingToolsExp/Public/BakeMeshAttributeMapsTool.h
@@ -54,7 +54,7 @@ class MESHMODELINGTOOLSEXP_API UBakeMeshAttributeMapsToolProperties : public UIn
 public:
 	/** The bake output types to generate */
 	UPROPERTY(EditAnywhere, Category = BakeOutput, meta = (DisplayName = "Output Types", Bitmask, BitmaskEnum = "/Script/MeshModelingToolsExp.EBakeMapType",
-		ValidEnumValues="TangentSpaceNormal, AmbientOcclusion, BentNormal, Curvature, Texture, ObjectSpaceNormal, FaceNormal, Position, MaterialID, PolyGroupID, MultiTexture, VertexColor"))
+		ValidEnumValues="TangentSpaceNormal, AmbientOcclusion, BentNormal, Curvature, Texture, ObjectSpaceNormal, FaceNormal, Position, MaterialID, PolyGroupID, MultiTexture, VertexColor, UVShell"))
 	int32 MapTypes = static_cast<int32>(EBakeMapType::None);
 
 	/** The baked output type used for preview in the viewport */
@@ -123,6 +123,9 @@ protected:
 	UPROPERTY()
 	TObjectPtr<UBakeMeshAttributeMapsResultToolProperties> ResultSettings;
 
+	UPROPERTY()
+	TObjectPtr<UBakeUVShellMapToolProperties> UVShellSettings;
+
 
 	// Begin UBakeMeshAttributeMapsToolBase interface
 	virtual void UpdateResult() override;
@@ -155,6 +158,9 @@ protected:
 	TSharedPtr<UE::Geometry::TImageBuilder<FVector4f>, ESPMode::ThreadSafe> CachedDetailNormalMap;
 	EBakeOpState UpdateResult_DetailNormalMap();
 
+	FUVShellMapSettings CachedUVShellMapSettings;
+	EBakeOpState UpdateResult_UVShellMap(const FImageDimensions& Dimensions);
+
 	void SetSourceObjectVisible(bool bState)
 	{
 		if (!bIsBakeToSelf)
diff --git a/Engine/Plugins/Experimental/MeshModelingToolsetExp/Source/MeshModelingToolsExp/Public/BakeMeshAttributeTool.h b/Engine/Plugins/Experimental/MeshModelingToolsetExp/Source/MeshModelingToolsExp/Public/BakeMeshAttributeTool.h
--- a/Engine/Plugins/Experimental/MeshModelingToolsetExp/Source/MeshModelingToolsExp/Public/BakeMeshAttributeTool.h
+++ b/Engine/Plugins/Experimental/MeshModelingToolsetExp/Source/MeshModelingToolsExp/Public/BakeMeshAttributeTool.h
@@ -53,8 +53,10 @@ enum class EBakeMapType
 	One                    = 1 << 12,
 	/* Constant value of Zero */
 	Zero                   = 1 << 13,
+	/* UV shell */
+	UVShell                = 1 << 14 UMETA(DisplayName = "UV Shell"),
 
-	All                    = 0x1FFF UMETA(Hidden)
+	All                    = 0x3FFF UMETA(Hidden)
 };
 ENUM_CLASS_FLAGS(EBakeMapType);
 
@@ -76,7 +78,8 @@ static constexpr EBakeMapType ENUM_EBAKEMAPTYPE_ALL[] =
 	EBakeMapType::MultiTexture,
 	EBakeMapType::VertexColor,
 	EBakeMapType::MaterialID,
-	EBakeMapType::PolyGroupID
+	EBakeMapType::PolyGroupID,
+	EBakeMapType::UVShell
 };
 
 
diff --git a/Engine/Plugins/Experimental/MeshModelingToolsetExp/Source/MeshModelingToolsExp/Public/BakeMeshAttributeToolCommon.h b/Engine/Plugins/Experimental/MeshModelingToolsetExp/Source/MeshModelingToolsExp/Public/BakeMeshAttributeToolCommon.h
--- a/Engine/Plugins/Experimental/MeshModelingToolsetExp/Source/MeshModelingToolsExp/Public/BakeMeshAttributeToolCommon.h
+++ b/Engine/Plugins/Experimental/MeshModelingToolsetExp/Source/MeshModelingToolsExp/Public/BakeMeshAttributeToolCommon.h
@@ -234,6 +234,33 @@ public:
 };
 
 
+UCLASS()
+class MESHMODELINGTOOLSEXP_API UBakeUVShellMapToolProperties : public UInteractiveToolPropertySet
+{
+	GENERATED_BODY()
+public:
+	/** The source mesh UV Layer to sample. */
+	UPROPERTY(EditAnywhere, Category = "UV Shell Output")
+	int UVLayer = 0;
+
+	/** The thickness of the wireframe in pixels. */
+	UPROPERTY(EditAnywhere, Category = "UV Shell Output")
+	float WireframeThickness = 1.0f;
+
+	/** The color of wireframe pixels. */
+	UPROPERTY(EditAnywhere, Category = "UV Shell Output")
+	FLinearColor WireframeColor = FLinearColor::Blue;
+
+	/** The color of the UV shell interior pixels. */
+	UPROPERTY(EditAnywhere, Category = "UV Shell Output")
+	FLinearColor ShellColor = FLinearColor::Gray;
+
+	/** The color of pixels external to UV shells. */
+	UPROPERTY(EditAnywhere, Category = "UV Shell Output")
+	FLinearColor BackgroundColor = FLinearColor(0.0f, 0.0f, 0.0f, 0.0f);
+};
+
+
 UCLASS()
 class MESHMODELINGTOOLSEXP_API UBakeTexture2DProperties : public UInteractiveToolPropertySet
 {
@@ -361,8 +388,8 @@ struct FCurvatureMapSettings
 {
 	FImageDimensions Dimensions;
 	int32 CurvatureType = 0;
-	float RangeMultiplier = 1.0;
-	float MinRangeMultiplier = 0.0;
+	float RangeMultiplier = 1.0f;
+	float MinRangeMultiplier = 0.0f;
 	int32 ColorMode = 0;
 	int32 ClampMode = 0;
 
@@ -382,6 +409,26 @@ struct FMeshPropertyMapSettings
 	}
 };
 
+struct FUVShellMapSettings
+{
+	FImageDimensions Dimensions;
+	int UVLayer = 0;
+	float WireframeThickness = 1.0f;
+	FLinearColor WireframeColor = FLinearColor::Blue;
+	FLinearColor ShellColor = FLinearColor::Gray;
+	FLinearColor BackgroundColor = FLinearColor(0.0f, 0.0f, 0.0f, 0.0f);
+
+	bool operator==(const FUVShellMapSettings& Other) const
+	{
+		return Dimensions == Other.Dimensions &&
+			UVLayer == Other.UVLayer &&
+			WireframeThickness == Other.WireframeThickness &&
+			WireframeColor == Other.WireframeColor &&
+			ShellColor == Other.ShellColor &&
+			BackgroundColor == Other.BackgroundColor;
+	}
+};
+
 struct FTexture2DSettings
 {
 	FImageDimensions Dimensions;
diff --git a/Engine/Plugins/Runtime/GeometryProcessing/Source/DynamicMesh/Private/Sampling/MeshUVShellMapEvaluator.cpp b/Engine/Plugins/Runtime/GeometryProcessing/Source/DynamicMesh/Private/Sampling/MeshUVShellMapEvaluator.cpp
new file mode 100644
--- /dev/null
+++ b/Engine/Plugins/Runtime/GeometryProcessing/Source/DynamicMesh/Private/Sampling/MeshUVShellMapEvaluator.cpp
@@ -0,0 +1,92 @@
+// Copyright Epic Games, Inc. All Rights Reserved.
+
+#include "Sampling/MeshUVShellMapEvaluator.h"
+#include "Sampling/MeshMapBaker.h"
+
+using namespace UE::Geometry;
+
+void FMeshUVShellMapEvaluator::Setup(const FMeshBaseBaker& Baker, FEvaluationContext& Context)
+{
+	// Cache data from the baker
+	Context.Evaluate = &EvaluateSample;
+	Context.EvaluateDefault = &EvaluateDefault;
+	Context.EvaluateColor = &EvaluateColor;
+	Context.EvaluateChannel = &EvaluateChannel;
+	Context.EvalData = this;
+	Context.AccumulateMode = EAccumulateMode::Add;
+	Context.DataLayout = DataLayout();
+	
+	DetailSampler = Baker.GetDetailSampler();
+}
+
+const TArray<FMeshMapEvaluator::EComponents>& FMeshUVShellMapEvaluator::DataLayout() const
+{
+	static const TArray<EComponents> Layout{ EComponents::Float4 };
+	return Layout;
+}
+
+void FMeshUVShellMapEvaluator::EvaluateSample(float*& Out, const FCorrespondenceSample& Sample, void* EvalData)
+{
+	const FMeshUVShellMapEvaluator* Eval = static_cast<FMeshUVShellMapEvaluator*>(EvalData);
+	const FVector4f SampleResult = Eval->SampleFunction(Sample);
+	WriteToBuffer(Out, SampleResult);
+}
+
+void FMeshUVShellMapEvaluator::EvaluateDefault(float*& Out, void* EvalData)
+{
+	const FMeshUVShellMapEvaluator* Eval = static_cast<FMeshUVShellMapEvaluator*>(EvalData);
+	WriteToBuffer(Out, Eval->BackgroundColor);
+}
+
+void FMeshUVShellMapEvaluator::EvaluateColor(const int DataIdx, float*& In, FVector4f& Out, void* EvalData)
+{
+	Out = FVector4f(In[0], In[1], In[2], In[3]);
+	In += 4;
+}
+
+void FMeshUVShellMapEvaluator::EvaluateChannel(const int DataIdx, float*& In, float& Out, void* EvalData)
+{
+	ensure(false);		// Should not be able to select per-channel evaluation for multi-dimensional properties
+	Out = In[0];
+	In += 1;
+}
+
+FVector4f FMeshUVShellMapEvaluator::SampleFunction(const FCorrespondenceSample& SampleData) const
+{
+	const void* DetailMesh = SampleData.DetailMesh;
+	const FVector3d& DetailBaryCoords = SampleData.DetailBaryCoords;
+	const int32 DetailTriID = SampleData.DetailTriID;
+	FVector4f Color = BackgroundColor;
+
+	if (DetailSampler->IsTriangle(DetailMesh, DetailTriID))
+	{
+		// Compute barycentric/areal coords in UV space.
+		FVector2f UV = FVector2f::Zero();
+		DetailSampler->TriBaryInterpolateUV(DetailMesh, DetailTriID, DetailBaryCoords, UVLayer, UV);
+
+		FVector2f UV0, UV1, UV2;
+		DetailSampler->GetTriUVs(DetailMesh, DetailTriID, UVLayer, UV0, UV1, UV2);
+
+		const FTriangle2f UVTri(UV0, UV1, UV2);
+		const FVector3f UVArealCoords = FVector3f(DetailBaryCoords) * UVTri.Area();
+
+		// Barycentric --> Trilinear to identify min distance to tri edges.
+		const float M20 = (UV2-UV0).Length();
+		const float M01 = (UV0-UV1).Length();
+		const float M12 = (UV1-UV2).Length();
+
+		const FVector3f Trilinear(
+			2.0f * UVArealCoords.X / M20,
+			2.0f * UVArealCoords.Y / M01,
+			2.0f * UVArealCoords.Z / M12);
+		const float MinTrilinear = FMath::Min3(Trilinear.X, Trilinear.Y, Trilinear.Z);
+
+		// Check edge distance in texel space (assumes square texels).
+		const float MinTexelDist = MinTrilinear / static_cast<float>(TexelSize.X);
+		Color = MinTexelDist < WireframeThickness ? WireframeColor : ShellColor;
+	}
+
+	return Color;
+}
+
+
diff --git a/Engine/Plugins/Runtime/GeometryProcessing/Source/DynamicMesh/Public/Sampling/MeshBakerCommon.h b/Engine/Plugins/Runtime/GeometryProcessing/Source/DynamicMesh/Public/Sampling/MeshBakerCommon.h
--- a/Engine/Plugins/Runtime/GeometryProcessing/Source/DynamicMesh/Public/Sampling/MeshBakerCommon.h
+++ b/Engine/Plugins/Runtime/GeometryProcessing/Source/DynamicMesh/Public/Sampling/MeshBakerCommon.h
@@ -230,6 +230,17 @@ public:
 	 */
 	virtual FVector3d GetTriNormal(const void* Mesh, const int TriId) const = 0;
 
+	/**
+	 * @param Mesh pointer to mesh to query 
+	 * @param TriId the triangle index to test
+	 * @param UVLayer the UVLayer index to query
+	 * @param UV0 the output UV for triangle vertex 0
+	 * @param UV1 the output UV for triangle vertex 1
+	 * @param UV2 the output UV for triangle vertex 2
+	 * @return true if a valid UV was returned, false otherwise.
+	 */
+	virtual bool GetTriUVs(const void* Mesh, int TriId, int UVLayer, FVector2f& UV0, FVector2f& UV1, FVector2f& UV2) const = 0;
+
 	/**
 	 * @param Mesh pointer to mesh to query
 	 * @param TriId the triangle index to test
@@ -473,6 +484,23 @@ public:
 		return DynamicMesh->GetTriNormal(TriId);
 	}
 
+	virtual bool GetTriUVs(const void* Mesh, int TriId, int UVLayer, FVector2f& UV0, FVector2f& UV1, FVector2f& UV2) const override
+	{
+		const FDynamicMesh3* DynamicMesh = static_cast<const FDynamicMesh3*>(Mesh);
+		if (const FDynamicMeshAttributeSet* Attributes = DynamicMesh->Attributes())
+		{
+			if (const FDynamicMeshUVOverlay* UVOverlay = Attributes->GetUVLayer(UVLayer))
+			{
+				if (UVOverlay->IsSetTriangle(TriId))
+				{
+					UVOverlay->GetTriElements(TriId, UV0, UV1, UV2);
+					return true;
+				}
+			}
+		}
+		return false;
+	}
+
 	virtual int32 GetMaterialID(const void* Mesh, const int TriId) const override
 	{
 		const FDynamicMesh3* DynamicMesh = static_cast<const FDynamicMesh3*>(Mesh);
diff --git a/Engine/Plugins/Runtime/GeometryProcessing/Source/DynamicMesh/Public/Sampling/MeshMapEvaluator.h b/Engine/Plugins/Runtime/GeometryProcessing/Source/DynamicMesh/Public/Sampling/MeshMapEvaluator.h
--- a/Engine/Plugins/Runtime/GeometryProcessing/Source/DynamicMesh/Public/Sampling/MeshMapEvaluator.h
+++ b/Engine/Plugins/Runtime/GeometryProcessing/Source/DynamicMesh/Public/Sampling/MeshMapEvaluator.h
@@ -20,7 +20,8 @@ enum class EMeshMapEvaluatorType
 	Property,
 	ResampleImage,
 	MultiResampleImage,
-	RenderCapture
+	RenderCapture,
+	UVShell
 };
 
 class FMeshMapEvaluator
diff --git a/Engine/Plugins/Runtime/GeometryProcessing/Source/DynamicMesh/Public/Sampling/MeshUVShellMapEvaluator.h b/Engine/Plugins/Runtime/GeometryProcessing/Source/DynamicMesh/Public/Sampling/MeshUVShellMapEvaluator.h
new file mode 100644
--- /dev/null
+++ b/Engine/Plugins/Runtime/GeometryProcessing/Source/DynamicMesh/Public/Sampling/MeshUVShellMapEvaluator.h
@@ -0,0 +1,58 @@
+// Copyright Epic Games, Inc. All Rights Reserved.
+
+#pragma once
+
+#include "Sampling/MeshMapEvaluator.h"
+#include "Sampling/MeshBakerCommon.h"
+
+namespace UE::Geometry
+{
+
+class IMeshBakerDetailSampler;	
+
+/**
+ * A mesh evaluator for mesh properties as color data.
+ */
+class DYNAMICMESH_API FMeshUVShellMapEvaluator : public FMeshMapEvaluator
+{
+public:
+	FVector2d TexelSize = FVector2d::One();
+	float WireframeThickness = 1.0f;
+	FVector4f WireframeColor = FLinearColor::Blue;
+	FVector4f ShellColor = FLinearColor::Gray;
+	FVector4f BackgroundColor = FVector4f(0.0f, 0.0f, 0.0f, 0.0f);
+
+	int UVLayer = 0;
+
+public:
+	PRAGMA_DISABLE_DEPRECATION_WARNINGS
+	FMeshUVShellMapEvaluator() = default;
+	FMeshUVShellMapEvaluator(const FMeshUVShellMapEvaluator&) = default;
+	FMeshUVShellMapEvaluator& operator=(const FMeshUVShellMapEvaluator&) = default;
+	PRAGMA_ENABLE_DEPRECATION_WARNINGS
+
+	// Begin FMeshMapEvaluator interface
+	virtual void Setup(const FMeshBaseBaker& Baker, FEvaluationContext& Context) override;
+
+	virtual const TArray<EComponents>& DataLayout() const override;
+
+	virtual EMeshMapEvaluatorType Type() const override { return EMeshMapEvaluatorType::UVShell; }
+	// End FMeshMapEvaluator interface
+
+	static void EvaluateSample(float*& Out, const FCorrespondenceSample& Sample, void* EvalData);
+
+	static void EvaluateDefault(float*& Out, void* EvalData);
+
+	static void EvaluateColor(const int DataIdx, float*& In, FVector4f& Out, void* EvalData);
+	static void EvaluateChannel(const int DataIdx, float*& In, float& Out, void* EvalData);
+
+protected:
+	// Cached data
+	const IMeshBakerDetailSampler* DetailSampler = nullptr;
+
+private:
+	FVector4f SampleFunction(const FCorrespondenceSample& SampleData) const;
+};
+
+} // end namespace UE::Geometry
+
diff --git a/Engine/Source/Runtime/GeometryCore/Public/VectorUtil.h b/Engine/Source/Runtime/GeometryCore/Public/VectorUtil.h
--- a/Engine/Source/Runtime/GeometryCore/Public/VectorUtil.h
+++ b/Engine/Source/Runtime/GeometryCore/Public/VectorUtil.h
@@ -409,10 +409,10 @@ namespace VectorUtil
 		RealType fR0 = kV02.Dot(kPV2);
 		RealType fR1 = kV12.Dot(kPV2);
 		RealType fDet = fM00 * fM11 - fM01 * fM01;
-		RealType fInvDet = 1.0 / fDet;
+		RealType fInvDet = (RealType)1.0 / fDet;
 		RealType fBary1 = (fM11 * fR0 - fM01 * fR1) * fInvDet;
 		RealType fBary2 = (fM00 * fR1 - fM01 * fR0) * fInvDet;
-		RealType fBary3 = 1.0 - fBary1 - fBary2;
+		RealType fBary3 = (RealType)1.0 - fBary1 - fBary2;
 		return TVector<RealType>(fBary1, fBary2, fBary3);
 	}
 
