Adding support for mirroring weights on isolated selections in skin weight tool.
Added support for undo'ing weight editing that was done in an isolated selection.
Added Copy/Paste weights feature.
Disabled weight editing UI until vertice(s) are selected.
Fixed bug, weight edits in an isolated selection do not persist if Accepted while isolated

diff --git a/Engine/Plugins/Experimental/Animation/SkeletalMeshModelingTools/Source/SkeletalMeshModelingTools/Private/DetailCustomization/SkinWeightDetailCustomization.cpp b/Engine/Plugins/Experimental/Animation/SkeletalMeshModelingTools/Source/SkeletalMeshModelingTools/Private/DetailCustomization/SkinWeightDetailCustomization.cpp
--- a/Engine/Plugins/Experimental/Animation/SkeletalMeshModelingTools/Source/SkeletalMeshModelingTools/Private/DetailCustomization/SkinWeightDetailCustomization.cpp
+++ b/Engine/Plugins/Experimental/Animation/SkeletalMeshModelingTools/Source/SkeletalMeshModelingTools/Private/DetailCustomization/SkinWeightDetailCustomization.cpp
@@ -452,6 +452,7 @@ void FSkinWeightDetailCustomization::AddSelectionUI(IDetailLayoutBuilder& Detail
 			.Padding(2.f, WeightEditVerticalPadding)
 			[
 				SNew(SButton)
+				.IsEnabled_Lambda([this]{ return ToolSettings->WeightTool->HasSelectedVertices(); })
 				.HAlign(HAlign_Center)
 				.VAlign(VAlign_Center)
 				.Text(LOCTEXT("GrowSelectionButtonLabel", "Grow"))
@@ -468,6 +469,7 @@ void FSkinWeightDetailCustomization::AddSelectionUI(IDetailLayoutBuilder& Detail
 			.Padding(2.f, WeightEditVerticalPadding)
 			[
 				SNew(SButton)
+				.IsEnabled_Lambda([this]{ return ToolSettings->WeightTool->HasSelectedVertices(); })
 				.HAlign(HAlign_Center)
 				.VAlign(VAlign_Center)
 				.Text(LOCTEXT("ShrinkSelectionButtonLabel", "Shrink"))
@@ -484,6 +486,7 @@ void FSkinWeightDetailCustomization::AddSelectionUI(IDetailLayoutBuilder& Detail
 			.Padding(2.f, WeightEditVerticalPadding)
 			[
 				SNew(SButton)
+				.IsEnabled_Lambda([this]{ return ToolSettings->WeightTool->HasSelectedVertices(); })
 				.HAlign(HAlign_Center)
 				.VAlign(VAlign_Center)
 				.Text(LOCTEXT("FloodSelectionButtonLabel", "Flood"))
@@ -522,6 +525,7 @@ void FSkinWeightDetailCustomization::AddSelectionUI(IDetailLayoutBuilder& Detail
 			.Padding(2.f, WeightEditVerticalPadding)
 			[
 				SNew(SButton)
+				.IsEnabled_Lambda([this]{ return ToolSettings->WeightTool->HasSelectedVertices(); })
 				.HAlign(HAlign_Center)
 				.VAlign(VAlign_Center)
 				.Text(LOCTEXT("BorderSelectionButtonLabel", "Convert to Border"))
@@ -560,6 +564,7 @@ void FSkinWeightDetailCustomization::AddSelectionUI(IDetailLayoutBuilder& Detail
 	.WholeRowContent()
 	[
 		SNew(SHorizontalBox)
+		.IsEnabled_Lambda([this]{ return ToolSettings->WeightTool->HasSelectedVertices(); })
 
 		+SHorizontalBox::Slot()
 		[
@@ -694,6 +699,7 @@ void FSkinWeightDetailCustomization::AddSelectionUI(IDetailLayoutBuilder& Detail
 	.WholeRowContent()
 	[
 		SNew(SVerticalBox)
+		.IsEnabled_Lambda([this]{ return ToolSettings->WeightTool->HasSelectedVertices(); })
 
 		+ SVerticalBox::Slot()
 		.Padding(WeightEditHorizontalPadding, WeightEditVerticalPadding)
@@ -751,6 +757,7 @@ void FSkinWeightDetailCustomization::AddSelectionUI(IDetailLayoutBuilder& Detail
 	.WholeRowContent()
 	[
 		SNew(SVerticalBox)
+		.IsEnabled_Lambda([this]{ return ToolSettings->WeightTool->HasSelectedVertices(); })
 
 		+ SVerticalBox::Slot()
 		.Padding(WeightEditHorizontalPadding, WeightEditVerticalPadding)
@@ -809,6 +816,7 @@ void FSkinWeightDetailCustomization::AddSelectionUI(IDetailLayoutBuilder& Detail
 	.WholeRowContent()
 	[
 		SNew(SVerticalBox)
+		.IsEnabled_Lambda([this]{ return ToolSettings->WeightTool->HasSelectedVertices(); })
 
 		+ SVerticalBox::Slot()
 		.Padding(WeightEditHorizontalPadding, WeightEditVerticalPadding)
@@ -859,6 +867,7 @@ void FSkinWeightDetailCustomization::AddSelectionUI(IDetailLayoutBuilder& Detail
 	.WholeRowContent()
 	[
 		SNew(SVerticalBox)
+		.IsEnabled_Lambda([this]{ return ToolSettings->WeightTool->HasSelectedVertices(); })
 
 		+ SVerticalBox::Slot()
 		.Padding(WeightEditHorizontalPadding, WeightEditVerticalPadding)
@@ -924,6 +933,7 @@ void FSkinWeightDetailCustomization::AddSelectionUI(IDetailLayoutBuilder& Detail
 	.WholeRowContent()
 	[
 		SNew(SVerticalBox)
+		.IsEnabled_Lambda([this]{ return ToolSettings->WeightTool->HasSelectedVertices(); })
 		
 		+ SVerticalBox::Slot()
 		.Padding(WeightEditHorizontalPadding, WeightEditVerticalPadding)
@@ -1012,6 +1022,7 @@ void FSkinWeightDetailCustomization::AddSelectionUI(IDetailLayoutBuilder& Detail
 	.WholeRowContent()
 	[
 		SNew(SHorizontalBox)
+		.IsEnabled_Lambda([this]{ return ToolSettings->WeightTool->HasSelectedVertices(); })
 
 		+SHorizontalBox::Slot()
 		.Padding(2.f, WeightEditVerticalPadding)
@@ -1060,52 +1071,89 @@ void FSkinWeightDetailCustomization::AddSelectionUI(IDetailLayoutBuilder& Detail
 	EditWeightsCategory.AddCustomRow(LOCTEXT("PruneWeightsRow", "Prune"), false)
 	.WholeRowContent()
 	[
-		SNew(SVerticalBox)
+		SNew(SHorizontalBox)
+		.IsEnabled_Lambda([this]{ return ToolSettings->WeightTool->HasSelectedVertices(); })
 
-		+ SVerticalBox::Slot()
-		.Padding(WeightEditHorizontalPadding, WeightEditVerticalPadding)
+		+SHorizontalBox::Slot()
+		.Padding(2.f, WeightEditVerticalPadding)
 		[
-			SNew(SHorizontalBox)
+			SNew(SButton)
+			.HAlign(HAlign_Center)
+			.Text(LOCTEXT("PruneWeightsButtonLabel", "Prune"))
+			.ToolTipText(LOCTEXT("PruneButtonTooltip",
+				"Removes influences with weights below the given threshold value.\n"
+				"Pruned bones are removed from the list of bones affecting the given vertex.\n"
+				"Pruned bones will no longer recieve weight when a vertex is normalized.\n"
+				"This command operates on the selected vertices."))
+			.OnClicked_Lambda([this]()
+			{
+				ToolSettings->WeightTool->PruneWeights(ToolSettings->PruneValue, TArray<BoneIndex>());
+				return FReply::Handled();
+			})
+		]
 
-			+SHorizontalBox::Slot()
-			[
-				SNew(SBox)
-				[
-					SNew(SButton)
-					.HAlign(HAlign_Center)
-					.Text(LOCTEXT("PruneWeightsButtonLabel", "Prune"))
-					.ToolTipText(LOCTEXT("PruneButtonTooltip",
-						"Removes influences with weights below the given threshold value.\n"
-						"Pruned bones are removed from the list of bones affecting the given vertex.\n"
-						"Pruned bones will no longer recieve weight when a vertex is normalized.\n"
-						"This command operates on the selected vertices."))
-					.OnClicked_Lambda([this]()
-					{
-						ToolSettings->WeightTool->PruneWeights(ToolSettings->PruneValue, TArray<BoneIndex>());
-						return FReply::Handled();
-					})
-				]
-			]
+		+SHorizontalBox::Slot()
+		.Padding(2.f, WeightEditVerticalPadding)
+		[
+			SNew(SSpinBox<float>)
+			.MinValue(0.f)
+			.MaxValue(1.f)
+			.Value_Lambda([this]()
+			{
+				return ToolSettings->PruneValue;
+			})
+			.OnValueChanged_Lambda([this](float NewValue)
+			{
+				ToolSettings->PruneValue = NewValue;
+			})
+			.OnValueCommitted_Lambda([this](float NewValue, ETextCommit::Type CommitType)
+			{
+				ToolSettings->SaveConfig();
+			})
+		]
+	];
 
-			+SHorizontalBox::Slot()
+	// COPY/PASTE WEIGHTS category
+	EditWeightsCategory.AddCustomRow(LOCTEXT("CopyPasteWeightsRow", "Copy Paste"), false)
+	.WholeRowContent()
+	[
+		SNew(SHorizontalBox)
+		.IsEnabled_Lambda([this]{ return ToolSettings->WeightTool->HasSelectedVertices(); })
+
+		+SHorizontalBox::Slot()
+		.Padding(2.f, WeightEditVerticalPadding)
+		[
+			SNew(SBox)
 			[
-				SNew(SSpinBox<float>)
-				.MinValue(0.f)
-				.MaxValue(1.f)
-				.Value_Lambda([this]()
-				{
-					return ToolSettings->PruneValue;
-				})
-				.OnValueChanged_Lambda([this](float NewValue)
-				{
-					ToolSettings->PruneValue = NewValue;
-				})
-				.OnValueCommitted_Lambda([this](float NewValue, ETextCommit::Type CommitType)
+				SNew(SButton)
+				.HAlign(HAlign_Center)
+				.Text(LOCTEXT("CopyWeightsButtonLabel", "Copy"))
+				.ToolTipText(LOCTEXT("CopyButtonTooltip",
+					"Copy the average weights of the selected vertices to the clipboard. \n"
+					"This is designed to work with the Paste command."))
+				.OnClicked_Lambda([this]()
 				{
-					ToolSettings->SaveConfig();
+					ToolSettings->WeightTool->CopyWeights();
+					return FReply::Handled();
 				})
 			]
 		]
+
+		+SHorizontalBox::Slot()
+		.Padding(2.f, WeightEditVerticalPadding)
+		[
+			SNew(SButton)
+			.HAlign(HAlign_Center)
+			.Text(LOCTEXT("PasteWeightsButtonLabel", "Paste"))
+			.ToolTipText(LOCTEXT("PasteButtonTooltip",
+				"Paste the weights on the selected vertices.\n"
+				"This command requires the clipboard contain weights from the Copy command."))
+			.OnClicked_Lambda([this]()
+			{
+				ToolSettings->WeightTool->PasteWeights();
+				return FReply::Handled();
+			})
+		]
 	];
 }
 
diff --git a/Engine/Plugins/Experimental/MeshModelingToolsetExp/Source/MeshModelingToolsEditorOnlyExp/MeshModelingToolsEditorOnlyExp.Build.cs b/Engine/Plugins/Experimental/MeshModelingToolsetExp/Source/MeshModelingToolsEditorOnlyExp/MeshModelingToolsEditorOnlyExp.Build.cs
--- a/Engine/Plugins/Experimental/MeshModelingToolsetExp/Source/MeshModelingToolsEditorOnlyExp/MeshModelingToolsEditorOnlyExp.Build.cs
+++ b/Engine/Plugins/Experimental/MeshModelingToolsetExp/Source/MeshModelingToolsEditorOnlyExp/MeshModelingToolsEditorOnlyExp.Build.cs
@@ -57,7 +57,9 @@ public class MeshModelingToolsEditorOnlyExp : ModuleRules
 				"UnrealEd",
 				"Persona",
 				"AdvancedPreviewScene",
-				"Slate"
+				"Slate",
+				"Json",
+				"ApplicationCore"
 				// ... add private dependencies that you statically link with here ...	
 			}
 		);
diff --git a/Engine/Plugins/Experimental/MeshModelingToolsetExp/Source/MeshModelingToolsEditorOnlyExp/Private/SkeletalMesh/SkinWeightsPaintTool.cpp b/Engine/Plugins/Experimental/MeshModelingToolsetExp/Source/MeshModelingToolsEditorOnlyExp/Private/SkeletalMesh/SkinWeightsPaintTool.cpp
--- a/Engine/Plugins/Experimental/MeshModelingToolsetExp/Source/MeshModelingToolsEditorOnlyExp/Private/SkeletalMesh/SkinWeightsPaintTool.cpp
+++ b/Engine/Plugins/Experimental/MeshModelingToolsetExp/Source/MeshModelingToolsEditorOnlyExp/Private/SkeletalMesh/SkinWeightsPaintTool.cpp
@@ -3,49 +3,58 @@
 #include "SkeletalMesh/SkinWeightsPaintTool.h"
 
 #include "AssetViewerSettings.h"
-#include "Engine/SkeletalMesh.h"
-#include "InteractiveToolManager.h"
-#include "SkeletalMeshAttributes.h"
-#include "Math/UnrealMathUtility.h"
-#include "Components/SkeletalMeshComponent.h"
 
-#include "TargetInterfaces/PrimitiveComponentBackedTarget.h"
-#include "ModelingToolTargetUtil.h"
+#include "Math/UnrealMathUtility.h"
+#include "Async/Async.h"
 
 #include "Materials/Material.h"
 #include "Materials/MaterialInstanceDynamic.h"
-#include "MeshDescription.h"
-#include "MeshModelingToolsEditorOnlyExp.h"
-#include "PointSetAdapter.h"
+
+#include "Engine/SkeletalMesh.h"
+#include "Components/SkeletalMeshComponent.h"
+#include "SkeletalDebugRendering.h"
+#include "AnimationRuntime.h"
 #include "Animation/MirrorDataTable.h"
-#include "DynamicMesh/NonManifoldMappingSupport.h"
-#include "Parameterization/MeshLocalParam.h"
-#include "Spatial/FastWinding.h"
-#include "Async/Async.h"
+
+#include "Editor/Persona/Public/IPersonaEditorModeManager.h"
+#include "Editor/Persona/Public/PersonaModule.h"
+#include "Preferences/PersonaOptions.h"
+#include "PreviewProfileController.h"
 #include "BaseGizmos/BrushStampIndicator.h"
-#include "DynamicMesh/MeshAdapterUtil.h"
+
+#include "InteractiveToolManager.h"
+#include "ToolSetupUtil.h"
+#include "ToolTargetManager.h"
+#include "TargetInterfaces/PrimitiveComponentBackedTarget.h"
+#include "ModelingToolTargetUtil.h"
 #include "Selection/PolygonSelectionMechanic.h"
-#include "Spatial/PointSetHashTable.h"
-#include "Operations/SmoothBoneWeights.h"
 #include "ContextObjectStore.h"
+
+#include "MeshDescription.h"
+#include "MeshModelingToolsEditorOnlyExp.h"
 #include "DynamicMeshToMeshDescription.h"
 #include "DynamicSubmesh3.h"
 #include "MeshDescriptionToDynamicMesh.h"
-#include "SkeletalDebugRendering.h"
-#include "ToolSetupUtil.h"
-#include "ToolTargetManager.h"
-#include "Editor/Persona/Public/IPersonaEditorModeManager.h"
-#include "Editor/Persona/Public/PersonaModule.h"
-#include "Preferences/PersonaOptions.h"
-#include "PreviewProfileController.h"
-#include "Animation/SkinWeightProfile.h"
-#include "AnimationRuntime.h"
 #include "DynamicMesh/DynamicVertexSkinWeightsAttribute.h"
+#include "DynamicMesh/MeshAdapterUtil.h"
+#include "DynamicMesh/NonManifoldMappingSupport.h"
 #include "TargetInterfaces/DynamicMeshCommitter.h"
 #include "TargetInterfaces/MeshDescriptionCommitter.h"
 #include "TargetInterfaces/MeshDescriptionProvider.h"
+#include "SkeletalMeshAttributes.h"
+
+#include "Operations/SmoothBoneWeights.h"
 #include "Operations/TransferBoneWeights.h"
+
 #include "Parameterization/MeshDijkstra.h"
+#include "Spatial/PointHashGrid3.h"
+#include "Parameterization/MeshLocalParam.h"
+#include "Spatial/FastWinding.h"
+
+#include "Widgets/Notifications/SNotificationList.h"
+#include "Windows/WindowsPlatformApplicationMisc.h"
+
+#include "Serialization/JsonSerializable.h"
 
 #include UE_INLINE_GENERATED_CPP_BY_NAME(SkinWeightsPaintTool)
 
@@ -53,6 +62,30 @@
 
 using namespace SkinPaintTool;
 
+static void ShowEditorMessage(ELogVerbosity::Type InMessageType, const FText& InMessage)
+{
+	FNotificationInfo Notification(InMessage);
+	Notification.bUseSuccessFailIcons = true;
+	Notification.ExpireDuration = 5.0f;
+
+	SNotificationItem::ECompletionState State = SNotificationItem::CS_Success;
+
+	switch(InMessageType)
+	{
+	case ELogVerbosity::Warning:
+		UE_LOG(LogMeshModelingToolsEditor, Warning, TEXT("%s"), *InMessage.ToString());
+		break;
+	case ELogVerbosity::Error:
+		State = SNotificationItem::CS_Fail;
+		UE_LOG(LogMeshModelingToolsEditor, Error, TEXT("%s"), *InMessage.ToString());
+		break;
+	default:
+		break; // don't log anything unless a warning or error
+	}
+	
+	FSlateNotificationManager::Get().AddNotification(Notification)->SetCompletionState(State);
+}
+
 namespace SkinPaintTool
 {
 	
@@ -374,6 +407,8 @@ void FMultiBoneWeightEdits::MergeSingleEdit(
 	const float OldWeight,
 	const float NewWeight)
 {
+	ensure(BoneIndex!=INDEX_NONE);
+	
 	FSingleBoneWeightEdits& BoneWeightEdit = PerBoneWeightEdits.FindOrAdd(BoneIndex);
 	BoneWeightEdit.BoneIndex = BoneIndex;
 	BoneWeightEdit.NewWeights.Add(VertexID, NewWeight);
@@ -382,6 +417,8 @@ void FMultiBoneWeightEdits::MergeSingleEdit(
 
 void FMultiBoneWeightEdits::MergeEdits(const FSingleBoneWeightEdits& BoneWeightEdits)
 {
+	ensure(BoneWeightEdits.BoneIndex!=INDEX_NONE);
+	
 	// make sure bone has an entry in the map of weight edits
 	const int32 BoneIndex = BoneWeightEdits.BoneIndex;
 	PerBoneWeightEdits.FindOrAdd(BoneIndex);
@@ -406,13 +443,13 @@ float FMultiBoneWeightEdits::GetVertexDeltaFromEdits(const int32 BoneIndex, cons
 	return 0.0f;
 }
 
-void FMultiBoneWeightEdits::GetEditedVertexIndices(TSet<int32>& OutVerticesToEdit) const
+void FMultiBoneWeightEdits::AddEditedVerticesToSet(TSet<int32>& OutEditedVertexSet) const
 {
 	TSet<int32> VerticesInEdit;
 	for (const TTuple<int32, FSingleBoneWeightEdits>& Pair : PerBoneWeightEdits)
 	{
 		Pair.Value.NewWeights.GetKeys(VerticesInEdit);
-		OutVerticesToEdit.Append(VerticesInEdit);
+		OutEditedVertexSet.Append(VerticesInEdit);
 	}
 }
 
@@ -533,7 +570,7 @@ void FSkinToolDeformer::UpdateVertexDeformation(
 			const VertexWeights& VertexPerBoneData = CurrentWeights[VertexID];
 			for (const FVertexBoneWeight& VertexData : VertexPerBoneData)
 			{
-				if (VertexData.BoneID == INDEX_NONE)
+				if (!ensure(VertexData.BoneID != INDEX_NONE))
 				{
 					continue;
 				}
@@ -685,7 +722,7 @@ void FSkinToolWeights::InitializeSkinWeights(
 	}
 }
 
-void FSkinToolWeights::EditVertexWeightAndNormalize(
+void FSkinToolWeights::CreateWeightEditForVertex(
 	const int32 BoneToHoldIndex,
 	const int32 VertexID,
 	float NewWeightValue,
@@ -705,6 +742,11 @@ void FSkinToolWeights::EditVertexWeightAndNormalize(
 		{
 			continue;
 		}
+
+		if (!ensure(VertexBoneData.BoneID != INDEX_NONE))
+		{
+			continue;
+		}
 		
 		RecordedBonesOnVertex.Add(VertexBoneData.BoneID);
 		ValuesToNormalize.Add(VertexBoneData.Weight);
@@ -863,6 +905,10 @@ void FSkinToolWeights::ApplyCurrentWeightsToMeshDescription(FMeshDescription* Me
 		const VertexWeights& VertexWeights = CurrentWeights[VertexIndex];
 		for (const FVertexBoneWeight& SingleBoneWeight : VertexWeights)
 		{
+			if (!ensure(SingleBoneWeight.BoneID != INDEX_NONE))
+			{
+				continue;
+			}
 			SourceBoneWeights.Add(UE::AnimationCore::FBoneWeight(SingleBoneWeight.BoneID, SingleBoneWeight.Weight));
 		}
 
@@ -894,6 +940,11 @@ void FSkinToolWeights::SetWeightOfBoneOnVertex(
 	TArray<VertexWeights>& InOutVertexWeights)
 {
 	Deformer.SetVertexNeedsUpdated(VertexID);
+
+	if (!ensure(BoneIndex != INDEX_NONE))
+	{
+		return;
+	}
 	
 	// incoming weights are assumed to be normalized already, so set it directly
 	VertexWeights& VertexWeights = InOutVertexWeights[VertexID];
@@ -974,9 +1025,10 @@ void FSkinToolWeights::RemoveInfluenceFromVertex(
 	{
 		return CurrentVertexWeight.BoneID == InBoneID;
 	});
+	
 	// can't prune an influence that doesn't exist on a vertex
 	// this may happen if the calling code already pruned the influence to avoid normalization weights
-	if (IndexOfBoneInVertex == INDEX_NONE)
+	if (!ensure(IndexOfBoneInVertex != INDEX_NONE))
 	{
 		return;
 	}
@@ -1027,7 +1079,7 @@ void FSkinToolWeights::AddNewInfluenceToVertex(
 	SingleVertexWeights.Emplace(InBoneID, PosLocalToBone, Weight);
 }
 
-void FSkinToolWeights::SwapAfterChange()
+void FSkinToolWeights::SyncWeightBuffers()
 {
 	PreChangeWeights = CurrentWeights;
 
@@ -1054,6 +1106,7 @@ void FSkinToolWeights::ApplyEditsToCurrentWeights(const FMultiBoneWeightEdits& E
 	{
 		const FSingleBoneWeightEdits& WeightEdits = BoneWeightEdits.Value;
 		const int32 BoneIndex = WeightEdits.BoneIndex;
+		check(BoneIndex != INDEX_NONE);
 		for (const TTuple<int32, float>& NewWeight : WeightEdits.NewWeights)
 		{
 			const int32 VertexID = NewWeight.Key;
@@ -1096,7 +1149,6 @@ BoneIndex FSkinToolWeights::GetParentBoneToWeightTo(BoneIndex ChildBone)
 	if (const USkeletalMesh* SkeletalMesh = Deformer.Component->GetSkeletalMeshAsset())
 	{
 		ParentBoneIndex = SkeletalMesh->GetRefSkeleton().GetParentIndex(ChildBone);
-		
 	}
 
 	// are we at the root? (no parent)
@@ -1115,7 +1167,6 @@ BoneIndex FSkinToolWeights::GetParentBoneToWeightTo(BoneIndex ChildBone)
 				ParentBoneIndex = RootsChildren[0];
 			}
 		}
-		
 	}
 
 	return ParentBoneIndex;
@@ -1156,16 +1207,58 @@ void FMeshSkinWeightsChange::Revert(UObject* Object)
 	Tool->OnWeightsChanged.Broadcast();
 }
 
-void FMeshSkinWeightsChange::AddBoneWeightEdit(const FSingleBoneWeightEdits& BoneWeightEdit)
+void FMeshSkinWeightsChange::StoreBoneWeightEdit(const FSingleBoneWeightEdits& BoneWeightEdit, const TFunction<int32(int32)>& VertexIndexConverter)
 {
+	if (VertexIndexConverter)
+	{
+		// remap vertex indices of old and new weights using the supplied converter function
+		FSingleBoneWeightEdits RemappedBoneWeightEdits = BoneWeightEdit;
+		TMap<VertexIndex, float> RemappedWeights;
+
+		// remap vertex indices in NEW weights
+		for (TPair<VertexIndex, float>& NewWeight : RemappedBoneWeightEdits.NewWeights)
+		{
+			RemappedWeights.Add(VertexIndexConverter(NewWeight.Key), NewWeight.Value);
+		}
+		RemappedBoneWeightEdits.NewWeights = RemappedWeights;
+
+		// remap vertex indices in OLD weights
+		RemappedWeights.Reset();
+		for (TPair<VertexIndex, float>& OldWeight : RemappedBoneWeightEdits.OldWeights)
+		{
+			RemappedWeights.Add(VertexIndexConverter(OldWeight.Key), OldWeight.Value);
+		}
+		RemappedBoneWeightEdits.OldWeights = RemappedWeights;
+
+		// merge the weight edits
+		AllWeightEdits.MergeEdits(RemappedBoneWeightEdits);
+		return;
+	}
+
+	// store weight edits directly	
 	AllWeightEdits.MergeEdits(BoneWeightEdit);
 }
 
-void FMeshSkinWeightsChange::AddPruneBoneEdit(const VertexIndex VertexToPruneFrom, const BoneIndex BoneToPrune)
+void FMeshSkinWeightsChange::StorePruneBoneEdit(const VertexIndex VertexToPruneFrom, const BoneIndex BoneToPrune)
 {
 	AllWeightEdits.PrunedInfluences.Emplace(VertexToPruneFrom, BoneToPrune);
 }
 
+void FMeshSkinWeightsChange::StoreMultipleWeightEdits(const FMultiBoneWeightEdits& WeightEdits, const TFunction<int32(int32)>& VertexIndexConverter)
+{
+	// store weight edits
+	for (const TTuple<BoneIndex, FSingleBoneWeightEdits>& BoneWeightEdits : WeightEdits.PerBoneWeightEdits)
+	{
+		StoreBoneWeightEdit(BoneWeightEdits.Value, VertexIndexConverter);
+	}
+
+	// store pruned influences
+	for (const TPair<VertexIndex, BoneIndex>& PrunedInfluence : WeightEdits.PrunedInfluences)
+	{
+		StorePruneBoneEdit(PrunedInfluence.Key, PrunedInfluence.Value);
+	}
+}
+
 /*
  * ToolBuilder
  */
@@ -1294,7 +1387,6 @@ void USkinWeightsPaintTool::Setup()
 	// notify preview mesh when triangle selection has been updated
 	PolygonSelectionMechanic->OnSelectionChanged.AddWeakLambda(this, [this]()
 	{
-		UpdateSelectedVertices();
 		PreviewMesh->FastNotifySecondaryTrianglesChanged();
 	});
 	PolygonSelectionMechanic->OnFaceSelectionPreviewChanged.AddWeakLambda(this, [this]()
@@ -1372,7 +1464,7 @@ FBox USkinWeightsPaintTool::GetWorldSpaceFocusBox()
 	case EWeightEditMode::Mesh:
 		{
 			FAxisAlignedBox3d Bounds = FAxisAlignedBox3d::Empty();
-			UpdateSelectedVertices();
+			const TArray<int32>& SelectedVertices = GetSelectedVertices();
 			if (!SelectedVertices.IsEmpty())
 			{
 				const FDynamicMesh3* Mesh = PreviewMesh->GetMesh();
@@ -1564,6 +1656,9 @@ void USkinWeightsPaintTool::PostEditMeshInitialization(
 	
 	// update smooth operator (this must be done after PreviewMesh & Weights have been updated)
 	InitializeSmoothWeightsOperator();
+
+	// after any mesh change, the mirror tables will need rebuilt next time mirroring is used
+	MirrorData.SetNeedsReinitialized();
 }
 
 void USkinWeightsPaintTool::CleanMesh() const
@@ -1999,14 +2094,14 @@ void USkinWeightsPaintTool::ApplyStamp(const FBrushStampData& Stamp)
 			// use mesh topology to iteratively smooth weights across neighboring vertices
 			const float UseStrength = CalculateBrushStrengthToUse(EWeightEditOperation::Relax);
 			constexpr int32 RelaxIterationsPerStamp = 3;
-			RelaxWeightOnVertices(VerticesInStamp, VertexFalloffs, UseStrength, RelaxIterationsPerStamp, WeightEditsFromStamp);
+			CreateWeightEditsToRelaxVertices(VerticesInStamp, VertexFalloffs, UseStrength, RelaxIterationsPerStamp, WeightEditsFromStamp);
 		}
 		else
 		{
 			// edit weight; either by "Add", "Remove", "Replace", "Multiply"
 			const float UseStrength = CalculateBrushStrengthToUse(WeightToolProperties->BrushMode);
 			const int32 CurrentBoneIndex = GetCurrentBoneIndex();
-			EditWeightOfBoneOnVertices(
+			CreateWeightEditsForVertices(
 				WeightToolProperties->BrushMode,
 				CurrentBoneIndex,
 				VerticesInStamp,
@@ -2017,7 +2112,7 @@ void USkinWeightsPaintTool::ApplyStamp(const FBrushStampData& Stamp)
 	}
 
 	// apply weight edits to the mesh without closing the transaction
-	ApplyWeightEditsToMeshMidChange(WeightEditsFromStamp);
+	ApplyWeightEditsWithoutTransaction(WeightEditsFromStamp);
 }
 
 float USkinWeightsPaintTool::CalculateBrushStrengthToUse(EWeightEditOperation EditMode) const
@@ -2054,7 +2149,7 @@ float USkinWeightsPaintTool::CalculateBrushStrengthToUse(EWeightEditOperation Ed
 	return UseStrength;
 }
 
-void USkinWeightsPaintTool::EditWeightOfBoneOnVertices(
+void USkinWeightsPaintTool::CreateWeightEditsForVertices(
 	EWeightEditOperation EditOperation,
 	const BoneIndex Bone,
 	const TArray<int32>& VertexIndices,
@@ -2111,7 +2206,7 @@ void USkinWeightsPaintTool::EditWeightOfBoneOnVertices(
 
 		// normalize the values across all bones affecting this vertex, and record the bone edits
 		// normalization is done while holding all weights on the current bone constant so that user edits are not overwritten
-		Weights.EditVertexWeightAndNormalize(
+		Weights.CreateWeightEditForVertex(
 			Bone,
 			VertexID,
 			NewValueAfterStamp,
@@ -2119,7 +2214,7 @@ void USkinWeightsPaintTool::EditWeightOfBoneOnVertices(
 	}
 }
 
-void USkinWeightsPaintTool::RelaxWeightOnVertices(
+void USkinWeightsPaintTool::CreateWeightEditsToRelaxVertices(
 	TArray<int32> VertexIndices,
 	TArray<float> VertexFalloffs,
 	const float Strength,
@@ -2223,50 +2318,24 @@ void USkinWeightsPaintTool::InitializeSmoothWeightsOperator()
 	SmoothWeightsOp->MinimumWeightThreshold = MinimumWeightThreshold;
 }
 
-void USkinWeightsPaintTool::ApplyWeightEditsToMeshMidChange(const SkinPaintTool::FMultiBoneWeightEdits& WeightEdits)
+void USkinWeightsPaintTool::ApplyWeightEditsWithoutTransaction(const FMultiBoneWeightEdits& WeightEdits)
 {
-	// store weight edits from all stamps made during a single stroke (1 transaction per stroke)
-	{
-		TRACE_CPUPROFILER_EVENT_SCOPE(SkinTool::ApplyWeightEditsToActiveChange);
-		for (const TTuple<BoneIndex, SkinPaintTool::FSingleBoneWeightEdits>& BoneWeightEdits : WeightEdits.PerBoneWeightEdits)
-		{
-			ActiveChange->AddBoneWeightEdit(BoneWeightEdits.Value);
-		}
-	}
-
 	// apply weights to current weights (triggers sparse deformation update)
-	{
-		TRACE_CPUPROFILER_EVENT_SCOPE(SkinTool::ApplyWeightEditsToCurrentWeights);
-		Weights.ApplyEditsToCurrentWeights(WeightEdits);
-	}
-	
-	// queue update of vertex colors
-	WeightEdits.GetEditedVertexIndices(VerticesToUpdateColor);
+	Weights.ApplyEditsToCurrentWeights(WeightEdits);
+	// update list of vertices needing updated vertex colors
+	WeightEdits.AddEditedVerticesToSet(VerticesToUpdateColor);
+	// store changes in the transaction buffer, but keep it open
+	auto PartialToFullVertexConverter = [this](int32 InVertexIndex){ return PartialToFullMeshVertexIndex(InVertexIndex); };
+	ActiveChange->StoreMultipleWeightEdits(WeightEdits, PartialToFullVertexConverter);
 }
 
-void USkinWeightsPaintTool::ApplyWeightEditsAsTransaction(const SkinPaintTool::FMultiBoneWeightEdits& WeightEdits, const FText& TransactionLabel)
+void USkinWeightsPaintTool::ApplyWeightEditsAsTransaction(const FMultiBoneWeightEdits& WeightEdits, const FText& TransactionLabel)
 {
 	// clear the active change to start a new one
 	BeginChange();
-
-	// store weight edits in the active change
-	for (const TTuple<BoneIndex, FSingleBoneWeightEdits>& BoneWeightEdits : WeightEdits.PerBoneWeightEdits)
-	{
-		ActiveChange->AddBoneWeightEdit(BoneWeightEdits.Value);
-	}
-
-	// store pruned influences
-	for (const TPair<VertexIndex, BoneIndex>& PrunedInfluence : WeightEdits.PrunedInfluences)
-	{
-		ActiveChange->AddPruneBoneEdit(PrunedInfluence.Key, PrunedInfluence.Value);
-	}
-	
-	// apply the weight edits of the active change to the actual mesh
-	// - copies weight modifications to the tool's weight data structure and updates the vertex colors
-	// - updates PreChangeWeights
-	ActiveChange->Apply(this);
-	
-	// store active change in the transaction buffer
+	// apply the weights
+	ApplyWeightEditsWithoutTransaction(WeightEdits);
+	// store active change in the transaction buffer and sync weight buffers
 	EndChange(TransactionLabel);
 }
 
@@ -2281,7 +2350,7 @@ BoneIndex USkinWeightsPaintTool::GetBoneIndexFromName(const FName BoneName) cons
 {
 	if (BoneName == NAME_None)
 	{
-		return  INDEX_NONE;		
+		return INDEX_NONE;		
 	}
 	const BoneIndex* Found = Weights.Deformer.BoneNameToIndexMap.Find(BoneName);
 	return Found ? *Found : INDEX_NONE;
@@ -2297,6 +2366,9 @@ void USkinWeightsPaintTool::SetFocusInViewport() const
 
 void USkinWeightsPaintTool::OnShutdown(EToolShutdownType ShutdownType)
 {
+	// shutdown must be performed on full mesh, so end isolated selection
+	FinishIsolatedSelection();
+	
 	// save tool properties
 	WeightToolProperties->SaveProperties(this);
 
@@ -2356,8 +2428,8 @@ void USkinWeightsPaintTool::BeginChange()
 
 void USkinWeightsPaintTool::EndChange(const FText& TransactionLabel)
 {
-	// swap weight buffers
-	Weights.SwapAfterChange();
+	// sync weight buffers
+	Weights.SyncWeightBuffers();
 	
 	// record transaction
 	GetToolManager()->BeginUndoTransaction(TransactionLabel);
@@ -2372,7 +2444,9 @@ void USkinWeightsPaintTool::ExternalUpdateWeights(const int32 BoneIndex, const T
 {
 	for (const TTuple<int32, float>& Pair : NewValues)
 	{
-		const int32 VertexID = Pair.Key;
+		// weights are always stored in transactions as full mesh indices
+		// if we are in an isolated selection, we must convert them to the partial mesh for them to be applied
+		const int32 VertexID = FullToPartialMeshVertexIndex(Pair.Key);
 		const float Weight = Pair.Value;
 		Weights.SetWeightOfBoneOnVertex(BoneIndex, VertexID, Weight, Weights.CurrentWeights);
 		Weights.SetWeightOfBoneOnVertex(BoneIndex, VertexID, Weight, Weights.PreChangeWeights);
@@ -2437,7 +2511,7 @@ void USkinWeightsPaintTool::ExternalRemoveInfluences(const TArray<TPair<VertexIn
 	}
 }
 
-void FSkinMirrorData::RegenerateMirrorData(
+void FSkinMirrorData::EnsureMirrorDataIsUpdated(
     const TArray<FName>& BoneNames,
     const TMap<FName, BoneIndex>& BoneNameToIndexMap,
 	const FReferenceSkeleton& RefSkeleton,
@@ -2474,86 +2548,90 @@ void FSkinMirrorData::RegenerateMirrorData(
 		//UE_LOG(LogTemp, Log, TEXT("-------"));
 	}
 
-	// hash grid constants
-	constexpr float HashGridCellSize = 2.0f;
-	constexpr float ThresholdRadius = 0.1f;
-
-	// build a point set of the rest pose vertices
-	UE::Geometry::FDynamicPointSet3d PointSet;
-	for (int32 PointID = 0; PointID < RefPoseVertices.Num(); ++PointID)
+	// build a spatial hash grid
+	constexpr float HashGridCellSize = 2.0f; // the length of the cell size in the point hash grid
+	UE::Geometry::TPointHashGrid3f<int32> VertHash(HashGridCellSize, INDEX_NONE);
+	VertHash.Reserve(RefPoseVertices.Num());
+	for (int32 VertexID = 0; VertexID < RefPoseVertices.Num(); ++VertexID)
 	{
-		PointSet.InsertVertex(PointID, RefPoseVertices[PointID]);
+		VertHash.InsertPointUnsafe(VertexID, static_cast<FVector3f>(RefPoseVertices[VertexID]));
 	}
-
-	// build a spatial hash map from the point set
-	UE::Geometry::FPointSetAdapterd PointSetAdapter = UE::Geometry::MakePointsAdapter(&PointSet);
-	UE::Geometry::FPointSetHashtable PointHash(&PointSetAdapter);
-	PointHash.Build(HashGridCellSize, FVector3d::Zero());
 	
 	// generate a map of point IDs on the target side, to their equivalent vertex ID on the source side 
-	TArray<int> PointsInThreshold;
-	TArray<int> PointsInSphere;
-	bAllVerticesMirrored = true;
 	for (int32 TargetVertexID = 0; TargetVertexID < RefPoseVertices.Num(); ++TargetVertexID)
 	{
 		const FVector& TargetPosition = RefPoseVertices[TargetVertexID];
 
-		if (Direction == EMirrorDirection::PositiveToNegative && TargetPosition[Axis-1] >= 0.f)
+		// we only generate the mirror map for vertices on the target side of the mirror plane
+		// if the mirror plane is flipped, the map must be regenerated because there is no guarantee that the mirror map is symmetrical
+		if (!IsPointOnTargetMirrorSide(TargetPosition))
 		{
-			continue; // copying to negative side, but vertex is on positive side
-		}
-		if (Direction == EMirrorDirection::NegativeToPositive && TargetPosition[Axis-1] <= 0.f)
-		{
-			continue; // copying to positive side, but vertex is on negative side
+			continue;
 		}
 
 		// flip position across the mirror axis
-		FVector MirroredPosition = TargetPosition;
+		FVector3f MirroredPosition = FVector3f(TargetPosition);
 		MirroredPosition[Axis-1] *= -1.f;
 
 		// query spatial hash near mirrored position, gradually increasing search radius until at least 1 point is found
-		PointsInSphere.Reset();
-		float SearchRadius = ThresholdRadius;
-		while(PointsInSphere.IsEmpty())
+		TPair<int32, double> ClosestMirroredPoint = {INDEX_NONE, TNumericLimits<double>::Max()};
+		float SearchRadius = HashGridCellSize;
+		while(ClosestMirroredPoint.Key == INDEX_NONE)
 		{
-			PointHash.FindPointsInBall(MirroredPosition, SearchRadius, PointsInSphere);
-			SearchRadius += ThresholdRadius;
+			ClosestMirroredPoint = VertHash.FindNearestInRadius(
+				MirroredPosition,
+				SearchRadius,
+				[&RefPoseVertices, MirroredPosition](int32 VID)
+				{
+					return FVector3f::DistSquared(FVector3f(RefPoseVertices[VID]), MirroredPosition);
+				});
+			
+			SearchRadius += HashGridCellSize;
 
-			// forcibly break out if search radius gets bigger than the mesh bounds.
-			// this could potentially happen if mesh is highly non-symmetrical along mirror axis
-			if (SearchRadius >= HashGridCellSize)
+			// forcibly break out if search radius gets bigger than the maximum search radius
+			constexpr float MaxSearchRadius = 5.f; // TODO we may want to expose this value to the user...
+			if (SearchRadius >= MaxSearchRadius)
 			{
 				break;
 			}
 		}
-
-		// no mirrored points?
-		if (PointsInSphere.IsEmpty())
-		{
-			bAllVerticesMirrored = false;
-			continue;
-		}
-
-		// find the closest single point
-		float ClosestDistSq = TNumericLimits<float>::Max();
-		int32 ClosestVertexID = INDEX_NONE;
-		for (const int32 PointInSphereID : PointsInSphere)
+		
+		// disallow copying from vertices that are on the target side of the mirror plane
+		if (ClosestMirroredPoint.Key != INDEX_NONE)
 		{
-			const float DistSq = FVector::DistSquared(RefPoseVertices[PointInSphereID], MirroredPosition);
-			if (DistSq < ClosestDistSq)
+			const FVector& SourcePointPosition = RefPoseVertices[ClosestMirroredPoint.Key];
+			if (IsPointOnTargetMirrorSide(SourcePointPosition))
 			{
-				ClosestDistSq = DistSq;
-				ClosestVertexID = PointInSphereID;
+				ClosestMirroredPoint.Key = INDEX_NONE;	
 			}
 		}
 		
-		// record the mirrored vertex ID for this vertex 
-		VertexMap.FindOrAdd(TargetVertexID, ClosestVertexID); // (TO, FROM)
+		// record the mirrored vertex ID for this vertex (may be INDEX_NONE)
+		VertexMap.FindOrAdd(TargetVertexID, ClosestMirroredPoint.Key); // (TO, FROM)
 	}
 	
 	bIsInitialized = true;
 }
 
+const TMap<int32, int32>& FSkinMirrorData::GetVertexMap() const
+{
+	ensure(bIsInitialized);
+	return VertexMap;
+}
+
+bool FSkinMirrorData::IsPointOnTargetMirrorSide(const FVector& InPoint) const
+{
+	if (Direction == EMirrorDirection::PositiveToNegative && InPoint[Axis-1] >= 0.f)
+	{
+		return false; // target is negative side, but point is on positive side
+	}
+	if (Direction == EMirrorDirection::NegativeToPositive && InPoint[Axis-1] <= 0.f)
+	{
+		return false; // target is positive side, but vertex is on negative side
+	}
+
+	return true;
+}
 
 void USkinWeightsPaintTool::MirrorWeights(EAxis::Type Axis, EMirrorDirection Direction)
 {
@@ -2564,7 +2642,7 @@ void USkinWeightsPaintTool::MirrorWeights(EAxis::Type Axis, EMirrorDirection Dir
 	const FReferenceSkeleton& RefSkeleton = Weights.Deformer.Component->GetSkeletalMeshAsset()->GetRefSkeleton();
 
 	// refresh mirror tables (cached / lazy generated)
-	MirrorData.RegenerateMirrorData(
+	MirrorData.EnsureMirrorDataIsUpdated(
 		Weights.Deformer.BoneNames,
 		Weights.Deformer.BoneNameToIndexMap,
 		RefSkeleton,
@@ -2576,40 +2654,57 @@ void USkinWeightsPaintTool::MirrorWeights(EAxis::Type Axis, EMirrorDirection Dir
 	const TMap<BoneIndex, BoneIndex>& BoneMap = MirrorData.GetBoneMap();
 	const TMap<VertexIndex, VertexIndex>& VertexMirrorMap = MirrorData.GetVertexMap(); // <Target, Source>
 
-	// get set of vertices to mirror
-	TArray<VertexIndex> AllVerticesToEdit = GetSelectedVertices();
-
-	// convert all vertex indices to the target side of the mirror plane
-	TSet<VertexIndex> VerticesToMirror;
-	
-	for (const VertexIndex SelectedVertex : AllVerticesToEdit)
+	// get the selected vertices
+	const TArray<int32>& SelectedVertices = GetSelectedVertices();
+	// we need to convert selection to the equivalent target vertex indices (on the target side of the mirror plane)
+	// if a vertex is already on the target side, great
+	// if the user selected vertices on the source side, we convert them to the mirrored equivalent on the target side
+	TSet<VertexIndex> TargetVertices;
+	TArray<VertexIndex> MissingVertices;
+	for (const VertexIndex SelectedVertex : SelectedVertices)
 	{
+		int32 TargetVertexID = INDEX_NONE;
 		if (VertexMirrorMap.Contains(SelectedVertex))
 		{
 			// vertex is located across the mirror plane (target side, to copy TO)
-			VerticesToMirror.Add(SelectedVertex);
+			TargetVertexID = SelectedVertex;
 		}
 		else
 		{
 			// vertex is located on the source side (to copy FROM), so we need to search for it's mirror target vertex
 			for (const TPair<VertexIndex, VertexIndex>& ToFromPair : VertexMirrorMap)
 			{
-				if (ToFromPair.Value != SelectedVertex)
+				if (ToFromPair.Value == SelectedVertex)
 				{
-					continue;
+					TargetVertexID = ToFromPair.Key;
+					break;
 				}
-				VerticesToMirror.Add(ToFromPair.Key);
-				break;
 			}
 		}
+
+		// selected vertex did not have a mirrored equivalent
+		if (TargetVertexID == INDEX_NONE)
+		{
+			MissingVertices.Add(TargetVertexID);
+			continue;;
+		}
+		
+		// add to the list of target vertices to set weights on
+		TargetVertices.Add(TargetVertexID);
 	}
 	
 	// spin through all target vertices to mirror and copy weights from source
 	FMultiBoneWeightEdits WeightEditsFromMirroring;
-	for (const VertexIndex VertexToMirror : VerticesToMirror)
+	for (const VertexIndex TargetVertexID : TargetVertices)
 	{
-		const VertexIndex SourceVertexID = VertexMirrorMap[VertexToMirror];
-		const VertexIndex TargetVertexID = VertexToMirror;
+		const VertexIndex SourceVertexID = VertexMirrorMap[TargetVertexID];
+
+		// store in the missing vertices array if no mirrored vertex was found
+		if (SourceVertexID == INDEX_NONE)
+		{
+			MissingVertices.Add(TargetVertexID);
+			continue;
+		}
 
 		// remove all weight on vertex
 		for (const FVertexBoneWeight& TargetBoneWeight : Weights.PreChangeWeights[TargetVertexID])
@@ -2634,9 +2729,9 @@ void USkinWeightsPaintTool::MirrorWeights(EAxis::Type Axis, EMirrorDirection Dir
 	ApplyWeightEditsAsTransaction(WeightEditsFromMirroring, TransactionLabel);
 
 	// warn if some vertices were not mirrored
-	if (!MirrorData.GetAllVerticesMirrored())
+	if (!MissingVertices.IsEmpty())
 	{
-		UE_LOG(LogMeshModelingToolsEditor, Log, TEXT("Mirror Skin Weights: some vertex weights were not mirrored because a vertex was not found close enough to the mirrored location."));
+		UE_LOG(LogMeshModelingToolsEditor, Log, TEXT("Mirror Skin Weights: %d vertex weights were not mirrored because a vertex was not found close enough to the mirrored location."), MissingVertices.Num());
 	}
 }
 
@@ -2648,17 +2743,27 @@ void USkinWeightsPaintTool::EditWeightsOnVertices(
 	const TArray<VertexIndex>& VertexIndices,
 	const bool bShouldTransact)
 {
+	if (!(Weights.Deformer.InvCSRefPoseTransforms.IsValidIndex(Bone)))
+	{
+		return;
+	}
+	
 	// create weight edits from setting the weight directly
 	FMultiBoneWeightEdits DirectWeightEdits;
 	const TArray<float> VertexFalloffs = {}; // no falloff
 
 	if (EditOperation == EWeightEditOperation::Relax)
 	{
-		RelaxWeightOnVertices(GetSelectedVertices(), VertexFalloffs, Value, Iterations, DirectWeightEdits);
+		CreateWeightEditsToRelaxVertices(
+			GetSelectedVertices(),
+			VertexFalloffs,
+			Value,
+			Iterations,
+			DirectWeightEdits);
 	}
 	else
 	{
-		EditWeightOfBoneOnVertices(
+		CreateWeightEditsForVertices(
 				EditOperation,
 				Bone,
 				VertexIndices,
@@ -2675,7 +2780,7 @@ void USkinWeightsPaintTool::EditWeightsOnVertices(
 	}
 	else
 	{
-		ApplyWeightEditsToMeshMidChange(DirectWeightEdits);
+		ApplyWeightEditsWithoutTransaction(DirectWeightEdits);
 	}
 }
 
@@ -2754,67 +2859,12 @@ void USkinWeightsPaintTool::AverageWeights(const float Strength)
 		return;
 	}
 	
-	// remove smallest weights values from a weight map to fit in MAX_TOTAL_INFLUENCES
-	auto TruncateWeights = [](TMap<BoneIndex, float>& InOutWeights)
-	{
-		// sort influences by total weight
-		InOutWeights.ValueSort([](const float& A, const float& B)
-		{
-			return A > B;
-		});
-
-		// truncate to MaxInfluences
-		int32 Index = 0;
-		for (TMap<BoneIndex, float>::TIterator It(InOutWeights); It; ++It)
-		{
-			if (Index >= MAX_TOTAL_INFLUENCES)
-			{
-				It.RemoveCurrent();
-			}
-			else
-			{
-				++Index;
-			}
-		}
-	};
-
-	// normalize weights values to sum to 1.0
-	auto NormalizeWeights = [](TMap<BoneIndex, float>& InOutWeights)
-	{
-		// normalize remaining influences
-		float TotalWeight = 0.f;
-		for (const TTuple<BoneIndex, float>& Weight : InOutWeights)
-		{
-			TotalWeight += Weight.Value;
-		}
-		for (TTuple<BoneIndex, float>& Weight : InOutWeights)
-		{
-			Weight.Value /= TotalWeight > SMALL_NUMBER ? TotalWeight : 1.f;
-		}
-	};
-
-	// sum up all weight on the given vertices
-	auto AccumulateWeights = [](
-		TMap<BoneIndex, float>& OutWeights,
-		TArray<VertexWeights>& AllWeights,
-		const TArray<VertexIndex>& VerticesToAccumulate)
-	{
-		for (const VertexIndex VertexID : VerticesToAccumulate)
-		{
-			for (const FVertexBoneWeight& BoneWeight : AllWeights[VertexID])
-			{
-				float& AccumulatedWeight = OutWeights.FindOrAdd(BoneWeight.BoneID);
-				AccumulatedWeight += BoneWeight.Weight;
-			}
-		}
-	};
-	
 	// get vertices to edit weights on
 	const TArray<VertexIndex>& VerticesToAverage = GetSelectedVertices();
 	TMap<BoneIndex, float> AveragedWeights;
-	AccumulateWeights(AveragedWeights, Weights.PreChangeWeights, VerticesToAverage);
-	TruncateWeights(AveragedWeights);
-	NormalizeWeights(AveragedWeights);
+	AccumulateWeights(Weights.PreChangeWeights, VerticesToAverage, AveragedWeights);
+	TruncateWeightMap(AveragedWeights);
+	NormalizeWeightMap(AveragedWeights);
 
 	// store weight edits to apply averaging to selected vertices
 	FMultiBoneWeightEdits WeightEditsFromAveraging;
@@ -2875,8 +2925,8 @@ void USkinWeightsPaintTool::AverageWeights(const float Strength)
 			}
 
 			// enforce max influences and normalize
-			TruncateWeights(BlendedWeights);
-			NormalizeWeights(BlendedWeights);
+			TruncateWeightMap(BlendedWeights);
+			NormalizeWeightMap(BlendedWeights);
 			
 			// apply blended weights to this vertex
 			for (const TTuple<BoneIndex, float>& BlendedWeight : BlendedWeights)
@@ -2907,13 +2957,13 @@ void USkinWeightsPaintTool::NormalizeWeights()
 			// ALL influences have been pruned from vertex, so assign it to the root
 			constexpr BoneIndex RootBoneIndex = 0;
 			constexpr float FullWeight = 1.f;
-			Weights.EditVertexWeightAndNormalize(RootBoneIndex, VertexID, FullWeight,WeightEditsFromNormalization);
+			Weights.CreateWeightEditForVertex(RootBoneIndex, VertexID, FullWeight,WeightEditsFromNormalization);
 		}
 		else
 		{
 			// set first weight to current value, just to force re-normalization
 			const FVertexBoneWeight& BoneWeight = VertexWeights[0];
-			Weights.EditVertexWeightAndNormalize(BoneWeight.BoneID, VertexID, BoneWeight.Weight,WeightEditsFromNormalization);
+			Weights.CreateWeightEditForVertex(BoneWeight.BoneID, VertexID, BoneWeight.Weight,WeightEditsFromNormalization);
 		}
 	}
 
@@ -3130,7 +3180,7 @@ void USkinWeightsPaintTool::TransferWeights()
 		{
 			for (const TTuple<int32, FSingleBoneWeightEdits>& BoneWeightEdits : WeightEdits.PerBoneWeightEdits)
 			{
-				ActiveChange->AddBoneWeightEdit(BoneWeightEdits.Value);
+				ActiveChange->StoreBoneWeightEdit(BoneWeightEdits.Value, nullptr);
 			}
 
 			static const FText TransactionLabel = LOCTEXT("TransferWeightsChange", "Transfer skin weights.");
@@ -3139,6 +3189,221 @@ void USkinWeightsPaintTool::TransferWeights()
 	}
 }
 
+const FString USkinWeightsPaintTool::CopyPasteWeightsIdentifier = TEXT("UNREAL_VERTEX_WEIGHTS:");
+
+void USkinWeightsPaintTool::CopyWeights()
+{
+	// get all the selected vertices
+	const TArray<VertexIndex>& VerticesToCopy = GetSelectedVertices();
+	if (VerticesToCopy.IsEmpty())
+	{
+		const FText NotificationText = FText::FromString(TEXT("No vertices were selected. No weights were copied to the clipboard."));
+		ShowEditorMessage(ELogVerbosity::Error, NotificationText);
+		return;
+	}
+
+	// get a map of the average weight of all the selected vertices (accumulating and normalizing produces an average)
+	TMap<BoneIndex, float> WeightsToCopy;
+	AccumulateWeights(Weights.PreChangeWeights, VerticesToCopy, WeightsToCopy);
+	TruncateWeightMap(WeightsToCopy);
+	NormalizeWeightMap(WeightsToCopy);
+
+	//
+	// serialize and store in the clipboard
+	//
+
+	// create a JSON array to hold our pairs
+	TArray<TSharedPtr<FJsonValue>> JSONArray;
+	for (const TPair<BoneIndex, float>& BoneWeight : WeightsToCopy)
+	{
+		TSharedPtr<FJsonObject> WeightJSON = MakeShared<FJsonObject>();
+		WeightJSON->SetStringField("BoneName", GetBoneNameFromIndex(BoneWeight.Key).ToString());
+		WeightJSON->SetNumberField("Weight", BoneWeight.Value);
+		JSONArray.Add(MakeShared<FJsonValueObject>(WeightJSON));
+	}
+
+	// convert JSON array to string
+	FString JSONString;
+	TSharedRef<TJsonWriter<>> Writer = TJsonWriterFactory<>::Create(&JSONString);
+	FJsonSerializer::Serialize(JSONArray, Writer);
+
+	// add a custom prefix to identify our data
+	const FString ClipboardString = CopyPasteWeightsIdentifier + JSONString;
+        
+	// copy to clipboard
+	FPlatformApplicationMisc::ClipboardCopy(*ClipboardString);
+
+	// notify user
+	const FText NotificationText = FText::FromString("Copied weights to clipboard.");
+	ShowEditorMessage(ELogVerbosity::Log, NotificationText);
+}
+
+void USkinWeightsPaintTool::PasteWeights()
+{
+	const TArray<VertexIndex>& VerticesToPasteOn = GetSelectedVertices();
+	if (VerticesToPasteOn.IsEmpty())
+	{
+		const FText NotificationText = FText::FromString("No vertices were selected. No weights were pasted.");
+		ShowEditorMessage(ELogVerbosity::Error, NotificationText);
+		return;
+	}
+	
+	// get the clipboard content and check if it matches our format
+	FString ClipboardContent;
+	FPlatformApplicationMisc::ClipboardPaste(ClipboardContent);
+	if (!ClipboardContent.StartsWith(CopyPasteWeightsIdentifier))
+	{
+		const FText NotificationText = FText::FromString("Failed to paste vertex weights from clipboard. Expected header not found.");
+		ShowEditorMessage(ELogVerbosity::Fatal, NotificationText);
+		return;
+	}
+
+	// deserialize the string into JSON and parse it as name/weight pairs
+	TMap<BoneIndex, float> LoadedWeights;
+	const FString JsonString = ClipboardContent.RightChop(CopyPasteWeightsIdentifier.Len());
+	TSharedPtr<FJsonValue> JsonParsed;
+	TSharedRef<TJsonReader<>> Reader = TJsonReaderFactory<>::Create(JsonString);
+	bool bFoundImproperlyFormattedData = false;
+	bool bFoundNonExistantBone = false;
+	if (FJsonSerializer::Deserialize(Reader, JsonParsed) && JsonParsed.IsValid())
+	{
+		TArray<TSharedPtr<FJsonValue>> JsonArray = JsonParsed->AsArray();
+		for (const TSharedPtr<FJsonValue>& JsonValue : JsonArray)
+		{
+			TSharedPtr<FJsonObject> JsonObject = JsonValue->AsObject();
+			if (!JsonObject.IsValid())
+			{
+				continue;
+			}
+			FString NameString;
+			double ValueAsDouble;
+			if (!(JsonObject->TryGetStringField(TEXT("BoneName"), NameString) && JsonObject->TryGetNumberField(TEXT("Weight"), ValueAsDouble)))
+			{
+				bFoundImproperlyFormattedData = true;
+				continue;
+			}
+			
+			BoneIndex IndexOfBone = GetBoneIndexFromName(FName(NameString));
+			float Weight = ValueAsDouble;
+			if (IndexOfBone == INDEX_NONE)
+			{
+				bFoundNonExistantBone = true;
+				UE_LOG(LogMeshModelingToolsEditor, Warning, TEXT("Pasted weights referenced a missing bone: %s"), *NameString);
+				continue;
+			}
+
+			// store the loaded weight value
+			LoadedWeights.Add(IndexOfBone, Weight);
+		}
+	}
+
+	if (bFoundNonExistantBone)
+	{
+		const FText NotificationText = FText::FromString("Pasted weights referenced a missing bone. See output for details.");
+		ShowEditorMessage(ELogVerbosity::Warning, NotificationText);
+	}
+	
+	if (bFoundImproperlyFormattedData)
+	{
+		const FText NotificationText = FText::FromString("Found improperly formatted data while pasting weights from clipboard. Expected array of (BoneName,Weight) pairs.");
+		ShowEditorMessage(ELogVerbosity::Warning, NotificationText);
+	}
+
+	if (LoadedWeights.IsEmpty())
+	{
+		const FText NotificationText = FText::FromString("No weights were loaded from the clipboard. Paste aborted.");
+		ShowEditorMessage(ELogVerbosity::Fatal, NotificationText);
+		return;
+	}
+
+	// truncate and normalize (these weights could have come from anywhere)
+	TruncateWeightMap(LoadedWeights);
+	NormalizeWeightMap(LoadedWeights);
+
+	// create set of weight edits that paste the weights
+	FMultiBoneWeightEdits PasteWeightEdits;
+	for (const int32 SelectedVertex : VerticesToPasteOn)
+	{
+		// remove all current weights
+		for (const FVertexBoneWeight& BoneWeight : Weights.PreChangeWeights[SelectedVertex])
+		{
+			const float OldWeight = BoneWeight.Weight;
+			constexpr float NewWeight = 0.f;
+			PasteWeightEdits.MergeSingleEdit(BoneWeight.BoneID, SelectedVertex, OldWeight, NewWeight);
+		}
+
+		// add weights from clipboard
+		for (const TPair<BoneIndex, float>& LoadedWeight : LoadedWeights)
+		{
+			BoneIndex BoneID = LoadedWeight.Key;
+			const float OldWeight = Weights.GetWeightOfBoneOnVertex(BoneID, SelectedVertex, Weights.PreChangeWeights);
+			const float NewWeight = LoadedWeight.Value;
+			PasteWeightEdits.MergeSingleEdit(BoneID, SelectedVertex, OldWeight, NewWeight);
+		}
+	}
+	
+	// apply the changes
+	const FText TransactionLabel = LOCTEXT("PasteWeightsChange", "Paste skin weights.");
+	ApplyWeightEditsAsTransaction(PasteWeightEdits, TransactionLabel);
+
+	// notify user
+	const FText NotificationText = FText::FromString("Pasted weights.");
+	ShowEditorMessage(ELogVerbosity::Log, NotificationText);
+}
+
+void USkinWeightsPaintTool::TruncateWeightMap(TMap<BoneIndex, float>& InOutWeights)
+{
+	// sort influences by total weight
+	InOutWeights.ValueSort([](const float& A, const float& B)
+	{
+		return A > B;
+	});
+
+	// truncate to MaxInfluences
+	int32 Index = 0;
+	for (TMap<BoneIndex, float>::TIterator It(InOutWeights); It; ++It)
+	{
+		if (Index >= MAX_TOTAL_INFLUENCES)
+		{
+			It.RemoveCurrent();
+		}
+		else
+		{
+			++Index;
+		}
+	}
+}
+
+void USkinWeightsPaintTool::NormalizeWeightMap(TMap<BoneIndex, float>& InOutWeights)
+{
+	// normalize remaining influences
+	float TotalWeight = 0.f;
+	for (const TTuple<BoneIndex, float>& Weight : InOutWeights)
+	{
+		TotalWeight += Weight.Value;
+	}
+	for (TTuple<BoneIndex, float>& Weight : InOutWeights)
+	{
+		Weight.Value /= TotalWeight > SMALL_NUMBER ? TotalWeight : 1.f;
+	}
+
+}
+
+void USkinWeightsPaintTool::AccumulateWeights(
+	const TArray<SkinPaintTool::VertexWeights>& AllWeights,
+	const TArray<VertexIndex>& VerticesToAccumulate,
+	TMap<BoneIndex, float>& OutWeights)
+{
+	for (const VertexIndex VertexID : VerticesToAccumulate)
+	{
+		for (const FVertexBoneWeight& BoneWeight : AllWeights[VertexID])
+		{
+			float& AccumulatedWeight = OutWeights.FindOrAdd(BoneWeight.BoneID);
+			AccumulatedWeight += BoneWeight.Weight;
+		}
+	}
+}
+
 void USkinWeightsPaintTool::HandleSkeletalMeshModified(const TArray<FName>& InBoneNames, const ESkeletalMeshNotifyType InNotifyType)
 {
 	switch (InNotifyType)
@@ -3538,16 +3803,18 @@ void USkinWeightsPaintTool::FinishIsolatedSelection()
 		// get the equivalent vertex on the full mesh
 		const int32 VertexIndexFull = PartialSubMesh.MapVertexToBaseMesh(VertexIndexPartial);
 		// clear all the weights on this vertex
-		Weights.CurrentWeights[VertexIndexFull].Init(FVertexBoneWeight(), UE::AnimationCore::MaxInlineBoneWeightCount);
+		Weights.CurrentWeights[VertexIndexFull].Reset();
 		// replace with weights from partial mesh
 		const FVertexBoneWeightsConst& VertexWeightsPartial = AllVertexWeights.Get(VertexIndexPartial);
-		for (int32 InfluenceIndex=0; InfluenceIndex<VertexWeightsPartial.Num(); ++InfluenceIndex)
+		for (const UE::AnimationCore::FBoneWeight& SingleBoneWeight : VertexWeightsPartial)
 		{
-			const UE::AnimationCore::FBoneWeight& SingleBoneWeight = VertexWeightsPartial[InfluenceIndex];
-			FVertexBoneWeight& VertexBoneWeight = Weights.CurrentWeights[VertexIndexFull][InfluenceIndex];
-			VertexBoneWeight.BoneID = SingleBoneWeight.GetBoneIndex();
-			VertexBoneWeight.Weight = SingleBoneWeight.GetWeight();
-			VertexBoneWeight.VertexInBoneSpace = Weights.Deformer.InvCSRefPoseTransforms[VertexBoneWeight.BoneID].TransformPosition(Weights.Deformer.RefPoseVertexPositions[VertexIndexFull]);
+			const FBoneIndexType BoneID = SingleBoneWeight.GetBoneIndex();
+			const float Weight = SingleBoneWeight.GetWeight();
+			const FVector& RefPoseVertexPosition = Weights.Deformer.RefPoseVertexPositions[VertexIndexFull];
+			const FTransform& InvRefPoseBoneTransform = Weights.Deformer.InvCSRefPoseTransforms[BoneID];
+			const FVector VertexInBoneSpace = InvRefPoseBoneTransform.TransformPosition(RefPoseVertexPosition);
+			// store the new influence on this vertex
+			Weights.CurrentWeights[VertexIndexFull].Emplace(BoneID, VertexInBoneSpace, Weight);
 		}
 	}
 	// sync both weight buffers
@@ -3563,12 +3830,32 @@ void USkinWeightsPaintTool::FinishIsolatedSelection()
 	PartialMeshDescription = nullptr;
 }
 
-void USkinWeightsPaintTool::UpdateSelectedVertices()
+int32 USkinWeightsPaintTool::PartialToFullMeshVertexIndex(int32 PartialMeshVertexIndex) const
 {
-	SelectedVertices.Empty();
+	if (!PartialMeshDescription)
+	{
+		return PartialMeshVertexIndex;
+	}
+
+	return PartialSubMesh.MapVertexToBaseMesh(PartialMeshVertexIndex);
+}
+
+int32 USkinWeightsPaintTool::FullToPartialMeshVertexIndex(int32 FullMeshVertexIndex) const
+{
+	if (!PartialMeshDescription)
+	{
+		return FullMeshVertexIndex;
+	}
+
+	return PartialSubMesh.MapVertexToSubmesh(FullMeshVertexIndex);
+}
+
+const TArray<int32>& USkinWeightsPaintTool::GetSelectedVertices()
+{
+	SelectedVerticesInternal.Empty();
 	if (!PolygonSelectionMechanic)
 	{
-		return;
+		return SelectedVerticesInternal;
 	}
 	
 	const FGroupTopologySelection& Selection = PolygonSelectionMechanic->GetActiveSelection();
@@ -3581,7 +3868,7 @@ void USkinWeightsPaintTool::UpdateSelectedVertices()
 		// we have to make sure that the vertex ids are safe to use as PolygonSelectionMechanic does not act on the
 		// mesh description but on the dynamic mesh that can duplicate vertices when dealing with degenerate triangles.
 		// cf. FMeshDescriptionToDynamicMesh::Convert for more details.
-		Algo::CopyIf(VerticesToAdd, SelectedVertices, [&](int32 VertexID)
+		Algo::CopyIf(VerticesToAdd, SelectedVerticesInternal, [&](int32 VertexID)
 		{
 			return AllVertices.IsValid(VertexID);	
 		});
@@ -3616,11 +3903,8 @@ void USkinWeightsPaintTool::UpdateSelectedVertices()
 		
 		AddVertices(VerticesInSelectedFaces);
 	}
-}
-
-const TArray<int32>& USkinWeightsPaintTool::GetSelectedVertices() const
-{
-	return SelectedVertices;
+	
+	return SelectedVerticesInternal;
 }
 
 void USkinWeightsPaintTool::GetVerticesAffectedByBone(BoneIndex IndexOfBone, TSet<int32>& OutVertexIndices) const
diff --git a/Engine/Plugins/Experimental/MeshModelingToolsetExp/Source/MeshModelingToolsEditorOnlyExp/Public/SkeletalMesh/SkinWeightsPaintTool.h b/Engine/Plugins/Experimental/MeshModelingToolsetExp/Source/MeshModelingToolsEditorOnlyExp/Public/SkeletalMesh/SkinWeightsPaintTool.h
--- a/Engine/Plugins/Experimental/MeshModelingToolsetExp/Source/MeshModelingToolsEditorOnlyExp/Public/SkeletalMesh/SkinWeightsPaintTool.h
+++ b/Engine/Plugins/Experimental/MeshModelingToolsetExp/Source/MeshModelingToolsEditorOnlyExp/Public/SkeletalMesh/SkinWeightsPaintTool.h
@@ -154,7 +154,7 @@ namespace SkinPaintTool
 		void MergeSingleEdit(const int32 BoneIndex, const int32 VertexID, const float OldWeight, const float NewWeight);
 		void MergeEdits(const FSingleBoneWeightEdits& BoneWeightEdits);
 		float GetVertexDeltaFromEdits(const int32 BoneIndex, const int32 VertexIndex);
-		void GetEditedVertexIndices(TSet<int32>& OutVerticesToEdit) const;
+		void AddEditedVerticesToSet(TSet<int32>& OutEditedVertexSet) const;
 		void AddPruneBoneEdit(const VertexIndex VertexToPruneFrom, const BoneIndex BoneToPrune);
 
 		// map of bone indices to weight edits made to that bone
@@ -182,14 +182,17 @@ namespace SkinPaintTool
 
 		virtual void Revert(UObject* Object) override;
 
-		void AddBoneWeightEdit(const FSingleBoneWeightEdits& BoneWeightEdit);
+		void StoreBoneWeightEdit(const FSingleBoneWeightEdits& BoneWeightEdit, const TFunction<int32(int32)>& VertexIndexConverter);
 
-		void AddPruneBoneEdit(const VertexIndex VertexToPruneFrom, const BoneIndex BoneToPrune);
+		void StorePruneBoneEdit(const VertexIndex VertexToPruneFrom, const BoneIndex BoneToPrune);
+
+		void StoreMultipleWeightEdits(const FMultiBoneWeightEdits& WeightEdits, const TFunction<int32(int32)>& VertexIndexConverter);
 
 	private:
 		FMultiBoneWeightEdits AllWeightEdits;
 		EMeshLODIdentifier LOD = EMeshLODIdentifier::Default;
 		FName SkinWeightProfile = FSkeletalMeshAttributesShared::DefaultSkinWeightProfileName;
+		
 	};
 
 	// intermediate storage of the weight maps for duration of tool
@@ -202,7 +205,7 @@ namespace SkinPaintTool
 
 		// applies an edit to a single vertex weight on a single bone, then normalizes the remaining weights while
 		// keeping the edited weight intact (ie, adapts OTHER influences to achieve normalization)
-		void EditVertexWeightAndNormalize(
+		void CreateWeightEditForVertex(
 			const int32 BoneIndex,
 			const int32 VertexId,
 			float NewWeightValue,
@@ -232,7 +235,10 @@ namespace SkinPaintTool
 			const float Weight,
 			TArray<VertexWeights>& InOutVertexWeights);
 
-		void SwapAfterChange();
+		// some weight editing operations are RELATIVE to existing weights before the change started (Multiply, Add etc)
+		// these "existing weights" are stored in the PreChangeWeights buffer
+		// PreChange and Current buffers must be synchronized after a transaction
+		void SyncWeightBuffers();
 
 		float SetCurrentFalloffAndGetMaxFalloffThisStroke(int32 VertexID, float CurrentStrength);
 
@@ -266,7 +272,8 @@ namespace SkinPaintTool
 
 	struct FSkinMirrorData
 	{
-		void RegenerateMirrorData(
+		// lazily updates the mirror data tables for the current skeleton/mesh/mirror plane
+		void EnsureMirrorDataIsUpdated(
 			const TArray<FName>& BoneNames,
 			const TMap<FName, BoneIndex>& BoneNameToIndexMap,
 			const FReferenceSkeleton& RefSkeleton,
@@ -274,14 +281,18 @@ namespace SkinPaintTool
 			EAxis::Type InMirrorAxis,
 			EMirrorDirection InMirrorDirection);
 
+		// get a map of Target > Source bone ids across the current mirror plane
 		const TMap<int32, int32>& GetBoneMap() const { return BoneMap; };
-		const TMap<int32, int32>& GetVertexMap() const { return VertexMap; };
-		bool GetAllVerticesMirrored() const {return bAllVerticesMirrored; };
+		// get the map of Target > Source vertex ids across the current mirror plane
+		const TMap<int32, int32>& GetVertexMap() const;
+		// return true if the point lies on the TARGET side of the mirror plane
+		bool IsPointOnTargetMirrorSide(const FVector& InPoint) const;
+		// forces mirror tables to be re-generated (do this after any mesh change operation)
+		void SetNeedsReinitialized() {bIsInitialized = false;};
 		
 	private:
 		
 		bool bIsInitialized = false;
-		bool bAllVerticesMirrored = false;
 		TEnumAsByte<EAxis::Type> Axis;
 		EMirrorDirection Direction; 
 		TMap<int32, int32> BoneMap;
@@ -491,6 +502,11 @@ public:
 	void NormalizeWeights();
 	void HammerWeights();
 	void TransferWeights();
+
+	// copy paste
+	void CopyWeights();
+	void PasteWeights();
+	static const FString CopyPasteWeightsIdentifier;
 	
 	// method to set weights directly (numeric input, for example)
 	void EditWeightsOnVertices(
@@ -517,7 +533,8 @@ public:
 	void SetIsolateSelected(const bool bIsolateSelection);
 
 	// get a list of currently selected vertices (converting edges and faces to vertices)
-	const TArray<int32>& GetSelectedVertices() const;
+	const TArray<int32>& GetSelectedVertices();
+	bool HasSelectedVertices() const { return !SelectedVerticesInternal.IsEmpty(); };
 	void GetVerticesAffectedByBone(BoneIndex IndexOfBone, TSet<int32>& OutVertexIndices) const;
 	void GetSelectedTriangles(TArray<int32>& OutTriangleIndices) const;
 
@@ -547,7 +564,7 @@ public:
 	// Deformations and vertex colors will be updated throughout the duration of the change.
 	void BeginChange();
 	void EndChange(const FText& TransactionLabel);
-	void ApplyWeightEditsToMeshMidChange(const SkinPaintTool::FMultiBoneWeightEdits& WeightEdits);
+	void ApplyWeightEditsWithoutTransaction(const SkinPaintTool::FMultiBoneWeightEdits& WeightEdits);
 	// "One-off" Edits:
 	// For all one-and-done edits, you can call ApplyWeightEditsAsTransaction().
 	// It will Begin/End the change and create a transaction for it.
@@ -590,18 +607,18 @@ protected:
 	bool bStampPending;
 	int32 TriangleUnderStamp;
 	FVector StampLocalPos;
-
-	// modify vertex weights according to the specified operation,
+	
 	// generating bone weight edits to be stored in a transaction
-	void EditWeightOfBoneOnVertices(
+	// does not actually change the weight buffers
+	void CreateWeightEditsForVertices(
 		EWeightEditOperation EditOperation,
 		const BoneIndex Bone,
 		const TArray<int32>& VerticesToEdit,
 		const TArray<float>& VertexFalloffs,
 		const float InValue,
 		SkinPaintTool::FMultiBoneWeightEdits& InOutWeightEdits);
-	// same as EditWeightOfVertices() but specific to relaxation (topology aware operation)
-	void RelaxWeightOnVertices(
+	// same as CreateWeightEditsForVertices() but specific to relaxation (topology aware operation)
+	void CreateWeightEditsToRelaxVertices(
 		TArray<int32> VerticesToEdit,
 		TArray<float> VertexFalloffs,
 		const float Strength,
@@ -627,12 +644,18 @@ protected:
 	FMeshDescription* EditedMesh = nullptr;
 	// when selection is isolated, we hide the full mesh and show a submesh
 	// when islated selection is unhidden, we remap all changes from the submesh back to the full mesh
-	TSharedPtr<FMeshDescription> PartialMeshDescription = nullptr; // during isolated selection
+	TSharedPtr<FMeshDescription> PartialMeshDescription = nullptr; // only non-null during isolated selection
 	UE::Geometry::FGeometrySelection IsolatedSelectionToRestoreVertices;
 	UE::Geometry::FGeometrySelection IsolatedSelectionToRestoreEdges;
 	UE::Geometry::FGeometrySelection IsolatedSelectionToRestoreFaces;
 	bool bPendingUpdateFromPartialMesh = false;
 	void FinishIsolatedSelection();
+	// isolate selection sub-meshes
+	UE::Geometry::FDynamicSubmesh3 PartialSubMesh;
+	UE::Geometry::FDynamicMesh3 FullDynamicMesh;
+	//
+	int32 PartialToFullMeshVertexIndex(int32 PartialMeshVertexIndex) const;
+	int32 FullToPartialMeshVertexIndex(int32 FullMeshVertexIndex) const;
 
 	// storage of vertex weights per bone 
 	SkinPaintTool::FSkinToolWeights Weights;
@@ -654,20 +677,27 @@ protected:
 	// vertex colors updated when make sparse edits to subset of vertices
 	void UpdateVertexColorForSubsetOfVertices();
 	TSet<int32> VerticesToUpdateColor;
-	
 	FVector4f GetColorOfVertex(VertexIndex InVertexIndex, BoneIndex InBoneIndex) const;
 
 	// which bone are we currently painting?
 	void UpdateCurrentBone(const FName &BoneName);
+	BoneIndex GetBoneIndexFromName(const FName BoneName) const;
 	FName CurrentBone = NAME_None;
 	TOptional<FName> PendingCurrentBone;
 	TArray<FName> SelectedBoneNames;
 	TArray<BoneIndex> SelectedBoneIndices;
 
-	// determines the set of vertices to operate on, using selection as the priority
-	void UpdateSelectedVertices();
-	
-	BoneIndex GetBoneIndexFromName(const FName BoneName) const;
+	// HELPER functions for modifying weights
+	//
+	// given a map of BoneIndex > Weight values for a single vertex, modify the map by removing the smallest weights to fit in Max Influences
+	void TruncateWeightMap(TMap<BoneIndex, float>& InOutWeights);
+	// given a map of BoneIndex > Weight values for a single vertex, modify the weights to sum to 1
+	void NormalizeWeightMap(TMap<BoneIndex, float>& InOutWeights);
+	// sum all the weights on all bones for a given list of vertices (results we not be normalized!)
+	void AccumulateWeights(
+		const TArray<SkinPaintTool::VertexWeights>& AllWeights,
+		const TArray<VertexIndex>& VerticesToAccumulate,
+		TMap<BoneIndex, float>& OutWeights);
 
 	// ISkeletalMeshEditionInterface
 	virtual void HandleSkeletalMeshModified(const TArray<FName>& InBoneNames, const ESkeletalMeshNotifyType InNotifyType) override;
@@ -678,11 +708,7 @@ protected:
 	TUniquePtr<UE::Geometry::FDynamicMeshAABBTree3> MeshSpatial = nullptr;
 	TUniquePtr<UE::Geometry::FTriangleGroupTopology> SelectionTopology = nullptr;
 	void InitializeSelectionMechanic();
-	TArray<VertexIndex> SelectedVertices;
-
-	// isolate selection sub-meshes
-	UE::Geometry::FDynamicSubmesh3 PartialSubMesh;
-	UE::Geometry::FDynamicMesh3 FullDynamicMesh;
+	TArray<VertexIndex> SelectedVerticesInternal;
 
 	// skin weight layer
 	void OnActiveLODChanged();
