Skeletal Mesh Transfer: fixed bug, mesh sections disappeared after transfer.

diff --git a/Engine/Plugins/Experimental/MeshModelingToolsetExp/Source/MeshModelingToolsEditorOnlyExp/Private/SkeletalMesh/SkinWeightsPaintTool.cpp b/Engine/Plugins/Experimental/MeshModelingToolsetExp/Source/MeshModelingToolsEditorOnlyExp/Private/SkeletalMesh/SkinWeightsPaintTool.cpp
--- a/Engine/Plugins/Experimental/MeshModelingToolsetExp/Source/MeshModelingToolsEditorOnlyExp/Private/SkeletalMesh/SkinWeightsPaintTool.cpp
+++ b/Engine/Plugins/Experimental/MeshModelingToolsetExp/Source/MeshModelingToolsEditorOnlyExp/Private/SkeletalMesh/SkinWeightsPaintTool.cpp
@@ -3053,6 +3053,8 @@ void USkinWeightsPaintTool::TransferWeights()
 
 	if (!SourceTarget)
 	{
+		const FText NotificationText = LOCTEXT("NoSourceInTransfer", "No source skeletal mesh was specified. No weights were transferred.");
+		ShowEditorMessage(ELogVerbosity::Error, NotificationText);
 		return;
 	}
 
@@ -3066,11 +3068,15 @@ void USkinWeightsPaintTool::TransferWeights()
 	
 	if (!SourceMesh.HasAttributes() || !SourceMesh.Attributes()->HasBones())
 	{
+    	const FText NotificationText = LOCTEXT("NoWeightsFoundInTransfer", "No skin weights were found in the source skeletal mesh. No weights were transferred.");
+    	ShowEditorMessage(ELogVerbosity::Error, NotificationText);
 		return;
 	}
 	
 	if (SourceMesh.Attributes()->GetNumBones() == 0)
 	{
+		const FText NotificationText = LOCTEXT("NoBonesFoundInTransfer", "No bones were found in the source skeletal mesh. No weights were transferred.");
+		ShowEditorMessage(ELogVerbosity::Error, NotificationText);
 		return;
 	}
 	
@@ -3116,77 +3122,66 @@ void USkinWeightsPaintTool::TransferWeights()
 	}
 
 	const FName TargetProfile = WeightToolProperties->GetActiveSkinWeightProfile();
-	if (TransferBoneWeights.TransferWeightsToMesh(TargetMesh, TargetProfile))
+	if (!TransferBoneWeights.TransferWeightsToMesh(TargetMesh, TargetProfile))
 	{
-		// store weight edits
-		BeginChange();
-		FMultiBoneWeightEdits WeightEdits;
+		// notify user
+		const FText NotificationText = LOCTEXT("TransferWeightsUnknownIssue", "Transfer weights operation encountered an unknown issue. No weights were transferred.");
+		ShowEditorMessage(ELogVerbosity::Error, NotificationText);
+		return;
+	}
+	
+	// weight edits for transaction
+	FMultiBoneWeightEdits WeightEdits;
+
+	// spin through all the transferred skin weights and record a weight edit to apply as a transaction
+	FDynamicMeshVertexSkinWeightsAttribute* TransferedSkinWeights = TargetAttributes->GetSkinWeightsAttribute(TargetProfile);
+	check(TransferedSkinWeights);
+	static constexpr float ZeroWeight = 0.f;
+	const bool bUseSubset = !TransferBoneWeights.TargetVerticesSubset.IsEmpty();
+	const int32 NumVertices = bUseSubset ? TransferBoneWeights.TargetVerticesSubset.Num() : TargetMesh.VertexCount();
+	const FNonManifoldMappingSupport NonManifoldMappingSupport(TargetMesh);
+	for (int32 VertexIndex = 0; VertexIndex < NumVertices; ++VertexIndex)
+	{
+		const int32 VertexID = bUseSubset ? TransferBoneWeights.TargetVerticesSubset[VertexIndex] : VertexIndex;
+		const int32 SrcVertexID = NonManifoldMappingSupport.GetOriginalNonManifoldVertexID(VertexID);
 		
-		{	
-			FDynamicMeshVertexSkinWeightsAttribute* TransferedSkinWeights = TargetAttributes->GetSkinWeightsAttribute(TargetProfile);
-			check(TransferedSkinWeights);
-			
-			const bool bUseSubset = !TransferBoneWeights.TargetVerticesSubset.IsEmpty();
-
-			static constexpr float ZeroWeight = 0.f;
-			
-			const int32 NumVertices = bUseSubset ? TransferBoneWeights.TargetVerticesSubset.Num() : TargetMesh.VertexCount();
-			const FNonManifoldMappingSupport NonManifoldMappingSupport(TargetMesh);
-			
-			for (int32 VertexIndex = 0; VertexIndex < NumVertices; ++VertexIndex)
+		// remove all weight on vertex
+		const VertexWeights& VertexBoneWeights = Weights.PreChangeWeights[SrcVertexID];
+		if (!VertexBoneWeights.IsEmpty())
+		{
+			for (const FVertexBoneWeight& BoneWeight : VertexBoneWeights)
 			{
-				const int32 VertexID = bUseSubset ? TransferBoneWeights.TargetVerticesSubset[VertexIndex] : VertexIndex;
-				const int32 SrcVertexID = NonManifoldMappingSupport.GetOriginalNonManifoldVertexID(VertexID);
-				
-				// remove all weight on vertex
-				const VertexWeights& VertexBoneWeights = Weights.PreChangeWeights[SrcVertexID];
-				if (!VertexBoneWeights.IsEmpty())
-				{
-					for (const FVertexBoneWeight& BoneWeight : VertexBoneWeights)
-					{
-						const float OldWeight = BoneWeight.Weight;
-						WeightEdits.MergeSingleEdit(BoneWeight.BoneID, SrcVertexID, OldWeight, ZeroWeight);
-					}
-				}
-				else
-				{
-					WeightEdits.MergeSingleEdit(0, SrcVertexID, 1.f, ZeroWeight);
-				}
-
-				// update with new weight
-				FBoneWeights TransferedBoneWeights;
-				TransferedSkinWeights->GetValue(VertexID, TransferedBoneWeights);
-				for (FBoneWeight BoneWeight: TransferedBoneWeights)
-				{
-					const int32 BoneIndex = BoneWeight.GetBoneIndex();					
-					const float OldWeight = Weights.GetWeightOfBoneOnVertex(BoneIndex, SrcVertexID, Weights.PreChangeWeights);
-					const float NewWeight = BoneWeight.GetWeight();
-					WeightEdits.MergeSingleEdit(BoneIndex, SrcVertexID, OldWeight, NewWeight);
-				}
+				const float OldWeight = BoneWeight.Weight;
+				WeightEdits.MergeSingleEdit(BoneWeight.BoneID, SrcVertexID, OldWeight, ZeroWeight);
 			}
 		}
-
-		// set new weights (we could probably use TransferedSkinWeights instead of a full convert) 
-		FDynamicMeshToMeshDescription Converter;
-		Converter.Convert(&TargetMesh, *EditedMesh);
-
-		// update weights
-		Weights = FSkinToolWeights();
-		Weights.Profile = TargetProfile;
-		Weights.InitializeSkinWeights(GetSkeletalMeshComponent(Target), EditedMesh);
-		bVertexColorsNeedUpdated = true;
-
-		// store weight edits in the active change & commit
+		else
 		{
-			for (const TTuple<int32, FSingleBoneWeightEdits>& BoneWeightEdits : WeightEdits.PerBoneWeightEdits)
-			{
-				ActiveChange->StoreBoneWeightEdit(BoneWeightEdits.Value, nullptr);
-			}
+			WeightEdits.MergeSingleEdit(0, SrcVertexID, 1.f, ZeroWeight);
+		}
 
-			static const FText TransactionLabel = LOCTEXT("TransferWeightsChange", "Transfer skin weights.");
-			EndChange(TransactionLabel);
+		// update with new weight
+		FBoneWeights TransferedBoneWeights;
+		TransferedSkinWeights->GetValue(VertexID, TransferedBoneWeights);
+		for (const FBoneWeight& BoneWeight: TransferedBoneWeights)
+		{
+			const int32 BoneIndex = BoneWeight.GetBoneIndex();					
+			const float OldWeight = Weights.GetWeightOfBoneOnVertex(BoneIndex, SrcVertexID, Weights.PreChangeWeights);
+			const float NewWeight = BoneWeight.GetWeight();
+			WeightEdits.MergeSingleEdit(BoneIndex, SrcVertexID, OldWeight, NewWeight);
 		}
 	}
+
+	// apply the changes as a transaction
+	const FText TransactionLabel = LOCTEXT("TransferWeightsChange", "Transfer skin weights.");
+	ApplyWeightEditsAsTransaction(WeightEdits, TransactionLabel);
+
+	// put the mesh back in it's current pose
+	Weights.Deformer.SetAllVerticesToBeUpdated();
+
+	// notify user that weights were transferred.
+	const FText NotificationText = LOCTEXT("WeightsTransferred", "Skin weights transferred.");
+	ShowEditorMessage(ELogVerbosity::Log, NotificationText);
 }
 
 const FString USkinWeightsPaintTool::CopyPasteWeightsIdentifier = TEXT("UNREAL_VERTEX_WEIGHTS:");
