Skin Weight Editing: fixes bug, could not paint with relax brush unless bone is selected.
Fixes ensure that sometimes triggers when closing the tool.

diff --git a/Engine/Plugins/Experimental/MeshModelingToolsetExp/Source/MeshModelingToolsEditorOnlyExp/Private/SkeletalMesh/SkinWeightsPaintTool.cpp b/Engine/Plugins/Experimental/MeshModelingToolsetExp/Source/MeshModelingToolsEditorOnlyExp/Private/SkeletalMesh/SkinWeightsPaintTool.cpp
--- a/Engine/Plugins/Experimental/MeshModelingToolsetExp/Source/MeshModelingToolsEditorOnlyExp/Private/SkeletalMesh/SkinWeightsPaintTool.cpp
+++ b/Engine/Plugins/Experimental/MeshModelingToolsetExp/Source/MeshModelingToolsEditorOnlyExp/Private/SkeletalMesh/SkinWeightsPaintTool.cpp
@@ -143,7 +143,7 @@ USkeletalMeshComponent* GetSkeletalMeshComponent(const UToolTarget* InTarget)
 	}
 	
 	USkeletalMeshComponent* Component = Cast<USkeletalMeshComponent>(TargetComponent->GetOwnerComponent());
-	if (!ensure(Component && Component->GetSkeletalMeshAsset()))
+	if (!(Component && Component->GetSkeletalMeshAsset()))
 	{
 		return nullptr;
 	}
@@ -2070,8 +2070,9 @@ void USkinWeightsPaintTool::ApplyStamp(const FBrushStampData& Stamp)
 {
 	TRACE_CPUPROFILER_EVENT_SCOPE(SkinTool::ApplyStamp);
 
-	// must select a bone to paint
-	if (CurrentBone == NAME_None)
+	// must select a bone to paint in all modes EXCEPT relax-mode which operates on ALL bones
+	const bool bIsInRelaxMode = WeightToolProperties->BrushMode == EWeightEditOperation::Relax;
+	if (!bIsInRelaxMode && CurrentBone == NAME_None)
 	{
 		return;
 	}
