new combine mesh interfaces:
+ separate creation of per-part meshes, so per-part approximations can be shared across multiple combine mesh calls
+ split out per-LOD settings of combined meshes to allow more flexbile control of the meshes

diff --git a/Engine/Plugins/Experimental/MeshModelingToolsetExp/Source/GeometryProcessingAdapters/Private/CombineMeshInstancesImpl.cpp b/Engine/Plugins/Experimental/MeshModelingToolsetExp/Source/GeometryProcessingAdapters/Private/CombineMeshInstancesImpl.cpp
--- a/Engine/Plugins/Experimental/MeshModelingToolsetExp/Source/GeometryProcessingAdapters/Private/CombineMeshInstancesImpl.cpp
+++ b/Engine/Plugins/Experimental/MeshModelingToolsetExp/Source/GeometryProcessingAdapters/Private/CombineMeshInstancesImpl.cpp
@@ -125,6 +125,9 @@
 	UStaticMesh* SourceAsset = nullptr;
 	const IGeometryProcessing_CombineMeshInstances::FMeshLODSet* SourceMeshLODSet = nullptr;
 
+	// Precomputed part meshes
+	int32 PrecomputeMeshesIndex = INDEX_NONE;
+
 	TArray<FMeshPartInstance> Instances;
 
 	bool bPreserveUVs = false;
@@ -154,12 +157,14 @@
 
 struct FSourceGeometry
 {
+	// Note: These arrays cannot be resized
 	TArray<UE::Geometry::FDynamicMesh3> SourceMeshLODs;
 	UE::Geometry::FSimpleShapeSet3d CollisionShapes;
 };
 
 struct FOptimizedGeometry
 {
+	// Note: These arrays cannot be resized
 	TArray<UE::Geometry::FDynamicMesh3> SimplifiedMeshLODs;
 
 	TArray<UE::Geometry::FDynamicMesh3> ApproximateMeshLODs;
@@ -176,15 +181,19 @@
 	FMeshPartsAssembly(FMeshPartsAssembly&) = delete;
 	FMeshPartsAssembly& operator=(const FMeshPartsAssembly&) = delete;
 
-
+	// Array of parts, sorted in descending order of triangle count. Note each part can have multiple instances
 	TArray<TUniquePtr<FMeshPart>> Parts;
 
+	// All materials used by parts in this assembly
 	TArray<UMaterialInterface*> UniqueMaterials;
 	TMap<UMaterialInterface*, int32> MaterialMap;
 
+
+	// For each part, and array of source LODs
 	TArray<FSourceGeometry> SourceMeshGeometry;
+	// For each part, arrays of simplified and approximated versions of the part
 	TArray<FOptimizedGeometry> OptimizedMeshGeometry;
-
+	// Array of aabb tree per SourceMeshGeometry; note these are always wrt LOD 0 of the corresponding source geometry
 	TArray<FDynamicMeshAABBTree3> SourceMeshSpatials;
 
 	// allow external code to preprocess dynamic mesh for a specific instance
@@ -209,17 +218,19 @@
 
 		UStaticMesh* StaticMesh = SourceMeshInstance.SourceMesh;
 		FMeshPart** FoundPart = StaticMeshToPartMap.Find(StaticMesh);
+		FMeshPart* LocalPtr = nullptr;
 		if (FoundPart == nullptr)
 		{
 			TUniquePtr<FMeshPart> NewPart = MakeUnique<FMeshPart>();
 			NewPart->SourceAsset = StaticMesh;
-			FMeshPart* Ptr = NewPart.Get();
+			NewPart->PrecomputeMeshesIndex = SourceMeshInstance.PrecomputeMeshesIndex;
+			LocalPtr = NewPart.Get();
 					
 			AssemblyOut.Parts.Add(MoveTemp(NewPart));
 			// store source model?
 
-			StaticMeshToPartMap.Add(StaticMesh, Ptr);
-			FoundPart = &Ptr;
+			StaticMeshToPartMap.Add(StaticMesh, LocalPtr);
+			FoundPart = &LocalPtr;
 		}
 
 		FMeshPartInstance NewInstance;
@@ -268,17 +279,19 @@
 		const IGeometryProcessing_CombineMeshInstances::FMeshLODSet* MeshLODSet = &SourceInstanceList.MeshLODSets[MeshSetIndex];
 
 		FMeshPart** FoundPart = MeshLODSetToPartMap.Find(MeshLODSet);
+		FMeshPart* LocalPtr = nullptr;
 		if (FoundPart == nullptr)
 		{
 			TUniquePtr<FMeshPart> NewPart = MakeUnique<FMeshPart>();
 			NewPart->SourceMeshLODSet = MeshLODSet;
-			FMeshPart* Ptr = NewPart.Get();
+			NewPart->PrecomputeMeshesIndex = SourceMeshInstance.PrecomputeMeshesIndex;
+			LocalPtr = NewPart.Get();
 
 			AssemblyOut.Parts.Add(MoveTemp(NewPart));
 			// store source model?
 
-			MeshLODSetToPartMap.Add(MeshLODSet, Ptr);
-			FoundPart = &Ptr;
+			MeshLODSetToPartMap.Add(MeshLODSet, LocalPtr);
+			FoundPart = &LocalPtr;
 		}
 
 		FMeshPartInstance NewInstance;
@@ -600,21 +613,30 @@
 		Options.Add(Option);
 	}
 
-	void ComputeMetric(FResultOption& Option)
+	static double ComputeMetricFromDeviation(FVector2d DeviationMetric, int32 MethodID, int32 TriCount, double MaxAllowableDeviation, double TriangleCost = .7)
 	{
-		Option.DeviationMetric = DeviationMetric(*Option.Mesh, *Spatial);
-		int32 TriCount = Option.Mesh->TriangleCount();
 		int32 BaseTriCount = 12;		// 2 tris for each face of box
-		if (MaxAllowableDeviation > 0 && Option.DeviationMetric[1] > MaxAllowableDeviation)
+		if (MaxAllowableDeviation > 0 && DeviationMetric[1] > MaxAllowableDeviation)
 		{
-			Option.CostMetric = TNumericLimits<float>::Max() + (double)Option.MethodID;
+			return TNumericLimits<float>::Max() + (double)MethodID;
 		}
 		else
 		{
-			Option.CostMetric = Option.DeviationMetric[0] * FMathd::Pow((double)TriCount / (double)BaseTriCount, TriangleCost);
+			return DeviationMetric[0] * FMathd::Pow((double)TriCount / (double)BaseTriCount, TriangleCost);
 		}
 	}
 
+	// Compute an error metric to decide which approximation of the input shape to favor
+	// Note this is a point-sampled (not area weighted) average squared error scaled by (TriCount / 12)^.7, or a huge number (but ordered by generation method) if the max error is too high
+	// So the metric favors e.g. using fewer triangles in higher error regions (so there are less samples there) and a higher error is 'ok' if it comes with an associated lower tri count.
+	// Note this metric is only 'for' comparing meshes that represent the same shape, it's not intended to compare quality of two different parts meshes
+	// (e.g. do not use it to decide which part to 'promote' to a worse LOD to hit a tri budget)
+	void ComputeMetric(FResultOption& Option)
+	{
+		Option.DeviationMetric = DeviationMetric(*Option.Mesh, *Spatial);
+		Option.CostMetric = ComputeMetricFromDeviation(Option.DeviationMetric, Option.MethodID, Option.Mesh->TriangleCount(), MaxAllowableDeviation, TriangleCost);
+	}
+
 	void SelectBestOption(
 		FDynamicMesh3& ResultMesh, 
 		int32& MethodID)
@@ -1429,7 +1451,7 @@
 			FMeshNormals::QuickRecomputeOverlayNormals(ApproxGeo.ApproximateMeshLODs[k]);
 		}
 
-		// try remeshing the last Source LOD to reduce it's triangle count, by removing spurious geometry
+		// try planar simplification for lower Source LODs, reducing triangle count in flat regions
 		if (CombineOptions.bRetriangulateSourceLODs)
 		{
 			for (int32 SourceLODIndex = CombineOptions.StartRetriangulateSourceLOD; SourceLODIndex < NumSourceLODs; ++SourceLODIndex)
@@ -1542,12 +1564,12 @@
 				PartCostInfo.ReplacedWeight = 1.0;
 			}
 
-			int32 LastTotalCurLODTriCount = 999999;
+			int32 LastTotalCurLODTriCount = TMathUtil<int32>::SafeLargeValue;
 			int32 MaxIters = 1000;
 			int32 NoProgressIters = 0;
 			for ( int32 NumIter = 0; NumIter < MaxIters; ++NumIter)
 			{
-				// compute current estimate of total part count for this LOD
+				// compute current estimate of per-part and total tri counts for this LOD
 				int32 TotalCurLODTriCount = 0;
 				for (int32 SetIndex = 0; SetIndex < NumParts; ++SetIndex)
 				{
@@ -3061,13 +3083,15 @@
 bool ComputeHiddenRemovalForLOD(
 	FDynamicMesh3& MeshLOD,
 	int32 LODIndex,
-	IGeometryProcessing_CombineMeshInstances::FOptions CombineOptions)
+	IGeometryProcessing_CombineMeshInstances::ERemoveHiddenFacesMode RemoveHiddenFacesMethod,
+	double RemoveHiddenSamplingDensity,
+	bool bDoubleSidedHiddenRemoval)
 {
 	bool bVerbose = CVarGeometryCombineMeshInstancesVerbose.GetValueOnAnyThread();
 
 	TRACE_CPUPROFILER_EVENT_SCOPE(RemoveHidden_LOD);
 	bool bModified = false;
-	switch (CombineOptions.RemoveHiddenFacesMethod)
+	switch (RemoveHiddenFacesMethod)
 	{
 		case IGeometryProcessing_CombineMeshInstances::ERemoveHiddenFacesMode::OcclusionBased:
 			RemoveHiddenFaces_Occlusion(MeshLOD, 200);		// 200 is arbitrary here! should improve once max-distance is actually available (currently ignored)
@@ -3075,7 +3099,7 @@
 			break;
 		case IGeometryProcessing_CombineMeshInstances::ERemoveHiddenFacesMode::ExteriorVisibility:
 		case IGeometryProcessing_CombineMeshInstances::ERemoveHiddenFacesMode::Fastest:
-			RemoveHiddenFaces_ExteriorVisibility(MeshLOD, CombineOptions.RemoveHiddenSamplingDensity, CombineOptions.bDoubleSidedHiddenRemoval, LODIndex);
+			RemoveHiddenFaces_ExteriorVisibility(MeshLOD, RemoveHiddenSamplingDensity, bDoubleSidedHiddenRemoval, LODIndex);
 			bModified = true;
 			break;
 	}
@@ -3088,22 +3112,42 @@
 	return bModified;
 }
 
+bool ComputeHiddenRemovalForLOD(
+	FDynamicMesh3& MeshLOD,
+	int32 LODIndex,
+	IGeometryProcessing_CombineMeshInstances::FOptions CombineOptions)
+{
+	return ComputeHiddenRemovalForLOD(MeshLOD, LODIndex, CombineOptions.RemoveHiddenFacesMethod, CombineOptions.RemoveHiddenSamplingDensity, CombineOptions.bDoubleSidedHiddenRemoval);
+}
+
 
 void OptimizeLODMeshTriangulation(
 	FDynamicMesh3& MeshLOD,
 	int32 LODIndex,
-	IGeometryProcessing_CombineMeshInstances::FOptions CombineOptions,
+	TFunction<FIndex3i(const FDynamicMesh3& Mesh, int32 TriangleID)> TriangleGroupingIDFunc,
+	bool bWantCoplanarMerging,
+	bool bWantPlanarRetriangulation,
 	double BaseGeometricTolerance,
 	TSet<int32>* SkipMaterialIDs = nullptr)
 {
 	bool bVerbose = CVarGeometryCombineMeshInstancesVerbose.GetValueOnAnyThread();
 
-	TFunction<FIndex3i(const FDynamicMesh3& Mesh, int32 TriangleID)> GroupingIDFunc = CombineOptions.TriangleGroupingIDFunc;
+	TFunction<FIndex3i(const FDynamicMesh3& Mesh, int32 TriangleID)> GroupingIDFunc = TriangleGroupingIDFunc;
 	if (!GroupingIDFunc)
 	{
 		GroupingIDFunc = [](const FDynamicMesh3&, int32) { return FIndex3i::Zero(); };
 	}
 
+	PostProcessHiddenFaceRemovedMesh(MeshLOD, BaseGeometricTolerance, bWantCoplanarMerging, bWantPlanarRetriangulation,	GroupingIDFunc, SkipMaterialIDs);
+}
+
+void OptimizeLODMeshTriangulation(
+	FDynamicMesh3& MeshLOD,
+	int32 LODIndex,
+	IGeometryProcessing_CombineMeshInstances::FOptions CombineOptions,
+	double BaseGeometricTolerance,
+	TSet<int32>* SkipMaterialIDs = nullptr)
+{
 	bool bWantCoplanarMerging = (CombineOptions.bMergeCoplanarFaces)
 		&& (LODIndex >= CombineOptions.MergeCoplanarFacesStartLOD)
 		&& (LODIndex > CombineOptions.PreserveUVLODLevel);
@@ -3111,7 +3155,7 @@
 		&& (CombineOptions.PlanarPolygonRetriangulationStartLOD >= 0)
 		&& (LODIndex >= CombineOptions.PlanarPolygonRetriangulationStartLOD);
 
-	PostProcessHiddenFaceRemovedMesh(MeshLOD, BaseGeometricTolerance, bWantCoplanarMerging, bWantPlanarRetriangulation,	GroupingIDFunc, SkipMaterialIDs);
+	OptimizeLODMeshTriangulation(MeshLOD, LODIndex, CombineOptions.TriangleGroupingIDFunc, bWantCoplanarMerging, bWantPlanarRetriangulation, BaseGeometricTolerance, SkipMaterialIDs);
 }
 
 
@@ -3122,24 +3166,38 @@
 	TArray<FCombinedMeshLOD>& MeshLODs,
 	const TArray<double>& OptimizationTolerances,
 	int32 FirstVoxWrappedIndex,
-	IGeometryProcessing_CombineMeshInstances::FOptions CombineOptions,
+	int32 NumLODs,
+	bool bRemoveHiddenFaces,
+	TFunctionRef<bool(int32)> LODRemoveHidden,
+	TFunctionRef<bool(int32)> LODWantCoplanarMerging,
+	TFunctionRef<bool(int32)> LODWantPlanarRetriangulation,
+	TFunctionRef<IGeometryProcessing_CombineMeshInstances::ERemoveHiddenFacesMode(int32)> LODRemoveHiddenFacesMethod,
+	TFunctionRef<double(int32)> LODRemoveHiddenSamplingDensity,
+	TFunctionRef<bool(int32)> LODDoubleSidedHiddenRemoval,
+	TFunction<FIndex3i(const FDynamicMesh3& Mesh, int32 TriangleID)> TriangleGroupingIDFunc,
+
+	IGeometryProcessing_CombineMeshInstances::ECoarseApproximationStrategy CoarseLODStrategy,
+	double CoarseApproximationDetailSize,
+	TFunctionRef<int32(int32)> GetCoarseLODMaxTriCount,
+	double CoarseLODBaseTolerance,
+	double HardNormalAngleDeg,
+	bool bAutoGenerateMissingUVs,
+	bool bAutoGenerateTangents,
 	TSet<int32>* PreserveTopologyMaterialIDs = nullptr
 )
 {
 	using namespace UE::Geometry;
 	bool bVerbose = CVarGeometryCombineMeshInstancesVerbose.GetValueOnAnyThread();
 
-	int32 NumLODs = CombineOptions.NumLODs;
-
 	TArray<UE::Tasks::FTask> PendingRemoveHiddenTasks;
-	bool bRemoveHiddenFaces =
-		(CombineOptions.RemoveHiddenFacesMethod != IGeometryProcessing_CombineMeshInstances::ERemoveHiddenFacesMode::None
-			&& CVarGeometryCombineMeshInstancesRemoveHidden.GetValueOnAnyThread() > 0);
 	if (bRemoveHiddenFaces)
 	{
-		for (int32 LODIndex = CombineOptions.RemoveHiddenStartLOD; LODIndex < NumLODs && LODIndex < FirstVoxWrappedIndex; ++LODIndex)
+		for (int32 LODIndex = 0; LODIndex < NumLODs; ++LODIndex)
 		{
-			if (MeshLODs[LODIndex].Mesh.TriangleCount() == 0) continue;
+			if (!LODRemoveHidden(LODIndex) || MeshLODs[LODIndex].Mesh.TriangleCount() == 0)
+			{
+				continue;
+			}
 
 			if (bVerbose)
 			{
@@ -3147,10 +3205,14 @@
 			}
 
 			double UseTolerance = OptimizationTolerances[LODIndex];
-			UE::Tasks::FTask RemoveHiddenTask = UE::Tasks::Launch(UE_SOURCE_LOCATION, [&MeshLODs, &CombineOptions, LODIndex, UseTolerance, PreserveTopologyMaterialIDs]()
+			bool bWantCoplanarMerging = LODWantCoplanarMerging(LODIndex), bWantPlanarRetriangulation = LODWantPlanarRetriangulation(LODIndex);
+			UE::Tasks::FTask RemoveHiddenTask = UE::Tasks::Launch(UE_SOURCE_LOCATION, 
+				[&MeshLODs, LODIndex, UseTolerance, PreserveTopologyMaterialIDs,
+				LODRemoveHiddenFacesMethod, LODRemoveHiddenSamplingDensity, LODDoubleSidedHiddenRemoval,
+				TriangleGroupingIDFunc, bWantCoplanarMerging, bWantPlanarRetriangulation]()
 			{
-				ComputeHiddenRemovalForLOD(MeshLODs[LODIndex].Mesh, LODIndex, CombineOptions);
-				OptimizeLODMeshTriangulation(MeshLODs[LODIndex].Mesh, LODIndex, CombineOptions, UseTolerance, PreserveTopologyMaterialIDs);
+				ComputeHiddenRemovalForLOD(MeshLODs[LODIndex].Mesh, LODIndex, LODRemoveHiddenFacesMethod(LODIndex), LODRemoveHiddenSamplingDensity(LODIndex), LODDoubleSidedHiddenRemoval(LODIndex));
+				OptimizeLODMeshTriangulation(MeshLODs[LODIndex].Mesh, LODIndex, TriangleGroupingIDFunc, bWantCoplanarMerging, bWantPlanarRetriangulation, UseTolerance, PreserveTopologyMaterialIDs);
 			});
 			PendingRemoveHiddenTasks.Add(RemoveHiddenTask);
 
@@ -3166,7 +3228,7 @@
 	// Process VoxWrapped LODs 
 	//
 	bool bUsingCoarseSweepApproximation = false;
-	if ( FirstVoxWrappedIndex < 9999 &&  (MeshLODs[FirstVoxWrappedIndex].Mesh.TriangleCount() > 0) )
+	if ( FirstVoxWrappedIndex < MeshLODs.Num() && (MeshLODs[FirstVoxWrappedIndex].Mesh.TriangleCount() > 0))
 	{
 		FDynamicMesh3 SourceVoxWrapMesh = MoveTemp(MeshLODs[FirstVoxWrappedIndex].Mesh);
 		FDynamicMeshAABBTree3 SourceSpatial(&SourceVoxWrapMesh, true);
@@ -3177,28 +3239,28 @@
 		{
 			TRACE_CPUPROFILER_EVENT_SCOPE(ComputeVoxWrap);
 			
-			if (CombineOptions.CoarseLODStrategy == IGeometryProcessing_CombineMeshInstances::ECoarseApproximationStrategy::VoxelBasedSolidApproximation)
+			if (CoarseLODStrategy == IGeometryProcessing_CombineMeshInstances::ECoarseApproximationStrategy::VoxelBasedSolidApproximation)
 			{
-				ComputeVoxWrapMesh(SourceVoxWrapMesh, SourceSpatial, InitialCoarseApproximation, CombineOptions.CoarseApproximationDetailSize, VoxelDimension);
+				ComputeVoxWrapMesh(SourceVoxWrapMesh, SourceSpatial, InitialCoarseApproximation, CoarseApproximationDetailSize, VoxelDimension);
 				bUsingCoarseSweepApproximation = false;
 			}
-			else if (CombineOptions.CoarseLODStrategy == IGeometryProcessing_CombineMeshInstances::ECoarseApproximationStrategy::SweptPlanarProjection)
+			else if (CoarseLODStrategy == IGeometryProcessing_CombineMeshInstances::ECoarseApproximationStrategy::SweptPlanarProjection)
 			{
-				ComputeBestFullProjectionMesh(SourceVoxWrapMesh, SourceSpatial, InitialCoarseApproximation, CombineOptions.CoarseApproximationDetailSize);
+				ComputeBestFullProjectionMesh(SourceVoxWrapMesh, SourceSpatial, InitialCoarseApproximation, CoarseApproximationDetailSize);
 				bUsingCoarseSweepApproximation = true;
 			}
-			else if (CombineOptions.CoarseLODStrategy == IGeometryProcessing_CombineMeshInstances::ECoarseApproximationStrategy::IntersectSweptPlanarProjections)
+			else if (CoarseLODStrategy == IGeometryProcessing_CombineMeshInstances::ECoarseApproximationStrategy::IntersectSweptPlanarProjections)
 			{
-				ComputeProjectionMeshIntersection(SourceVoxWrapMesh, SourceSpatial, InitialCoarseApproximation, CombineOptions.CoarseApproximationDetailSize);
+				ComputeProjectionMeshIntersection(SourceVoxWrapMesh, SourceSpatial, InitialCoarseApproximation, CoarseApproximationDetailSize);
 				bUsingCoarseSweepApproximation = true;
 			}
 			else  // Automatic
 			{
 				// try swept-planar-projection as it is cheaper and generally better. If it deviates too much, fall back to voxel
 				FDynamicMesh3 SweptPlanarCoarseMesh;
-				ComputeBestFullProjectionMesh(SourceVoxWrapMesh, SourceSpatial, SweptPlanarCoarseMesh, CombineOptions.CoarseApproximationDetailSize);
+				ComputeBestFullProjectionMesh(SourceVoxWrapMesh, SourceSpatial, SweptPlanarCoarseMesh, CoarseApproximationDetailSize);
 				FVector2d SweepDeviation = ComputeGeometricDeviation(SweptPlanarCoarseMesh, SourceSpatial);
-				bUsingCoarseSweepApproximation = (SweepDeviation.Y < 2.0 * CombineOptions.CoarseApproximationDetailSize);
+				bUsingCoarseSweepApproximation = (SweepDeviation.Y < 2.0 * CoarseApproximationDetailSize);
 				if (bUsingCoarseSweepApproximation)
 				{
 					InitialCoarseApproximation = MoveTemp(SweptPlanarCoarseMesh);
@@ -3206,7 +3268,7 @@
 				else
 				{
 					FDynamicMesh3 VoxWrapCoarseMesh;
-					ComputeVoxWrapMesh(SourceVoxWrapMesh, SourceSpatial, VoxWrapCoarseMesh, CombineOptions.CoarseApproximationDetailSize, VoxelDimension);
+					ComputeVoxWrapMesh(SourceVoxWrapMesh, SourceSpatial, VoxWrapCoarseMesh, CoarseApproximationDetailSize, VoxelDimension);
 					InitialCoarseApproximation = MoveTemp(VoxWrapCoarseMesh);
 				}
 			}
@@ -3237,8 +3299,8 @@
 			UE_LOG(LogGeometry, Log, TEXT("         FastCollapse         - Tris %8d Verts %8d"), InitialCoarseApproximation.TriangleCount(), InitialCoarseApproximation.VertexCount());
 		}
 
-		int32 MaxTriCount = CombineOptions.CoarseLODMaxTriCountBase;
-		double SimplifyTolerance = CombineOptions.CoarseLODBaseTolerance;
+		int32 MaxTriCount = GetCoarseLODMaxTriCount(FirstVoxWrappedIndex);
+		double SimplifyTolerance = CoarseLODBaseTolerance;
 
 		// for very simple parts it can be the case that the last approximate LOD is
 		// lower tri-count than the first coarse approximation. In that case just use it.
@@ -3262,16 +3324,16 @@
 			}
 		}
 		InitialCoarseApproximation.EnableAttributes();
-		InitializeNormalsFromAngleThreshold(InitialCoarseApproximation, CombineOptions.HardNormalAngleDeg);
+		InitializeNormalsFromAngleThreshold(InitialCoarseApproximation, HardNormalAngleDeg);
 		ComputeVoxWrapMeshAutoUV(InitialCoarseApproximation);
 		MeshLODs[FirstVoxWrappedIndex].Mesh = MoveTemp(InitialCoarseApproximation);
 
 		// iterate simplification criteria to next level
 		SimplifyTolerance *= 1.5;
-		MaxTriCount /= 2;
 
 		for (int32 LODIndex = FirstVoxWrappedIndex+1; LODIndex < NumLODs; ++LODIndex)
 		{
+			MaxTriCount = GetCoarseLODMaxTriCount(LODIndex);
 			// need to simplify from previous level to preserve UVs/etc
 			MeshLODs[LODIndex].Mesh = MeshLODs[LODIndex-1].Mesh;
 
@@ -3281,7 +3343,6 @@
 			}
 
 			SimplifyTolerance *= 1.5;
-			MaxTriCount /= 2;
 		}
 
 		// Project colors and materials after mesh simplification to avoid constraining it.
@@ -3298,14 +3359,14 @@
 
 	// parallel regenerate UVs and potentially tangents for any areas of LODs that are missing UVs
 	TArray<UE::Tasks::FTask> PendingAutoUVTasks;
-	bool bComputeTangents = (CombineOptions.bAutoGenerateMissingUVs && CombineOptions.bAutoGenerateTangents);
-	if (CombineOptions.bAutoGenerateMissingUVs)
+	bool bComputeTangents = (bAutoGenerateMissingUVs && bAutoGenerateTangents);
+	if (bAutoGenerateMissingUVs)
 	{
 		for (int32 LODIndex = 0; LODIndex < NumLODs && LODIndex < FirstVoxWrappedIndex; ++LODIndex)
 		{
 			if (MeshLODs[LODIndex].Mesh.TriangleCount() == 0) continue;
 
-			UE::Tasks::FTask AutoUVTask = UE::Tasks::Launch(UE_SOURCE_LOCATION, [&MeshLODs, &CombineOptions, LODIndex, bComputeTangents]()
+			UE::Tasks::FTask AutoUVTask = UE::Tasks::Launch(UE_SOURCE_LOCATION, [&MeshLODs, LODIndex, bComputeTangents]()
 			{
 				ComputeMissingUVs(MeshLODs[LODIndex].Mesh);
 				if (bComputeTangents)
@@ -3332,7 +3393,7 @@
 					UE_LOG(LogGeometry, Log, TEXT("  Optimizing LOD%d - Tris %6d Verts %6d"), LODIndex, MeshLODs[LODIndex].Mesh.TriangleCount(), MeshLODs[LODIndex].Mesh.VertexCount());
 				}
 
-				ComputeHiddenRemovalForLOD(MeshLODs[LODIndex].Mesh, LODIndex, CombineOptions);
+				ComputeHiddenRemovalForLOD(MeshLODs[LODIndex].Mesh, LODIndex, LODRemoveHiddenFacesMethod(LODIndex), LODRemoveHiddenSamplingDensity(LODIndex), LODDoubleSidedHiddenRemoval(LODIndex));
 			}
 		}, (bVerbose) ? EParallelForFlags::ForceSingleThread :  EParallelForFlags::None );
 	}
@@ -3342,6 +3403,50 @@
 	UE::Tasks::Wait(PendingAutoUVTasks);
 }
 
+void ProcessCombinedLODChain(
+	TArray<FCombinedMeshLOD>& MeshLODs,
+	const TArray<double>& OptimizationTolerances,
+	int32 FirstVoxWrappedIndex,
+	IGeometryProcessing_CombineMeshInstances::FOptions CombineOptions,
+	TSet<int32>* PreserveTopologyMaterialIDs = nullptr
+)
+{
+	bool bRemoveHiddenFaces =
+		(CombineOptions.RemoveHiddenFacesMethod != IGeometryProcessing_CombineMeshInstances::ERemoveHiddenFacesMode::None
+			&& CVarGeometryCombineMeshInstancesRemoveHidden.GetValueOnAnyThread() > 0);
+	ProcessCombinedLODChain(MeshLODs, OptimizationTolerances, FirstVoxWrappedIndex,
+		CombineOptions.NumLODs,
+		bRemoveHiddenFaces,
+		[&CombineOptions, FirstVoxWrappedIndex](int32 LODIndex) { return LODIndex >= CombineOptions.RemoveHiddenStartLOD && LODIndex < FirstVoxWrappedIndex; },
+		[&CombineOptions](int32 LODIndex) { return CombineOptions.bMergeCoplanarFaces && LODIndex >= CombineOptions.MergeCoplanarFacesStartLOD && LODIndex >= CombineOptions.PreserveUVLODLevel; },
+		[&CombineOptions](int32 LODIndex)
+		{
+			return CombineOptions.bMergeCoplanarFaces && LODIndex >= CombineOptions.MergeCoplanarFacesStartLOD && LODIndex >= CombineOptions.PreserveUVLODLevel
+				&& CombineOptions.PlanarPolygonRetriangulationStartLOD >= 0 && LODIndex >= CombineOptions.PlanarPolygonRetriangulationStartLOD;
+		},
+		[&CombineOptions](int32) { return CombineOptions.RemoveHiddenFacesMethod; },
+		[&CombineOptions](int32) { return CombineOptions.RemoveHiddenSamplingDensity; },
+		[&CombineOptions](int32) { return CombineOptions.bDoubleSidedHiddenRemoval; },
+		CombineOptions.TriangleGroupingIDFunc,
+		CombineOptions.CoarseLODStrategy,
+		CombineOptions.CoarseApproximationDetailSize,
+		[&CombineOptions, FirstVoxWrappedIndex](int32 LODIndex)
+		{
+			int32 MaxCoarseTri = CombineOptions.CoarseLODMaxTriCountBase;
+			for (int32 Idx = FirstVoxWrappedIndex; Idx < LODIndex; ++Idx)
+			{
+				MaxCoarseTri /= 2;
+			}
+			return MaxCoarseTri;
+		},
+		CombineOptions.CoarseLODBaseTolerance,
+		CombineOptions.HardNormalAngleDeg,
+		CombineOptions.bAutoGenerateMissingUVs,
+		CombineOptions.bAutoGenerateTangents,
+		PreserveTopologyMaterialIDs
+	);
+}
 
+
 struct FCombinedSubAssembly
 {
@@ -3369,7 +3474,7 @@
 		FDynamicMesh3 LODMesh = MoveTemp(MeshLODs[LODLevel].Mesh);
 
 		// If we ended up larger than the mesh in the previous LOD, we should use that instead!
-		// This can happen particular with VoxWrap LODs
+		// This can happen particular with VoxWrap LODs [TODO: should we therefore not be excluding voxwrap lods w/ MaxReplaceLOD, above?]
 		if (LODLevel > 0 && LODLevel < MaxReplaceLOD)
 		{
 			if (LODMesh.TriangleCount() > OutputSubAssembly.MeshLODs.Last().TriangleCount())
@@ -3394,9 +3499,9 @@
 	TArray<FCombinedMeshLOD> MeshLODs;
 	MeshLODs.SetNum(NumLODs);
 
-	int FirstVoxWrappedIndex = 9999;
+	int FirstVoxWrappedIndex = MAX_int32;
 	TArray<ECombinedLODType> LODTypes;
-	LODTypes.Init(ECombinedLODType::Approximated, NumLODs);
+	LODTypes.Init(ECombinedLODType::Approximated, NumLODs); // Note Approximated is the default here, and will cover the span between simplified and voxwrapped
 	for (int32 LODLevel = 0; LODLevel < NumLODs; ++LODLevel)
 	{
 		if (LODLevel < CombineOptions.NumCopiedLODs)
@@ -3685,6 +3790,608 @@
 }
 
 
+void BuildCombinedMeshFromPrecomputedMeshes(
+	TConstArrayView<IGeometryProcessing_CombineMeshInstances::FSinglePartMeshSet> CachedPartMeshes,
+	const FMeshPartsAssembly& Assembly,
+	const IGeometryProcessing_CombineMeshInstances::FCombineMeshInstancesOptionsGeneral& AllLODOptions,
+	TConstArrayView<IGeometryProcessing_CombineMeshInstances::FCombineMeshInstancesOptionsPerLOD> PerLODOptions,
+	TArray<FCombinedSubAssembly>& CombinedResults)
+{
+	using namespace UE::Geometry;
+	using FLODOpts = IGeometryProcessing_CombineMeshInstances::FCombineMeshInstancesOptionsPerLOD;
+	using FAllOpts = IGeometryProcessing_CombineMeshInstances::FCombineMeshInstancesOptionsGeneral;
+
+	bool bVerbose = CVarGeometryCombineMeshInstancesVerbose.GetValueOnAnyThread();
+
+	int32 NumLODs = PerLODOptions.Num();
+	TArray<FCombinedMeshLOD> MeshLODs;
+	MeshLODs.SetNum(NumLODs);
+
+	int32 FirstVoxWrappedIndex = MAX_int32;
+
+	// LOD types tracks the optimization method per LOD. We enforce that this is increasing, e.g. must always use at-least-or-more aggressive optimization methods for higher LODs
+	TArray<ECombinedLODType> LODTypes;
+	LODTypes.Reserve(NumLODs);
+	IGeometryProcessing_CombineMeshInstances::EMeshOptimizationMethod LastLODMethod = IGeometryProcessing_CombineMeshInstances::EMeshOptimizationMethod::None;
+	for (const FLODOpts& Opts : PerLODOptions)
+	{
+		IGeometryProcessing_CombineMeshInstances::EMeshOptimizationMethod UseMethod = Opts.OptimizationMethod;
+		if ((int32)UseMethod < (int32)LastLODMethod)
+		{
+			UE_LOG(LogGeometry, Warning, TEXT("LOD optimization methods must be increasing in optimization level."));
+			UseMethod = LastLODMethod;
+		}
+		switch (UseMethod)
+		{
+		case IGeometryProcessing_CombineMeshInstances::EMeshOptimizationMethod::None:
+			LODTypes.Add(ECombinedLODType::Copied);
+			break;
+		case IGeometryProcessing_CombineMeshInstances::EMeshOptimizationMethod::SimplifyOrApproximate:
+			LODTypes.Add(ECombinedLODType::Simplified);
+			break;
+		case IGeometryProcessing_CombineMeshInstances::EMeshOptimizationMethod::ApproximateOnly:
+			LODTypes.Add(ECombinedLODType::Approximated);
+			break;
+		case IGeometryProcessing_CombineMeshInstances::EMeshOptimizationMethod::VoxelizeAndDecimate:
+			int32 LODIdx = LODTypes.Add(ECombinedLODType::VoxWrapped);
+			FirstVoxWrappedIndex = FMath::Min(LODIdx, FirstVoxWrappedIndex);
+			break;
+		}
+		LastLODMethod = UseMethod;
+	}
+	check(LODTypes.Num() == NumLODs);
+
+	int32 NumParts = Assembly.Parts.Num();
+
+	// determine maximum number of UV channels on any Part LOD0, and configure the output combined
+	// LOD meshes to have that many UV channels (clamping to at least 1). Triangles from Parts that
+	// have fewer UV channels will end up with (0,0) UVs in the extra channels.
+	int32 MaxNumUVChannels = 1;
+	for (int32 SetIndex = 0; SetIndex < NumParts; ++SetIndex)
+	{
+		const IGeometryProcessing_CombineMeshInstances::FSinglePartMeshSet& PartMeshes = CachedPartMeshes[Assembly.Parts[SetIndex]->PrecomputeMeshesIndex];
+		const FDynamicMesh3& SourceMesh = PartMeshes.Source[0];
+		
+		MaxNumUVChannels = FMath::Max(MaxNumUVChannels, (SourceMesh.HasAttributes() ? SourceMesh.Attributes()->NumUVLayers() : 0));
+	}
+	for (FCombinedMeshLOD& LODMeshData : MeshLODs)
+	{
+		LODMeshData.Mesh.Attributes()->SetNumUVLayers(MaxNumUVChannels);
+	}
+
+
+	// determine if we have multiple part subsets. In this case we need to be able to split the mesh
+	// by part later, which we will do by appending a polygroup layer
+	TArray<int32> SubsetIDs;
+	for (const TUniquePtr<FMeshPart>& Part : Assembly.Parts)
+	{
+		for (const FMeshPartInstance& Instance : Part->Instances)
+		{
+			SubsetIDs.AddUnique(Instance.SubsetID);
+		}
+	}
+	bool bHaveMultiplePartSubsets = (SubsetIDs.Num() > 1);
+	if (bHaveMultiplePartSubsets)
+	{
+		for (FCombinedMeshLOD& LODMeshData : MeshLODs)
+		{
+			LODMeshData.Mesh.Attributes()->SetNumPolygroupLayers(1);
+			LODMeshData.SubsetIDs = LODMeshData.Mesh.Attributes()->GetPolygroupLayer(0);
+		}
+	}
+
+	TArray<int32> PartSourcesPerLOD; // array of precomputed mesh indices per LOD per Part, ordered as [ Lod0[Part0 Part1 ...], Lod1[Part0 Part1 ...], ...]
+	PartSourcesPerLOD.SetNumZeroed(NumParts * NumLODs);
+	// Get the category of approximation as quality-ordered integer -- source == 0, simplified == 1, approximated == 2
+	auto SourceIndexToMeshCategory = [&CachedPartMeshes, &Assembly](int32 PartIdx, int32 SourceIdx, int32& WithinCategoryIdx, bool& bCanPromoteWithinCategory, bool& bCanPromoteAtAll, double& PromotedAvgError)
+		{
+			const IGeometryProcessing_CombineMeshInstances::FSinglePartMeshSet& PartMeshes = CachedPartMeshes[Assembly.Parts[PartIdx]->PrecomputeMeshesIndex];
+			int32 MeshesCount = PartMeshes.Source.Num() + PartMeshes.Simplified.Num() + PartMeshes.Approximated.Num();
+			PromotedAvgError = 0;
+			bCanPromoteAtAll = SourceIdx + 1 < MeshesCount;
+			if (!bCanPromoteAtAll)
+			{
+				SourceIdx = MeshesCount - 1;
+			}
+			if (SourceIdx < PartMeshes.Source.Num())
+			{
+				WithinCategoryIdx = SourceIdx;
+				bCanPromoteWithinCategory = SourceIdx + 1 < PartMeshes.Source.Num();
+				if (!bCanPromoteWithinCategory && bCanPromoteAtAll)
+				{
+					PromotedAvgError = PartMeshes.SimplifiedMeshErrors.IsEmpty() ? PartMeshes.ApproximatedMeshErrors[0].AverageError : PartMeshes.SimplifiedMeshErrors[0].AverageError;
+				}
+				return 0;
+			}
+			SourceIdx -= PartMeshes.Source.Num();
+			if (SourceIdx < PartMeshes.Simplified.Num())
+			{
+				WithinCategoryIdx = SourceIdx;
+				bCanPromoteWithinCategory = SourceIdx + 1 < PartMeshes.Simplified.Num();
+				if (bCanPromoteWithinCategory)
+				{
+					PromotedAvgError = PartMeshes.SimplifiedMeshErrors[SourceIdx + 1].AverageError;
+				}
+				else if (bCanPromoteAtAll)
+				{
+					PromotedAvgError = PartMeshes.ApproximatedMeshErrors[0].AverageError;
+				}
+				return 1;
+			}
+			SourceIdx -= PartMeshes.Simplified.Num();
+			WithinCategoryIdx = SourceIdx;
+			bCanPromoteWithinCategory = SourceIdx + 1 < PartMeshes.Approximated.Num();
+			if (bCanPromoteWithinCategory)
+			{
+				PromotedAvgError = PartMeshes.ApproximatedMeshErrors[SourceIdx + 1].AverageError;
+			}
+			return 2;
+		};
+	auto SourceIndexToMesh = [&CachedPartMeshes, &Assembly](int32 PartIdx, int32 SourceIdx)
+		{
+			const IGeometryProcessing_CombineMeshInstances::FSinglePartMeshSet& PartMeshes = CachedPartMeshes[Assembly.Parts[PartIdx]->PrecomputeMeshesIndex];
+			if (SourceIdx < PartMeshes.Source.Num())
+			{
+				return &PartMeshes.Source[SourceIdx];
+			}
+			SourceIdx -= PartMeshes.Source.Num();
+			if (SourceIdx < PartMeshes.Simplified.Num())
+			{
+				return &PartMeshes.Simplified[SourceIdx];
+			}
+			SourceIdx -= PartMeshes.Simplified.Num();
+			if (PartMeshes.Approximated.Num() > 0)
+			{
+				return &PartMeshes.Approximated[FMath::Min(SourceIdx, PartMeshes.Approximated.Num() - 1)];
+			}
+			else // If there weren't approximated meshes, fall back to the last existing mesh
+			{
+				if (PartMeshes.Simplified.Num() > 0)
+				{
+					return &PartMeshes.Simplified.Last();
+				}
+				else
+				{
+					check(!PartMeshes.Source.IsEmpty()); // there must at least be source meshes available for every part
+					return &PartMeshes.Source.Last();
+				}
+			}
+		};
+	auto GetPrecomputedMesh = [&PartSourcesPerLOD, &SourceIndexToMesh, &CachedPartMeshes, &Assembly, &PerLODOptions, NumParts](int32 LODLevel, int32 PartIdx, int32 InstIdx) -> const FDynamicMesh3*
+		{
+			int32 SourceIdx = PartSourcesPerLOD[LODLevel * NumParts + PartIdx];
+			bool bIsDecorative = Assembly.Parts[PartIdx]->Instances[InstIdx].DetailLevel == EMeshDetailLevel::Decorative;
+			if (bIsDecorative)
+			{
+				if (PerLODOptions[LODLevel].Decorations == IGeometryProcessing_CombineMeshInstances::EDecorationHandling::Remove)
+				{
+					return nullptr;
+				}
+				else if (PerLODOptions[LODLevel].Decorations == IGeometryProcessing_CombineMeshInstances::EDecorationHandling::Approximate)
+				{
+					SourceIdx = MAX_int32; // approximate decorations with the coarsest available approximation
+				}
+			}
+			// respect whether part/instance allows approximation also
+			bool bAllowApproximation = (Assembly.Parts[PartIdx]->bAllowApproximation && Assembly.Parts[PartIdx]->Instances[InstIdx].bAllowApproximation);
+			if (!bAllowApproximation)
+			{
+				const IGeometryProcessing_CombineMeshInstances::FSinglePartMeshSet& PartMeshes = CachedPartMeshes[Assembly.Parts[PartIdx]->PrecomputeMeshesIndex];
+				PartIdx = FMath::Min(PartIdx, PartMeshes.Source.Num() + PartMeshes.Simplified.Num() - 1);
+			}
+			const FDynamicMesh3* Mesh = SourceIndexToMesh(PartIdx, SourceIdx);
+			return Mesh;
+		};
+	auto GetLODTriCount = [&PartSourcesPerLOD, &GetPrecomputedMesh, &CachedPartMeshes, &Assembly, &PerLODOptions, NumParts, &SourceIndexToMesh](int32 LODLevel)
+		{
+			int32 TriCount = 0;
+			int32 UseLODLevel = FMath::Max(LODLevel, 0);
+			for (int32 PartIdx = 0; PartIdx < NumParts; ++PartIdx)
+			{
+				// sum tris per instance, because some instances (and therefore have fewer tris, or no mesh at all)
+				int32 NumInst = Assembly.Parts[PartIdx]->Instances.Num();
+				for (int32 InstIdx = 0; InstIdx < NumInst; ++InstIdx)
+				{
+					if (const FDynamicMesh3* Mesh = GetPrecomputedMesh(LODLevel, PartIdx, InstIdx))
+					{
+						TriCount += Mesh->TriangleCount();
+					}
+				}
+			}
+			return TriCount;
+		};
+
+	TArray<int32> TriBudget;
+	TriBudget.Init(-1, NumLODs);
+	for (int32 LODLevel = 0; LODLevel < NumLODs; ++LODLevel)
+	{
+		int32 PrevLODLevel = FMath::Max(0, LODLevel - 1);
+		ECombinedLODType LODType = LODTypes[LODLevel];
+		const FLODOpts& Opts = PerLODOptions[LODLevel];
+		if (LODType == ECombinedLODType::Copied)
+		{
+			for (int32 PartIdx = 0; PartIdx < NumParts; ++PartIdx)
+			{
+				const IGeometryProcessing_CombineMeshInstances::FSinglePartMeshSet& PartMeshes = CachedPartMeshes[Assembly.Parts[PartIdx]->PrecomputeMeshesIndex];
+				PartSourcesPerLOD[LODLevel * NumParts + PartIdx] =
+					FMath::Clamp(Opts.PreferredLOD, 0, PartMeshes.Source.Num() - 1);
+			}
+		}
+		else if (LODType == ECombinedLODType::Simplified || LODType == ECombinedLODType::Approximated)
+		{
+			// first pass: copy forward from the previous source indices, + make all part indices appropriate for the requested LOD type
+			for (int32 PartIdx = 0; PartIdx < NumParts; ++PartIdx)
+			{
+				const IGeometryProcessing_CombineMeshInstances::FSinglePartMeshSet& PartMeshes = CachedPartMeshes[Assembly.Parts[PartIdx]->PrecomputeMeshesIndex];
+				// Start from the previous LOD level's source index
+				int32 UseSourceIdx = PartSourcesPerLOD[PrevLODLevel * NumParts + PartIdx];
+				// If approximated meshes are required, enforce that
+				if (LODType == ECombinedLODType::Approximated)
+				{
+					UseSourceIdx = FMath::Max(UseSourceIdx, PartMeshes.Source.Num() + PartMeshes.Simplified.Num());
+				}
+				PartSourcesPerLOD[LODLevel * NumParts + PartIdx] = UseSourceIdx;
+			}
+
+			if (!Opts.bEnableBudgetStrategy_PartLODPromotion || Opts.TriangleBudgetOptions.Method == IGeometryProcessing_CombineMeshInstances::ETriangleBudgetMethod::NoRestriction)
+			{
+				continue;
+			}
+
+			int32 TargetTriCount = Opts.TriangleBudgetOptions.TriangleBudget;
+			if (Opts.TriangleBudgetOptions.Method == IGeometryProcessing_CombineMeshInstances::ETriangleBudgetMethod::UsePercentageOfPreviousLOD)
+			{
+				int32 PrevTriCount = GetLODTriCount(PrevLODLevel);
+				TargetTriCount = Opts.TriangleBudgetOptions.LODReductionPercentage * PrevTriCount;
+			}
+			TriBudget[LODLevel] = TargetTriCount;
+			int32 CurTriCount = GetLODTriCount(LODLevel);
+			// While needed, apply part promotion.
+			// We first promote within categories, so e.g. if we have a mesh still using a source LOD that we can push to a worse source LOD, we do that before pushing anything to a simplified version
+			// We then promote to a 'worse' category, prioritizing the parts where the promoted parts have lowest average error
+			while (CurTriCount > TargetTriCount * Opts.PartLODPromotionBudgetMultiplier)
+			{
+				int32 BestCategory = MAX_int32;
+				int32 BestSourceIdx = MAX_int32;
+				int32 BestPromotableCategory = MAX_int32;
+				int32 BestPromotableWithinCatIdx = MAX_int32;
+				double BestWithinCatAvgError = FMathd::MaxReal;
+				double BestAvgError = FMathd::MaxReal;
+				// track the best-to-promote-within-category part index
+				int32 BestPromotablePartIdx = -1;
+				// track the best-to-promote-overall part index (which may require crossing categories)
+				int32 BestPartIdx = -1;
+				for (int32 PartIdx = 0; PartIdx < NumParts; ++PartIdx)
+				{
+					int32 SourceIdx = PartSourcesPerLOD[LODLevel * NumParts + PartIdx];
+					bool bCanPromoteWithinCategory, bCanPromoteAtAll;
+					int32 WithinCatIdx;
+					double PromotedAvgError = 0;
+					int32 Category = SourceIndexToMeshCategory(PartIdx, SourceIdx, WithinCatIdx, bCanPromoteWithinCategory, bCanPromoteAtAll, PromotedAvgError);
+					if (!bCanPromoteAtAll)
+					{
+						continue;
+					}
+
+					// Prioritize across-category promotions by lower promoted part average error
+					if (Category < BestCategory || (Category == BestCategory && PromotedAvgError < BestAvgError))
+					{
+						BestCategory = Category;
+						BestSourceIdx = SourceIdx;
+						BestAvgError = PromotedAvgError;
+						BestPartIdx = PartIdx;
+					}
+					if (bCanPromoteWithinCategory)
+					{
+						// For source meshes (category 0), promote by LOD index primarily (note that source meshes don't have average error stats)
+						// For simplified and approximated meshes, promote primarily by average error
+						if (Category < BestPromotableCategory || 
+						   (Category == 0 && WithinCatIdx < BestPromotableWithinCatIdx) ||
+						   (Category == BestPromotableCategory && Category > 0 && PromotedAvgError < BestWithinCatAvgError))
+						{
+							BestPromotableCategory = Category;
+							BestPromotableWithinCatIdx = WithinCatIdx;
+							BestPromotablePartIdx = PartIdx;
+							BestWithinCatAvgError = PromotedAvgError;
+						}
+					}
+				}
+				if (BestPromotablePartIdx != -1 && BestPromotableCategory <= BestCategory)
+				{
+					PartSourcesPerLOD[LODLevel * NumParts + BestPromotablePartIdx]++;
+				}
+				else if (BestPartIdx != -1)
+				{
+					PartSourcesPerLOD[LODLevel * NumParts + BestPartIdx]++;
+				}
+				else
+				{
+					break; // couldn't find a promotable part
+				}
+				CurTriCount = GetLODTriCount(LODLevel);
+			}
+		}
+		else if (LODType == ECombinedLODType::VoxWrapped)
+		{
+			int32 TargetTriCount = Opts.TriangleBudgetOptions.TriangleBudget;
+			if (Opts.TriangleBudgetOptions.Method == IGeometryProcessing_CombineMeshInstances::ETriangleBudgetMethod::UsePercentageOfPreviousLOD
+				|| Opts.TriangleBudgetOptions.Method == IGeometryProcessing_CombineMeshInstances::ETriangleBudgetMethod::NoRestriction) // for CoarseLODs, treat 'no restriction' as keeping the budget constant
+			{
+				// Note: after the first coarse LOD we set the tri budget with respect to the previous budget, not the previous actual tri count,
+				// because we haven't yet computed the coarse meshes so do not have the tri counts
+				int32 PrevBudget = TriBudget[PrevLODLevel];
+				if (LODLevel == FirstVoxWrappedIndex || PrevBudget == -1)
+				{
+					PrevBudget = GetLODTriCount(PrevLODLevel);
+				}
+				TargetTriCount = PrevBudget;
+				if (Opts.TriangleBudgetOptions.Method == IGeometryProcessing_CombineMeshInstances::ETriangleBudgetMethod::UsePercentageOfPreviousLOD)
+				{
+					TargetTriCount *= Opts.TriangleBudgetOptions.LODReductionPercentage;
+				}
+			}
+			TriBudget[LODLevel] = TargetTriCount;
+
+			// For voxwrapped LOD levels, reference the last source LOD
+			for (int32 PartIdx = 0; PartIdx < NumParts; ++PartIdx)
+			{
+				const IGeometryProcessing_CombineMeshInstances::FSinglePartMeshSet& PartMeshes = CachedPartMeshes[Assembly.Parts[PartIdx]->PrecomputeMeshesIndex];
+				PartSourcesPerLOD[LODLevel * NumParts + PartIdx] = PartMeshes.Source.Num() - 1;
+			}
+		}
+	}
+
+
+	// iterate over part sets, then for each part, over part LODs, and
+	// for each instance append the part LOD to the accumulate LOD mesh
+	for (int32 SetIndex = 0; SetIndex < NumParts; ++SetIndex)
+	{
+		const TUniquePtr<FMeshPart>& Part = Assembly.Parts[SetIndex];
+		const IGeometryProcessing_CombineMeshInstances::FSinglePartMeshSet& PartMeshes = CachedPartMeshes[Assembly.Parts[SetIndex]->PrecomputeMeshesIndex];
+
+		check(Part->bAllowMerging == true);		// do not support this yet
+
+		FMeshIndexMappings Mappings;
+
+		for (int32 LODLevel = 0; LODLevel < NumLODs; ++LODLevel)
+		{
+			int32 SourceIndex = PartSourcesPerLOD[NumParts * LODLevel + SetIndex];
+			const FDynamicMesh3* UseAppendMesh = SourceIndexToMesh(SetIndex, SourceIndex);
+
+			ECombinedLODType LevelLODType = LODTypes[LODLevel];
+
+			FCombinedMeshLOD& CombinedMeshLODData = MeshLODs[LODLevel];
+
+			for (int32 InstIdx = 0; InstIdx < Part->Instances.Num(); ++InstIdx)
+			{
+				const FMeshPartInstance& Instance = Part->Instances[InstIdx];
+				const FDynamicMesh3* InstanceAppendMesh = GetPrecomputedMesh(LODLevel, SetIndex, InstIdx);
+				if (!InstanceAppendMesh)
+				{
+					continue;
+				}
+
+				if (Instance.FilterLODLevel >= 0 && LODLevel >= Instance.FilterLODLevel)
+				{
+					continue;
+				}
+
+				// need to make a copy to run pre-process func
+				FDynamicMesh3 TempAppendMesh(*InstanceAppendMesh);
+				if (Assembly.PreProcessInstanceMeshFunc)
+				{
+					Assembly.PreProcessInstanceMeshFunc(TempAppendMesh, Instance);
+				}
+
+				// determine if we should be keeping UVs around for this Part
+				bool bPreserveUVs = PerLODOptions[LODLevel].bPreserveUVs
+					|| Part->bPreserveUVs;
+
+				// if part does not require UVs, but still has them, discard them here to encourage merging downstream
+				if (bPreserveUVs == false && TempAppendMesh.HasAttributes())
+				{
+					for (int32 UVLayer = 0; UVLayer < TempAppendMesh.Attributes()->NumUVLayers(); ++UVLayer)
+					{
+						TempAppendMesh.Attributes()->GetUVLayer(UVLayer)->ClearElements();
+					}
+					// if we have no UVs then tangents are invalid
+					TempAppendMesh.Attributes()->DisableTangents();
+				}
+
+				Mappings.Reset();
+				CombinedMeshLODData.Editor.AppendMesh(&TempAppendMesh, Mappings,
+					[&](int, const FVector3d& Pos) { return Instance.WorldTransform.TransformPosition(Pos); },
+					[&](int, const FVector3d& Normal) { return Instance.WorldTransform.TransformNormal(Normal); });
+
+				// transfer part IDs etc
+				if (CombinedMeshLODData.SubsetIDs != nullptr)
+				{
+					for (int32 tid : TempAppendMesh.TriangleIndicesItr())
+					{
+						CombinedMeshLODData.SubsetIDs->SetValue(Mappings.GetNewTriangle(tid), Instance.SubsetID);
+					}
+				}
+
+				// transfer Material IDs if part mesh has them
+				FDynamicMeshMaterialAttribute* AppendMaterialAttrib = TempAppendMesh.HasAttributes() ? TempAppendMesh.Attributes()->GetMaterialID() : nullptr;
+				for (int32 tid : TempAppendMesh.TriangleIndicesItr())
+				{
+					int32 SourceMaterialID = (AppendMaterialAttrib != nullptr) ? AppendMaterialAttrib->GetValue(tid) : 0;
+					UMaterialInterface* UseMaterial = Instance.Materials.IsValidIndex(SourceMaterialID) ? Instance.Materials[SourceMaterialID] : nullptr;
+					const int32* FoundMaterialIndex = Assembly.MaterialMap.Find(UseMaterial);
+					int32 AssignMaterialIndex = (FoundMaterialIndex != nullptr) ? *FoundMaterialIndex : 0;
+
+					CombinedMeshLODData.MaterialIDs->SetValue(Mappings.GetNewTriangle(tid), AssignMaterialIndex);
+				}
+			}
+		}
+	}
+
+	// Some Material regions may need to be explicitly preserved, this set will be passed on later
+	TSet<int32> PreserveTopologyMaterialIDSet;
+	for (UMaterialInterface* Material : AllLODOptions.PreventMergingMaterialSet)
+	{
+		if (const int32* FoundMaterialIndex = Assembly.MaterialMap.Find(Material))
+		{
+			PreserveTopologyMaterialIDSet.Add(*FoundMaterialIndex);
+		}
+	}
+
+	// make a list of per-LOD geometric tolerances that will drive additional optimization,
+	// by taking the value from the per-LOD options and making sure it never decreases
+	double LastTolerance = 0;
+	TArray<double> OptimizationTolerances;
+	for (int32 LODIndex = 0; LODIndex < NumLODs; ++LODIndex)
+	{
+		double Tolerance = FMath::Max(LastTolerance, PerLODOptions[LODIndex].SimplificationTolerance);
+		OptimizationTolerances.Add(Tolerance);
+		LastTolerance = Tolerance;
+	}
+
+
+	if (bHaveMultiplePartSubsets)
+	{
+		TArray<int32> OrderedSubsetIDs(SubsetIDs);
+
+		int32 NumSubsets = SubsetIDs.Num();
+		TArray<TArray<FCombinedMeshLOD>> SubsetMeshLODChains;
+		SubsetMeshLODChains.SetNum(NumSubsets);
+		for (int32 SubsetIndex = 0; SubsetIndex < NumSubsets; ++SubsetIndex)
+		{
+			SubsetMeshLODChains[SubsetIndex].SetNum(NumLODs);
+		}
+
+		for (int32 LODIndex = 0; LODIndex < NumLODs; ++LODIndex)
+		{
+			// split the LOD by subset ID
+			FDynamicMesh3& LODMesh = MeshLODs[LODIndex].Mesh;
+			FDynamicMeshPolygroupAttribute* SubsetIDAttrib = MeshLODs[LODIndex].SubsetIDs;
+			TArray<FDynamicMesh3> SplitMeshes;
+			FDynamicMeshEditor::SplitMesh(&LODMesh, SplitMeshes, [&](int32 tid) { return OrderedSubsetIDs.IndexOfByKey(SubsetIDAttrib->GetValue(tid)); });
+
+			// code below assumes this. If it's not the case, then we have some more complex processing to figure out...
+			check(SplitMeshes.Num() == NumSubsets);
+
+			// give each subset submesh to the 
+			for (int32 SubsetIndex = 0; SubsetIndex < NumSubsets; ++SubsetIndex)
+			{
+				TArray<FCombinedMeshLOD>& LODChain = SubsetMeshLODChains[SubsetIndex];
+				FDynamicMesh3& Submesh = SplitMeshes[SubsetIndex];
+				LODChain[LODIndex].SetMesh(MoveTemp(Submesh));
+			}
+		}
+
+		TArray<UE::Tasks::FTask> PendingSubsetTasks;
+
+		for (int32 SubsetIndex = 0; SubsetIndex < NumSubsets; ++SubsetIndex)
+		{
+			TArray<FCombinedMeshLOD>& LODChain = SubsetMeshLODChains[SubsetIndex];
+			bool bRemoveHiddenFaces = CVarGeometryCombineMeshInstancesRemoveHidden.GetValueOnAnyThread() > 0;
+			UE::Tasks::FTask ProcessSubsetTask = UE::Tasks::Launch(UE_SOURCE_LOCATION, [&LODChain, &OptimizationTolerances, &PreserveTopologyMaterialIDSet, FirstVoxWrappedIndex,
+				bRemoveHiddenFaces, &PerLODOptions, &AllLODOptions, NumLODs, &TriBudget]()
+				{
+					auto LODRemoveHidden = [&PerLODOptions, FirstVoxWrappedIndex](int32 LODIndex) { return PerLODOptions[LODIndex].RemoveHiddenFacesMethod != IGeometryProcessing_CombineMeshInstances::ERemoveHiddenFacesMode::None && LODIndex < FirstVoxWrappedIndex; };
+					auto LODMergeCoplanar = [&PerLODOptions](int32 LODIndex) { return PerLODOptions[LODIndex].bMergeCoplanarFaces; };
+					auto LODPlanarRetriangulate = [&PerLODOptions](int32 LODIndex)
+						{
+							return PerLODOptions[LODIndex].bMergeCoplanarFaces && PerLODOptions[LODIndex].bPlanarPolygonRetriangulation;
+						};
+					auto LODRemoveHiddenFaces = [&PerLODOptions](int32 LODIndex)
+						{
+							return PerLODOptions[LODIndex].RemoveHiddenFacesMethod;
+						};
+					auto LODRemoveHiddenSamplingDensity = [&PerLODOptions](int32 LODIndex)
+						{
+							return PerLODOptions[LODIndex].RemoveHiddenSamplingDensity;
+						};
+					auto LODDoubleSidedHiddenRemoval = [&PerLODOptions](int32 LODIndex)
+						{
+							return PerLODOptions[LODIndex].bDoubleSidedHiddenRemoval;
+						};
+					ProcessCombinedLODChain(LODChain, OptimizationTolerances, FirstVoxWrappedIndex,
+						NumLODs,
+						bRemoveHiddenFaces,
+						LODRemoveHidden,
+						LODMergeCoplanar,
+						LODPlanarRetriangulate,
+						LODRemoveHiddenFaces,
+						LODRemoveHiddenSamplingDensity,
+						LODDoubleSidedHiddenRemoval,
+						AllLODOptions.TriangleGroupingIDFunc,
+						AllLODOptions.CoarseLODStrategy,
+						AllLODOptions.CoarseApproximationDetailSize,
+						[&TriBudget](int32 LODIndex) { return TriBudget[LODIndex]; },
+						AllLODOptions.CoarseLODBaseTolerance,
+						AllLODOptions.HardNormalAngleDeg,
+						AllLODOptions.bAutoGenerateMissingUVs,
+						AllLODOptions.bAutoGenerateTangents,
+						&PreserveTopologyMaterialIDSet
+					);
+				});
+			PendingSubsetTasks.Add(ProcessSubsetTask);
+			if (bVerbose)
+			{
+				ProcessSubsetTask.Wait();
+			}
+		}
+
+		// wait for all subsets to finish processing
+		UE::Tasks::Wait(PendingSubsetTasks);
+
+		CombinedResults.SetNum(NumSubsets);
+		for (int32 SubsetIndex = 0; SubsetIndex < NumSubsets; ++SubsetIndex)
+		{
+			TArray<FCombinedMeshLOD>& LODChain = SubsetMeshLODChains[SubsetIndex];
+			BuildOutputSubAssembly(LODChain, OrderedSubsetIDs[SubsetIndex], FirstVoxWrappedIndex, CombinedResults[SubsetIndex]);
+		}
+
+	}
+	else
+	{
+		bool bRemoveHiddenFaces = CVarGeometryCombineMeshInstancesRemoveHidden.GetValueOnAnyThread() > 0;
+		auto LODRemoveHidden = [&PerLODOptions, FirstVoxWrappedIndex](int32 LODIndex) { return PerLODOptions[LODIndex].RemoveHiddenFacesMethod != IGeometryProcessing_CombineMeshInstances::ERemoveHiddenFacesMode::None && LODIndex < FirstVoxWrappedIndex; };
+		auto LODMergeCoplanar = [&PerLODOptions](int32 LODIndex) { return PerLODOptions[LODIndex].bMergeCoplanarFaces; };
+		auto LODPlanarRetriangulate = [&PerLODOptions](int32 LODIndex)
+			{
+				return PerLODOptions[LODIndex].bMergeCoplanarFaces && PerLODOptions[LODIndex].bPlanarPolygonRetriangulation;
+			};
+		auto LODRemoveHiddenFaces = [&PerLODOptions](int32 LODIndex)
+			{
+				return PerLODOptions[LODIndex].RemoveHiddenFacesMethod;
+			};
+		auto LODRemoveHiddenSamplingDensity = [&PerLODOptions](int32 LODIndex)
+			{
+				return PerLODOptions[LODIndex].RemoveHiddenSamplingDensity;
+			};
+		auto LODDoubleSidedHiddenRemoval = [&PerLODOptions](int32 LODIndex)
+			{
+				return PerLODOptions[LODIndex].bDoubleSidedHiddenRemoval;
+			};
+		ProcessCombinedLODChain(MeshLODs, OptimizationTolerances, FirstVoxWrappedIndex,
+			NumLODs,
+			bRemoveHiddenFaces,
+			LODRemoveHidden,
+			LODMergeCoplanar,
+			LODPlanarRetriangulate,
+			LODRemoveHiddenFaces,
+			LODRemoveHiddenSamplingDensity,
+			LODDoubleSidedHiddenRemoval,
+			AllLODOptions.TriangleGroupingIDFunc,
+			AllLODOptions.CoarseLODStrategy,
+			AllLODOptions.CoarseApproximationDetailSize,
+			[&TriBudget](int32 LODIndex) { return TriBudget[LODIndex]; },
+			AllLODOptions.CoarseLODBaseTolerance,
+			AllLODOptions.HardNormalAngleDeg,
+			AllLODOptions.bAutoGenerateMissingUVs,
+			AllLODOptions.bAutoGenerateTangents,
+			&PreserveTopologyMaterialIDSet
+		);
+
+		CombinedResults.SetNum(1);
+		BuildOutputSubAssembly(MeshLODs, 0, FirstVoxWrappedIndex, CombinedResults[0]);
+	}
+
+}
 
 
 
@@ -3781,7 +4488,6 @@
 void BuildCombinedCollisionShapes(
 	const FMeshPartsAssembly& Assembly,
 	TArray<int32> SubsetIDsOrdering,
-	IGeometryProcessing_CombineMeshInstances::FOptions CombineOptions,
 	TArray<FSimpleShapeSet3d>& CombinedCollisionShapes)
 {
 	int32 NumParts = Assembly.Parts.Num();
@@ -3917,7 +4623,333 @@
 	}
 }
 
+// Common compute part meshes code -- to be called after source LOD meshes are populated
+static void ComputeSinglePartMeshSet_Helper(
+	const IGeometryProcessing_CombineMeshInstances::FComputePartMeshesOptions& Options,
+	const IGeometryProcessing_CombineMeshInstances::FComputePartMeshesSinglePartOptions& PartOptions,
+	IGeometryProcessing_CombineMeshInstances::FSinglePartMeshSet& Result)
+{
+	using namespace UE::Geometry;
+	const int32 NumSimplifiedLODs = Options.NumSimplifiedLODs;
+	const int32 NumApproxLODs = Options.ApproximationTriCosts.Num();
+	const double AngleThresholdDeg = Options.HardNormalAngleDeg;
 
+	int32 NumSourceLODs = Result.Source.Num();
 
+	const FDynamicMesh3* SimplificationSourceMesh = Options.SimplificationSourceLOD < Result.Source.Num() ?
+		&Result.Source[Options.SimplificationSourceLOD] : &Result.Source.Last();
+	FDynamicMeshAABBTree3 SimplificationSourceMeshSpatial(SimplificationSourceMesh, true);
+	Result.Simplified.Reserve(NumSimplifiedLODs);
+	double UseSimplifyTolerance = Options.SimplifyBaseTolerance;
+	for (int32 SimplifyIdx = 0; SimplifyIdx < NumSimplifiedLODs; ++SimplifyIdx)
+	{
+		TUniquePtr<FDynamicMesh3> ToAdd = MakeUnique<FDynamicMesh3>(*SimplificationSourceMesh);
+		FDynamicMesh3& ToSimplify = *ToAdd;
+		SimplifyPartMesh(ToSimplify, UseSimplifyTolerance, AngleThresholdDeg, Options.bSimplifyPreserveCorners,
+			Options.bSimplifyPreserveUVs || PartOptions.bPreserveUVs,
+			Options.bSimplifyPreserveVertexColors, Options.SimplifySharpEdgeAngleDeg, Options.SimplifyMinSalientDimension);
+		UseSimplifyTolerance *= Options.SimplifyLODLevelToleranceScale;
+		// Add the simplification if it's simpler than the lowest-tri source LOD
+		if (ToAdd->TriangleCount() < Result.Source.Last().TriangleCount())
+		{
+			Result.Simplified.Add(ToAdd.Release());
+		}
+	}
+		
+	const FDynamicMesh3* ApproximationSourceMesh = Options.ApproximationSourceLOD < Result.Source.Num() ?
+		&Result.Source[Options.ApproximationSourceLOD] : &Result.Simplified[FMath::Clamp(Options.ApproximationSourceLOD - Result.Source.Num(), 0, Result.Simplified.Num())];
+
+	Result.Approximated.Reserve(NumApproxLODs);
+	TArray<EApproximatePartMethod> SelectedMethodID; SelectedMethodID.Reserve(NumApproxLODs);		// useful for debugging
+	for (int32 ApproxIdx = 0; ApproxIdx < NumApproxLODs; ++ApproxIdx)
+	{
+		double UseTriCost = Options.ApproximationTriCosts[ApproxIdx];
+		Result.Approximated.Add(new FDynamicMesh3(*ApproximationSourceMesh));
+		FDynamicMesh3& ToApprox = Result.Approximated.Last();
+		EApproximatePartMethod UsedMethod;
+		SelectBestFittingMeshApproximation(*ApproximationSourceMesh, SimplificationSourceMeshSpatial,
+			PartOptions.ApproximationConstraint, ToApprox, UsedMethod,
+			Options.SimplifyBaseTolerance, UseTriCost, Options.MaxAllowableApproximationDeviation);
+			
+		// update enabled attribs (is this good?)
+		ToApprox.EnableMatchingAttributes(*ApproximationSourceMesh);
+
+		// recompute normals
+		FMeshNormals::InitializeOverlayTopologyFromOpeningAngle(&ToApprox, ToApprox.Attributes()->PrimaryNormals(), AngleThresholdDeg);
+		FMeshNormals::QuickRecomputeOverlayNormals(ToApprox);
+
+		// stop making approximations if we are already down to a box
+		if (ToApprox.TriangleCount() <= 12)
+		{
+			break;
+		}
+	}
+
+	// Compute error metrics for simplified and approximated meshes
+	Result.SimplifiedMeshErrors.Reserve(Result.Simplified.Num());
+	for (int32 SimplifiedIdx = 0; SimplifiedIdx < Result.Simplified.Num(); ++SimplifiedIdx)
+	{
+		FVector2d Metric = DeviationMetric(Result.Simplified[SimplifiedIdx], SimplificationSourceMeshSpatial);
+		Result.SimplifiedMeshErrors.Add({ Metric[0], Metric[1] });
+	}
+	Result.ApproximatedMeshErrors.Reserve(Result.Approximated.Num());
+	for (int32 ApproxIdx = 0; ApproxIdx < Result.Approximated.Num(); ++ApproxIdx)
+	{
+		FVector2d Metric = DeviationMetric(Result.Approximated[ApproxIdx], SimplificationSourceMeshSpatial);
+		Result.ApproximatedMeshErrors.Add({ Metric[0], Metric[1] });
+	}
+
+	if (!Result.Approximated.IsEmpty() && !Result.Simplified.IsEmpty())
+	{
+		double FirstApproxMetric = FPartApproxSelector::ComputeMetricFromDeviation(
+			FVector2d(Result.ApproximatedMeshErrors[0].AverageError, Result.ApproximatedMeshErrors[0].MaxError), 2, Result.Approximated[0].TriangleCount(), Options.MaxAllowableApproximationDeviation);
+		for (int32 SimplifiedIdx = Result.Simplified.Num() - 1; SimplifiedIdx >= 0; --SimplifiedIdx)
+		{
+			double SimplifiedMetric = FPartApproxSelector::ComputeMetricFromDeviation(
+				FVector2d(Result.SimplifiedMeshErrors[SimplifiedIdx].AverageError, Result.SimplifiedMeshErrors[SimplifiedIdx].MaxError), 1, Result.Simplified[SimplifiedIdx].TriangleCount(), Options.MaxAllowableApproximationDeviation);
+			if (SimplifiedMetric > FirstApproxMetric)
+			{
+				// simplification is worse than approximation, so delete it
+				Result.Simplified.RemoveAt(SimplifiedIdx);
+				Result.SimplifiedMeshErrors.RemoveAt(SimplifiedIdx);
+			}
+			else
+			{
+				break;
+			}
+		}
+	}
+
+	// optionally planar-remesh the requested Source LODs to reduce their triangle count by removing spurious geometry
+	// Note: May invalidate SimplificationSourceMeshSpatial, so we do this last
+	if (Options.bRetriangulateSourceLODs)
+	{
+		for (int32 SourceLODIndex = Options.StartRetriangulateSourceLOD; SourceLODIndex < Result.Source.Num(); ++SourceLODIndex)
+		{
+			if (!(PartOptions.bPreserveUVs || Options.bSimplifyPreserveUVs || (SourceLODIndex <= Options.PreserveUVLODLevel)))
+			{
+				PlanarRetriangulatePartMesh(Result.Source[SourceLODIndex], Options.SimplifyBaseTolerance, AngleThresholdDeg);
+			}
+		}
+	}
+
+	// Enforce that simplified and approximated meshes must have decreasing triangle counts
+	int32 LastTriCount = Result.Source.Last().TriangleCount();
+	auto EnforceDecreasingTriCount = [&LastTriCount](TIndirectArray<FDynamicMesh3>& Meshes, TArray<IGeometryProcessing_CombineMeshInstances::FSinglePartMeshSet::FErrorStats>& MeshErrors)
+	{
+		for (int32 Idx = 0, Shift = 0; Idx < Meshes.Num(); ++Idx)
+		{
+			// Shift forward to the first mesh with smaller tri count
+			while (Idx + Shift < Meshes.Num() && Meshes[Idx + Shift].TriangleCount() >= LastTriCount)
+			{
+				Shift++;
+			}
+
+			// If there are no meshes to copy back, remove the rest of the array
+			if (Idx + Shift >= Meshes.Num())
+			{
+				Meshes.RemoveAtSwap(Idx, Meshes.Num() - Idx);
+				MeshErrors.RemoveAtSwap(Idx, MeshErrors.Num() - Idx);
+				break;
+			}
+
+			// Update the last tri count and shift if needed
+			LastTriCount = Meshes[Idx + Shift].TriangleCount();
+			if (Shift > 0)
+			{
+				Meshes[Idx] = MoveTemp(Meshes[Idx + Shift]);
+				MeshErrors[Idx] = MeshErrors[Idx + Shift];
+			}
+		}
+	};
+	EnforceDecreasingTriCount(Result.Simplified, Result.SimplifiedMeshErrors);
+	EnforceDecreasingTriCount(Result.Approximated, Result.ApproximatedMeshErrors);
+
+}
+
+int32 AddSourceLODs(TIndirectArray<UE::Geometry::FDynamicMesh3>& OutSourceLODs, int32 NumLODs, TFunctionRef<const FMeshDescription*(int)> GetLOD)
+{
+	OutSourceLODs.Reset(NumLODs);
+	for (int32 Idx = 0; Idx < NumLODs; ++Idx)
+	{
+		const FMeshDescription* Source = GetLOD(Idx);
+		if (Source)
+		{
+			UE::Geometry::FDynamicMesh3* OutputLODMesh = new UE::Geometry::FDynamicMesh3();
+			FMeshDescriptionToDynamicMesh Converter;
+			Converter.bEnableOutputGroups = true;
+			Converter.bTransformVertexColorsLinearToSRGB = true;		// possibly this should be false...
+			Converter.Convert(Source, *OutputLODMesh);
+			OutSourceLODs.Add(OutputLODMesh);
+		}
+	}
+	return OutSourceLODs.Num();
+}
+
+void FCombineMeshInstancesImpl::ComputeSinglePartMeshSet(
+	TConstArrayView<const FMeshDescription*> SourceMeshLODs,
+	const FComputePartMeshesOptions& Options,
+	const FComputePartMeshesSinglePartOptions& PartOptions,
+	FSinglePartMeshSet& ResultMeshes)
+{
+	TRACE_CPUPROFILER_EVENT_SCOPE(ComputePartMeshSet_MeshDescriptions);
+
+	int32 NumSources = AddSourceLODs(ResultMeshes.Source, SourceMeshLODs.Num(), [&SourceMeshLODs](int32 LOD) {return SourceMeshLODs[LOD];});
+
+	if (NumSources <= 0)
+	{
+		return;
+	}
+	
+	TRACE_CPUPROFILER_EVENT_SCOPE(ComputePartMeshSet_PartApprox);
+	ComputeSinglePartMeshSet_Helper(Options, PartOptions, ResultMeshes);
+}
+
+void FCombineMeshInstancesImpl::ComputeSinglePartMeshSet(
+	UStaticMesh* SourceMesh,
+	const FComputePartMeshesOptions& Options,
+	const FComputePartMeshesSinglePartOptions& PartOptions,
+	FSinglePartMeshSet& ResultMeshes
+)
+{
+	TRACE_CPUPROFILER_EVENT_SCOPE(ComputePartMeshSet_StaticMesh);
+
+	check(SourceMesh);
+	int32 NumSources = AddSourceLODs(ResultMeshes.Source, SourceMesh->GetNumSourceModels(), [&SourceMesh](int32 LOD) {return SourceMesh->GetMeshDescription(LOD);});
+
+	if (NumSources <= 0)
+	{
+		return;
+	}
+
+	TRACE_CPUPROFILER_EVENT_SCOPE(ComputePartMeshSet_PartApprox);
+	ComputeSinglePartMeshSet_Helper(Options, PartOptions, ResultMeshes);
+}
+
+void FCombineMeshInstancesImpl::ComputePartMeshSets(
+	FSourceInstanceList& SourceInstanceList,
+	const FComputePartMeshesOptions& Options,
+	bool bKeepExistingPartMeshes,
+	TArray<FSinglePartMeshSet>& ResultMeshSets
+)
+{
+	TRACE_CPUPROFILER_EVENT_SCOPE(ComputePartMeshSet_Instances);
+
+	TMap<UStaticMesh*, int32> StaticMeshToPartMap;
+	TMap<const IGeometryProcessing_CombineMeshInstances::FMeshLODSet*, int32> MeshLODSetToPartMap;
+
+	int32 NumStaticMeshInstances = SourceInstanceList.StaticMeshInstances.Num();
+	int32 NumMeshLODSetInstances = SourceInstanceList.MeshLODSetInstances.Num();
+
+	if (!bKeepExistingPartMeshes)
+	{
+		ResultMeshSets.Empty();
+		for (int32 Index = 0; Index < NumStaticMeshInstances; ++Index)
+		{
+			IGeometryProcessing_CombineMeshInstances::FStaticMeshInstance& SourceMeshInstance = SourceInstanceList.StaticMeshInstances[Index];
+			SourceMeshInstance.PrecomputeMeshesIndex = -1;
+		}
+		for (int32 Index = 0; Index < NumMeshLODSetInstances; ++Index)
+		{
+			IGeometryProcessing_CombineMeshInstances::FMeshLODSetInstance& SourceMeshInstance = SourceInstanceList.MeshLODSetInstances[Index];
+			SourceMeshInstance.PrecomputeMeshesIndex = -1;
+		}
+	}
+	else
+	{
+		for (int32 Index = 0; Index < NumStaticMeshInstances; ++Index)
+		{
+			const IGeometryProcessing_CombineMeshInstances::FStaticMeshInstance& SourceMeshInstance = SourceInstanceList.StaticMeshInstances[Index];
+			if (SourceMeshInstance.PrecomputeMeshesIndex != -1)
+			{
+				UStaticMesh* StaticMesh = SourceMeshInstance.SourceMesh;
+				StaticMeshToPartMap.Add(StaticMesh, SourceMeshInstance.PrecomputeMeshesIndex);
+			}
+		}
+		for (int32 Index = 0; Index < NumMeshLODSetInstances; ++Index)
+		{
+			const IGeometryProcessing_CombineMeshInstances::FMeshLODSetInstance& SourceMeshInstance = SourceInstanceList.MeshLODSetInstances[Index];
+			if (SourceMeshInstance.PrecomputeMeshesIndex != -1)
+			{
+				int32 MeshSetIndex = SourceMeshInstance.MeshLODSetIndex;
+				if (MeshSetIndex < 0 || MeshSetIndex > SourceInstanceList.MeshLODSets.Num())
+				{
+					ensure(false);
+					continue;
+				}
+				const IGeometryProcessing_CombineMeshInstances::FMeshLODSet* MeshLODSet = &SourceInstanceList.MeshLODSets[MeshSetIndex];
+				MeshLODSetToPartMap.Add(MeshLODSet, SourceMeshInstance.PrecomputeMeshesIndex);
+			}
+		}
+	}
+
+	
+
+	// TODO: compute these in parallel?
+	for (int32 Index = 0; Index < NumStaticMeshInstances; ++Index)
+	{
+		IGeometryProcessing_CombineMeshInstances::FStaticMeshInstance& SourceMeshInstance = SourceInstanceList.StaticMeshInstances[Index];
+		if (SourceMeshInstance.PrecomputeMeshesIndex != -1)
+		{
+			continue;
+		}
+
+		UStaticMesh* StaticMesh = SourceMeshInstance.SourceMesh;
+		int32* FoundIdx = StaticMeshToPartMap.Find(StaticMesh);
+		int32 LocalIdx = -1;
+		if (FoundIdx == nullptr)
+		{
+			LocalIdx = ResultMeshSets.Emplace();
+			EApproximationType ApproximationConstraint = EApproximationType::NoConstraint;
+			bool bPreservePartUVs = false;
+			if (SourceInstanceList.InstanceGroupDatas.IsValidIndex(SourceMeshInstance.GroupDataIndex))
+			{
+				ApproximationConstraint = SourceInstanceList.InstanceGroupDatas[SourceMeshInstance.GroupDataIndex].ApproximationConstraint;
+				bPreservePartUVs = SourceInstanceList.InstanceGroupDatas[SourceMeshInstance.GroupDataIndex].bPreserveUVs;
+			}
+			IGeometryProcessing_CombineMeshInstances::FComputePartMeshesSinglePartOptions PartOptions(ApproximationConstraint, bPreservePartUVs);
+			ComputeSinglePartMeshSet(StaticMesh, Options, PartOptions, ResultMeshSets[LocalIdx]);
+			FoundIdx = &LocalIdx;
+			StaticMeshToPartMap.Add(StaticMesh, LocalIdx);
+		}
+		SourceMeshInstance.PrecomputeMeshesIndex = *FoundIdx;
+	}
+
+	// TODO: compute these in parallel?
+	for (int32 Index = 0; Index < NumMeshLODSetInstances; ++Index)
+	{
+		IGeometryProcessing_CombineMeshInstances::FMeshLODSetInstance& SourceMeshInstance = SourceInstanceList.MeshLODSetInstances[Index];
+		if (SourceMeshInstance.PrecomputeMeshesIndex != -1)
+		{
+			continue;
+		}
+
+		int32 MeshSetIndex = SourceMeshInstance.MeshLODSetIndex;
+		if (MeshSetIndex < 0 || MeshSetIndex > SourceInstanceList.MeshLODSets.Num())
+		{
+			ensure(false);
+			continue;
+		}
+		const IGeometryProcessing_CombineMeshInstances::FMeshLODSet* MeshLODSet = &SourceInstanceList.MeshLODSets[MeshSetIndex];
+
+		int32* FoundIdx = MeshLODSetToPartMap.Find(MeshLODSet);
+		int32 LocalIdx = -1;
+		if (FoundIdx == nullptr)
+		{
+			LocalIdx = ResultMeshSets.Emplace();
+			EApproximationType ApproximationConstraint = EApproximationType::NoConstraint;
+			if (SourceInstanceList.InstanceGroupDatas.IsValidIndex(SourceMeshInstance.GroupDataIndex))
+			{
+				ApproximationConstraint = SourceInstanceList.InstanceGroupDatas[SourceMeshInstance.GroupDataIndex].ApproximationConstraint;
+			}
+			IGeometryProcessing_CombineMeshInstances::FComputePartMeshesSinglePartOptions PartOptions(ApproximationConstraint);
+			ComputeSinglePartMeshSet(MeshLODSet->ReferencedMeshLODs, Options, PartOptions, ResultMeshSets[LocalIdx]);
+			FoundIdx = &LocalIdx;
+			MeshLODSetToPartMap.Add(MeshLODSet, LocalIdx);
+		}
+		SourceMeshInstance.PrecomputeMeshesIndex = *FoundIdx;
+	}
+}
+
 void FCombineMeshInstancesImpl::CombineMeshInstances(
 	const FSourceInstanceList& MeshInstances, const FOptions& Options, FResults& ResultsOut)
@@ -4003,11 +5035,11 @@
 		}
 	}
 
-	// build combined collision shapes, repsecting sub-assembly ordering
+	// build combined collision shapes, respecting sub-assembly ordering
 	TArray<FSimpleShapeSet3d> CombinedCollisionShapes;
 	{
 		TRACE_CPUPROFILER_EVENT_SCOPE(CombineMeshInst_BuildCollision);
-		BuildCombinedCollisionShapes(PartAssembly, SubsetIDsOrdering, Options, CombinedCollisionShapes);
+		BuildCombinedCollisionShapes(PartAssembly, SubsetIDsOrdering, CombinedCollisionShapes);
 
 		if (bVerbose && CombinedCollisionShapes.Num() == 1)
 		{
@@ -4038,3 +5070,141 @@
 		ResultsOut.CombinedMeshes[k].InstanceSubsetID = CombinedResults[k].SubAssemblyID;
 	}
 }
+
+
+void FCombineMeshInstancesImpl::CombineMeshInstances(
+	const FSourceInstanceList& MeshInstances,
+	TConstArrayView<FSinglePartMeshSet> CachedPartMeshes,
+	const FCombineMeshInstancesOptionsGeneral& AllLODOptions,
+	TConstArrayView<FCombineMeshInstancesOptionsPerLOD> PerLODOptions,
+	FResults& ResultsOut
+)
+{
+	TRACE_CPUPROFILER_EVENT_SCOPE(CombineMeshInstances);
+
+	bool bVerbose = CVarGeometryCombineMeshInstancesVerbose.GetValueOnAnyThread();
+	
+	FMeshPartsAssembly PartAssembly;
+
+	{
+		TRACE_CPUPROFILER_EVENT_SCOPE(CombineMeshInst_Setup);
+		InitializeMeshPartAssembly(MeshInstances, PartAssembly);
+		if (bVerbose)
+		{
+			UE_LOG(LogGeometry, Log, TEXT("  PartAssembly contains %d Parts, %d Unique Materials"),
+				PartAssembly.Parts.Num(), PartAssembly.UniqueMaterials.Num());
+		}
+
+
+		// require that all instances have precomputed part meshes
+		for (const TUniquePtr<FMeshPart>& Part : PartAssembly.Parts)
+		{
+			if (Part->PrecomputeMeshesIndex == INDEX_NONE)
+			{
+				UE_LOG(LogGeometry, Warning, TEXT("Failed to compute CombineMeshInstances because some instances did not have precomputed part meshes"));
+				return;
+			}
+		}
+
+		// Initialize just the collision shapes on source mesh geometry
+		PartAssembly.SourceMeshGeometry.SetNum(PartAssembly.Parts.Num());
+		for (int32 Index = 0; Index < PartAssembly.Parts.Num(); ++Index)
+		{
+			TUniquePtr<FMeshPart>& Part = PartAssembly.Parts[Index];
+			FSourceGeometry& Target = PartAssembly.SourceMeshGeometry[Index];
+
+			if (UStaticMesh* StaticMesh = Part->SourceAsset)
+			{
+				if (UBodySetup* BodySetup = StaticMesh->GetBodySetup())
+				{
+					UE::Geometry::GetShapeSet(BodySetup->AggGeom, Target.CollisionShapes);
+				}
+			}
+			else if (Part->SourceMeshLODSet != nullptr)
+			{
+				UE::Geometry::GetShapeSet(Part->SourceMeshLODSet->SimpleCollisionShapes, Target.CollisionShapes);
+			}
+
+			// sometimes simple collision is a convex when it's actually a box - could try to detect here?
+
+		}
+	}
+
+	PartAssembly.PreProcessInstanceMeshFunc = [&PartAssembly, &MeshInstances, &AllLODOptions](FDynamicMesh3& AppendMesh, const FMeshPartInstance& Instance)
+	{
+		int32 SourceInstance = Instance.ExternalInstanceIndex.A;
+		bool bIsStaticMeshInstance = (Instance.ExternalInstanceIndex.B == 0);		// a bit of a hack here but we configured this above
+
+		int GroupDataIdx = (bIsStaticMeshInstance) ?
+			MeshInstances.StaticMeshInstances[SourceInstance].GroupDataIndex
+			: MeshInstances.MeshLODSetInstances[SourceInstance].GroupDataIndex;
+
+		if (MeshInstances.InstanceGroupDatas[GroupDataIdx].bHasConstantOverrideVertexColor)
+		{
+			FLinearColor VertexColorLinear(0, 0, 0, 1);
+			if (AllLODOptions.VertexColorMappingMode == EVertexColorMappingMode::TriangleCountMetric)
+			{
+				const double UseMax = 25.0;
+				double TriCountRelToBox = FMath::Clamp((double)AppendMesh.TriangleCount() / (double)12, 1.0, UseMax);		// 12 is num tris in a bounding box
+				double T = (TriCountRelToBox) / (UseMax);
+				T = FMathd::Sqrt(T);	// improve color mapping somewhat (try better options?)
+				VertexColorLinear = FLinearColor::LerpUsingHSV(
+					FLinearColor::White, FLinearColor::Red, FMath::Clamp(T, 0.0, 1.0));
+			}
+			else
+			{
+				FColor VertexColorSRGB = MeshInstances.InstanceGroupDatas[GroupDataIdx].OverrideVertexColor;
+				VertexColorLinear = VertexColorSRGB.ReinterpretAsLinear();
+			}
+
+			SetConstantVertexColor(AppendMesh, VertexColorLinear);
+		}
+	};
+
+	// build combined mesh LOD chains for each sub-assembly
+	TArray<FCombinedSubAssembly> CombinedResults;
+	TArray<int32> SubsetIDsOrdering;
+	{
+		TRACE_CPUPROFILER_EVENT_SCOPE(CombineMeshInst_BuildMeshes);
+		BuildCombinedMeshFromPrecomputedMeshes(CachedPartMeshes, PartAssembly, AllLODOptions, PerLODOptions, CombinedResults);
+		for (FCombinedSubAssembly& Assembly : CombinedResults)
+		{
+			SubsetIDsOrdering.Add(Assembly.SubAssemblyID);
+		}
+	}
+
+	// build combined collision shapes, respecting sub-assembly ordering
+	TArray<FSimpleShapeSet3d> CombinedCollisionShapes;
+	{
+		TRACE_CPUPROFILER_EVENT_SCOPE(CombineMeshInst_BuildCollision);
+		BuildCombinedCollisionShapes(PartAssembly, SubsetIDsOrdering, CombinedCollisionShapes);
+
+		if (bVerbose && CombinedCollisionShapes.Num() == 1)
+		{
+			UE_LOG(LogGeometry, Log, TEXT("  CombinedCollisionShapes[0] contains %d Boxes, %d Convexes"),
+				CombinedCollisionShapes[0].Boxes.Num(), CombinedCollisionShapes[0].Convexes.Num());
+		}
+	}
+
+
+	// build final results data structure
+	ResultsOut.CombinedMeshes.SetNum(CombinedResults.Num());
+	for (int32 k = 0; k < CombinedResults.Num(); ++k)
+	{
+		ResultsOut.CombinedMeshes[k].MeshLODs = MoveTemp(CombinedResults[k].MeshLODs);
+
+		FPhysicsDataCollection PhysicsData;
+		PhysicsData.Geometry = CombinedCollisionShapes[k];
+		PhysicsData.CopyGeometryToAggregate();		// need FPhysicsDataCollection to convert to agg geom, should fix this
+		ResultsOut.CombinedMeshes[k].SimpleCollisionShapes = PhysicsData.AggGeom;
+
+		ResultsOut.CombinedMeshes[k].MaterialSet = PartAssembly.UniqueMaterials;
+		if (CombinedResults.Num() > 1)
+		{
+			// if we have multiple outputs, they may not each use the full material set, in this case we will filter the materials (should this be optional?)
+			FilterUnusedMaterials(ResultsOut.CombinedMeshes[k].MeshLODs, ResultsOut.CombinedMeshes[k].MaterialSet);
+		}
+
+		ResultsOut.CombinedMeshes[k].InstanceSubsetID = CombinedResults[k].SubAssemblyID;
+	}
+}
diff --git a/Engine/Plugins/Experimental/MeshModelingToolsetExp/Source/GeometryProcessingAdapters/Public/GeometryProcessing/CombineMeshInstancesImpl.h b/Engine/Plugins/Experimental/MeshModelingToolsetExp/Source/GeometryProcessingAdapters/Public/GeometryProcessing/CombineMeshInstancesImpl.h
--- a/Engine/Plugins/Experimental/MeshModelingToolsetExp/Source/GeometryProcessingAdapters/Public/GeometryProcessing/CombineMeshInstancesImpl.h
+++ b/Engine/Plugins/Experimental/MeshModelingToolsetExp/Source/GeometryProcessingAdapters/Public/GeometryProcessing/CombineMeshInstancesImpl.h
@@ -24,6 +24,37 @@ public:
 
 	virtual void CombineMeshInstances(const FSourceInstanceList& MeshInstances, const FOptions& Options, FResults& ResultsOut) override;
 
+	virtual void CombineMeshInstances(
+		const FSourceInstanceList& MeshInstances,
+		TConstArrayView<FSinglePartMeshSet> CachedPartMeshes,
+		const FCombineMeshInstancesOptionsGeneral& AllLODOptions,
+		TConstArrayView<FCombineMeshInstancesOptionsPerLOD> PerLODOptions,
+		FResults& ResultsOut
+	) override;
+
+
+	virtual void ComputeSinglePartMeshSet(
+		TConstArrayView<const FMeshDescription*> SourceMeshLODs,
+		const FComputePartMeshesOptions& GeneralOptions,
+		const FComputePartMeshesSinglePartOptions& PartOptions,
+		FSinglePartMeshSet& ResultMeshes
+	) override;
+
+	virtual void ComputeSinglePartMeshSet(
+		UStaticMesh* SourceMesh,
+		const FComputePartMeshesOptions& GeneralOptions,
+		const FComputePartMeshesSinglePartOptions& PartOptions,
+		FSinglePartMeshSet& ResultPartMeshSet
+	) override;
+
+	virtual void ComputePartMeshSets(
+		FSourceInstanceList& MeshInstances,
+		const FComputePartMeshesOptions& GeneralOptions,
+		bool bKeepExistingPartMeshes,
+		TArray<FSinglePartMeshSet>& ResultMeshSets
+	) override;
+
+
 protected:
 
 };
diff --git a/Engine/Source/Developer/GeometryProcessingInterfaces/Public/GeometryProcessingInterfaces/CombineMeshInstances.h b/Engine/Source/Developer/GeometryProcessingInterfaces/Public/GeometryProcessingInterfaces/CombineMeshInstances.h
--- a/Engine/Source/Developer/GeometryProcessingInterfaces/Public/GeometryProcessingInterfaces/CombineMeshInstances.h
+++ b/Engine/Source/Developer/GeometryProcessingInterfaces/Public/GeometryProcessingInterfaces/CombineMeshInstances.h
@@ -44,6 +44,30 @@ public:
 		All = 0xFFFF
 	};
 
+	// Set of versions of a single part mesh that can be used in the combined mesh, including source LODs, simplified versions, and approximated versions
+	// Can be re-used across multiple CombineMeshInstances calls
+	struct FSinglePartMeshSet
+	{
+		struct FErrorStats
+		{
+			double AverageError;
+			double MaxError;
+		};
+
+		// Source LOD meshes
+		TIndirectArray<UE::Geometry::FDynamicMesh3> Source;
+		// Meshes generated by running simplification on a source LOD
+		TIndirectArray<UE::Geometry::FDynamicMesh3> Simplified;
+		// Error metrics for simplified meshes
+		TArray<FErrorStats> SimplifiedMeshErrors;
+		// Meshes generated by an approximation method (e.g. convex hull, sweep, etc), either from a source LOD or a simplified mesh
+		TIndirectArray<UE::Geometry::FDynamicMesh3> Approximated;
+		// Error metrics for simplified meshes
+		TArray<FErrorStats> ApproximatedMeshErrors;
+
+		// Note: Materials can vary per-instance so are not stored with the part mesh set
+	};
+
 
 	/**
 	 * FMeshInstanceGroupData is data shared among one or more FBaseMeshInstances.
@@ -81,6 +105,11 @@ public:
 
 		int32 FilterLODLevel = -1;					// LOD level to filter out this mesh, value -1 disable this option. 
 
+		// Index of precomputed meshes for this instance, w.r.t. a pre-computed part meshes array
+		// Note the pre-computed meshes array should be stored separately from the instance data,
+		// and is provided as an argument to CombineMeshInstances.
+		int32 PrecomputeMeshesIndex = -1;
+
 		// in some cases it may be desirable to have "groups" of instances which should be output as separate meshes, but
 		// be jointly processed in terms of (eg) the part LODs. If any InstanceSubsetID is non-zero, then instance subsets
 		// are grouped/extracted by integer ID and will be returned as separate FOutputMesh's in the FResults. 
@@ -313,6 +342,222 @@ public:
 		EVertexColorMappingMode VertexColorMappingMode = EVertexColorMappingMode::None;
 	};
 
+	enum class ETriangleBudgetMethod : uint8
+	{
+		/** Use each LOD as-is with no optimizations to the bricks themselves other than general purpose mesh optimizations. */
+		NoRestriction,
+		/** Try to fit the explicit budget. Note that success in hitting the budget depends on the MeshOptimizationMethod. */
+		UseTriangleBudget,
+		/** Try to fit a percentage of the prior budget. */
+		UsePercentageOfPreviousLOD
+	};
+
+	enum class EMeshOptimizationMethod : uint8
+	{
+		/** No invasive part replacement to be used. Rely on the LODs defined in the bricks. */
+		None,
+		/** Go through the bricks and find the most expensive brick, try to approximate it using some geometry matching algorithms or decimation. Repeat until budget goal is met or all bricks have been handled. */
+		SimplifyOrApproximate,
+		/** Go through the bricks and find the most expensive brick, try to approximate it using some geometry matching algorithms. Repeat until budget goal is met or all bricks have been handled. */
+		ApproximateOnly,
+		/** Voxelize (or otherwise apply very coarse approximation) and decimate to hit budget goals. */
+		VoxelizeAndDecimate
+	};
+
+	struct FTriangleBudgetOptions
+	{
+		/** Determine what the source for the budget should be. */
+		ETriangleBudgetMethod Method = ETriangleBudgetMethod::NoRestriction;
+
+		/** The explicit budget to use if Method == UseTriangleBudget */
+		int32 TriangleBudget = 0;
+
+		/** The percentage of the prior LOD to use as the triangle budget if Method == UsePercentageOfPreviousLOD */
+		float LODReductionPercentage = 0.75f;
+	};
+
+	enum class EDecorationHandling : uint8
+	{
+		Keep,
+		Approximate,
+		Remove
+	};
+
+	// Options that can be set separately per LOD
+	struct FCombineMeshInstancesOptionsPerLOD
+	{
+		/** What is the desired source mesh LOD levels to use for each brick. If not available, it will find the closest to this index. */
+		int32 PreferredLOD = 0;
+
+		/** The optimization method that can be used to generate this LOD */
+		EMeshOptimizationMethod OptimizationMethod = EMeshOptimizationMethod::None;
+
+		/** How decoration should be treated for this LOD (e.g., kept, approximated, or removed). */
+		EDecorationHandling Decorations = EDecorationHandling::Keep;
+
+		/** Tolerance to use for simplification of the LOD (if simplification is applied) */
+		double SimplificationTolerance = 1.0;
+
+		/** If two triangles are co-planar, determine if we can merge them  */
+		bool bMergeCoplanarFaces = false;
+
+		/** Whether UV coordinates must be preserved for this LOD, or can be fully discarded */
+		bool bPreserveUVs = true;
+
+		// Specifies whether to remove hidden faces, and if so method to use
+		ERemoveHiddenFacesMode RemoveHiddenFacesMethod = ERemoveHiddenFacesMode::None;
+		// (approximately) spacing between samples on triangle faces used for determining exterior visibility
+		double RemoveHiddenSamplingDensity = 1.0;
+		// treat faces as double-sided for hidden removal
+		bool bDoubleSidedHiddenRemoval = false;
+
+		// Whether to attempt extraction and retriangulation of detected planar polygonal areas, after removing hidden faces.
+		// This occurs after coplanar face merging, and may further merge adjacent coplanar faces.
+		// Can significantly reduce triangle count, but attributes on the interiors of polygonal areas will be completely discarded
+		bool bPlanarPolygonRetriangulation = false;
+
+		// Whether a part promotion strategy should be used to attempt to hit triangle budgets
+		bool bEnableBudgetStrategy_PartLODPromotion = false;
+		// If promoting parts to worse approximations to fit a triangle budget, scale the triangle budget by this amount (to account for potential reductions from e.g. hidden triangle removal, that will happen after part meshes are chosen and merged)
+		double PartLODPromotionBudgetMultiplier = 2.0;
+
+		// The budget settings for this LOD
+		FTriangleBudgetOptions TriangleBudgetOptions;
+	};
+
+	// Options that apply across all LODs
+	struct FCombineMeshInstancesOptionsGeneral
+	{
+
+		// TriangleGroupingIDFunc allows client to specify external adjacency relationships between triangles via an integer ID tuple.
+		// Adjacent triangles will only be considered for Coplanar Merging if they have the same FIndex3i ID.
+		TFunction<UE::Geometry::FIndex3i(const UE::Geometry::FDynamicMesh3& Mesh, int32 TriangleID)> TriangleGroupingIDFunc;
+
+		// Triangles with Materials assigned that are also in this set will not be allowed to be combined/retriangulated
+		// with adjacent triangles. This can be used to preserve topology/UVs/etc for specific material areas.
+		TArray<UMaterialInterface*> PreventMergingMaterialSet;
+
+		//
+		// Common settings for Coarse/Voxel LODs. Coarse LODs should be the lowest-quality (highest-index) LODs and generally assumed to have 
+		// very low triangle counts. Two strategies are supported:
+		//   1) VoxelBasedSolidApproximation - does a solid approximation with topological closure (outset + inset), which will fill in any small holes/gaps
+		//      or small features. CoarseApproximationDetailSize is used as the closure radius/distance. 
+		//   2) SweptPlanarProjection - The mesh is effectively projected to 2D along each X/Y/Z axis, with 2D polygon boolean & topological closure using CoarseApproximationDetailSize,
+		//      as well as various other cleanup, and then that 2D polygon is triangulated and swept along the part bounding-box extent. The max distance
+		//      is measured from the approximation to the original mesh, and the axis-swept-mesh with the smallest max-distance is used.
+		// 
+		// The 'Automatic' strategy uses the SweptPlanarProjection if it's max approximation deviation is within K*CoarseApproximationDetailSize (currently K=2),
+		// and otherwise falls back to using VoxelBasedSolidApproximation
+		// 
+		// Note the *first* coarse LOD will use the specified strategy, and subsequent LODs will apply mesh simplification to that result.
+		// 
+		ECoarseApproximationStrategy CoarseLODStrategy = ECoarseApproximationStrategy::Automatic;
+		double CoarseLODBaseTolerance = 1.0;
+		double CoarseApproximationDetailSize = 10.0;
+		// opening angle used to detect/assign sharp edges
+		double HardNormalAngleDeg = 15;
+
+		//
+		// Final processing options applied to each LOD after generation
+		//
+
+		// ensure that all triangles have some UV values set. Various geometry processing steps may discard UVs, this flag will
+		// cause missing UVs to be recomputed (currently using a box projection)
+		bool bAutoGenerateMissingUVs = false;
+
+		// generate tangents on the output meshes - requires that UVs be available
+		bool bAutoGenerateTangents = false;
+
+
+		//
+		// Debug/utility options
+		// 
+
+		// Color mapping modes for vertex colors, primarily used for debugging
+		EVertexColorMappingMode VertexColorMappingMode = EVertexColorMappingMode::None;
+
+		// Initialize options from the combined FOptions struct, including settings up per-LOD options
+		void Init(const FOptions& InOptions, TArray<FCombineMeshInstancesOptionsPerLOD>& OutPerLODOptions)
+		{
+			// General
+
+			TriangleGroupingIDFunc = InOptions.TriangleGroupingIDFunc;
+			PreventMergingMaterialSet = InOptions.PreventMergingMaterialSet;
+			CoarseLODStrategy = InOptions.CoarseLODStrategy;
+			CoarseLODBaseTolerance = InOptions.CoarseLODBaseTolerance;
+			CoarseApproximationDetailSize = InOptions.CoarseApproximationDetailSize;
+			HardNormalAngleDeg = InOptions.HardNormalAngleDeg;
+			bAutoGenerateMissingUVs = InOptions.bAutoGenerateMissingUVs;
+			bAutoGenerateTangents = InOptions.bAutoGenerateTangents;
+			VertexColorMappingMode = InOptions.VertexColorMappingMode;
+
+			// Per LOD
+
+			OutPerLODOptions.SetNum(InOptions.NumLODs);
+			double SimplificationTolerance = InOptions.SimplifyBaseTolerance;
+			for (int32 LODIndex = 0; LODIndex < InOptions.NumLODs; ++LODIndex)
+			{
+				FCombineMeshInstancesOptionsPerLOD& Options = OutPerLODOptions[LODIndex];
+				Options.PreferredLOD = LODIndex;
+				Options.OptimizationMethod = EMeshOptimizationMethod::SimplifyOrApproximate;
+				if (LODIndex < InOptions.NumCopiedLODs)
+				{
+					Options.OptimizationMethod = EMeshOptimizationMethod::None;
+				}
+				else if (LODIndex < InOptions.NumCopiedLODs + InOptions.NumSimplifiedLODs)
+				{
+					Options.OptimizationMethod = EMeshOptimizationMethod::SimplifyOrApproximate;
+				}
+				else if (LODIndex >= InOptions.NumLODs - InOptions.NumCoarseLODs)
+				{
+					Options.OptimizationMethod = EMeshOptimizationMethod::VoxelizeAndDecimate;
+				}
+				Options.bDoubleSidedHiddenRemoval = InOptions.bDoubleSidedHiddenRemoval;
+				Options.RemoveHiddenFacesMethod = InOptions.RemoveHiddenFacesMethod;
+				Options.RemoveHiddenSamplingDensity = InOptions.RemoveHiddenSamplingDensity;
+				Options.bMergeCoplanarFaces = InOptions.bMergeCoplanarFaces;
+				Options.bPlanarPolygonRetriangulation = InOptions.bMergeCoplanarFaces && (InOptions.PlanarPolygonRetriangulationStartLOD > -1 && LODIndex >= InOptions.PlanarPolygonRetriangulationStartLOD);
+				Options.Decorations = EDecorationHandling::Keep;
+				if (LODIndex >= InOptions.FilterDecorativePartsLODLevel - InOptions.ApproximateDecorativePartLODs)
+				{
+					Options.Decorations = EDecorationHandling::Approximate;
+				}
+				else if (LODIndex >= InOptions.FilterDecorativePartsLODLevel)
+				{
+					Options.Decorations = EDecorationHandling::Remove;
+				}
+				Options.SimplificationTolerance = SimplificationTolerance;
+				if (Options.OptimizationMethod != EMeshOptimizationMethod::None)
+				{
+					SimplificationTolerance *= InOptions.SimplifyLODLevelToleranceScale;
+				}
+				Options.bPreserveUVs = LODIndex <= InOptions.PreserveUVLODLevel;
+
+				Options.bEnableBudgetStrategy_PartLODPromotion = InOptions.bEnableBudgetStrategy_PartLODPromotion;
+				Options.PartLODPromotionBudgetMultiplier = InOptions.PartLODPromotionBudgetMultiplier;
+				Options.TriangleBudgetOptions.Method = ETriangleBudgetMethod::NoRestriction;
+				if (Options.OptimizationMethod != EMeshOptimizationMethod::VoxelizeAndDecimate)
+				{
+					if (InOptions.HardLODBudgets.IsValidIndex(LODIndex))
+					{
+						int32 HardBudget = InOptions.HardLODBudgets[LODIndex];
+						if (HardBudget > -1)
+						{
+							Options.TriangleBudgetOptions.Method = ETriangleBudgetMethod::UseTriangleBudget;
+							Options.TriangleBudgetOptions.TriangleBudget = HardBudget;
+						}
+					}
+				}
+				else
+				{
+					Options.TriangleBudgetOptions.Method = LODIndex > InOptions.NumLODs - InOptions.NumCoarseLODs ? ETriangleBudgetMethod::UsePercentageOfPreviousLOD : ETriangleBudgetMethod::UseTriangleBudget;
+					Options.TriangleBudgetOptions.TriangleBudget = InOptions.CoarseLODMaxTriCountBase;
+					Options.TriangleBudgetOptions.LODReductionPercentage = .5;
+				}
+			}
+		}
+	};
+
 
 	struct FOutputMesh
 	{
@@ -321,7 +566,7 @@ public:
 
 		FKAggregateGeom SimpleCollisionShapes;
 
-		// All part instances accumulated into the MeshLODs will have had this InstanceSubsetID in their input FBaseMeshInstance's
+		// Each part instance accumulated into the MeshLODs will have had this InstanceSubsetID in their input FBaseMeshInstance
 		int32 InstanceSubsetID = 0;
 	};
 
@@ -345,12 +590,168 @@ public:
 
 	virtual void CombineMeshInstances(
 		const FSourceInstanceList& MeshInstances,
-		const FOptions& Options, 
-		FResults& ResultsOut) 
+		const FOptions& Options,
+		FResults& ResultsOut)
 	{
 		check(false);		// not implemented in base class
 	}
 
+	// General options for compute part meshes, which may be typically kept consistent across parts
+	struct FComputePartMeshesOptions
+	{
+		//
+		// Settings for Simplifed LODs
+		//
+
+		// which LOD to use as basis for Simplifications. Index is clamped to source LODs
+		int32 SimplificationSourceLOD = 0;
+		int32 NumSimplifiedLODs = 3;
+		double SimplifyBaseTolerance = 1.0;
+		double SimplifyLODLevelToleranceScale = 2.0;
+		// if true, UVs will be preserved in Simplified LODs. This generally will result in lower-quality geometric shape approximation.
+		bool bSimplifyPreserveUVs = false;
+		// if true, vertex colors will be preserved in Simplified LODs. This generally will result in lower-quality geometric shape approximation.
+		bool bSimplifyPreserveVertexColors = false;
+		// if true, geometrically-detected "sharp" corners (eg like the corners of a box) will be preserved with hard constraints during simplification when creating Simplified LODs. This can be desirable on mechanical/geometric shapes.
+		bool bSimplifyPreserveCorners = true;
+		double SimplifySharpEdgeAngleDeg = 44.0;
+		double SimplifyMinSalientDimension = 1.0;
+
+		// Output mesh edges with larger dihedral angle than this threshold will be assigned sharp (split) normals
+		double HardNormalAngleDeg = 15.0;
+
+		//
+		// settings for Approximate LODs
+		//
+
+		// which LOD to use as basis for Approximations. Index is interpreted relative to [CopiedLODs...SimplifiedLODs] set, ie can be set to a Simplified LOD
+		int32 ApproximationSourceLOD = 0;
+		double MaxAllowableApproximationDeviation = 5.0;
+		// 'Triangle cost' to use for each approximation mesh; should be increasing.
+		// The length of this array determines the number of approximations meshes to attempt to make. Note we may have fewer meshes, since we discard approximations that do not decrease triangle count.
+		TArray<double> ApproximationTriCosts;
+
+		// This helper configures ApproximationTriCosts, translating the automatic strategy used by the FOptions-based code to an explicit array of tri costs
+		void SetApproximationLODs(int32 NumInitialApproximations, double OptimizeBaseTriCost, double InitialTriCostScale, int32 NumAdditionalApproximations = 10, double AdditionalTriCostAdd = 0.25)
+		{
+			int32 NumApproximatedLODs = NumInitialApproximations + NumAdditionalApproximations;
+			ApproximationTriCosts.Reset(NumInitialApproximations + NumAdditionalApproximations);
+			double TriCost = OptimizeBaseTriCost;
+			ApproximationTriCosts.Add(TriCost);
+			for (int32 Idx = 1; Idx < NumInitialApproximations; ++Idx)
+			{
+				TriCost *= OptimizeBaseTriCost;
+				ApproximationTriCosts.Add(TriCost);
+			}
+			while (ApproximationTriCosts.Num() < NumApproximatedLODs)
+			{
+				TriCost += AdditionalTriCostAdd;
+				ApproximationTriCosts.Add(TriCost);
+			}
+		}
+
+		// If enabled, attempt to retriangulate planar areas of Source LODs to remove redundant coplanar geometry. 
+		// This option affects individual parts and not the combined prefab.
+		bool bRetriangulateSourceLODs = true;
+		// which Source LOD to start planar retriangulation at
+		int32 StartRetriangulateSourceLOD = 1;
+
+		// UVs on input geometry will be preserved up to this LOD level (inclusive). 
+		// Note that this setting will severely constrain and/or fully disable many other optimizations. 
+		// In particular, Coplanar merging and retriangulation cannot be applied if UVs are to be preserved. 
+		int PreserveUVLODLevel = -1;
+
+		// Initialize part option settings from the larger FOptions struct
+		void Init(const FOptions& Options)
+		{
+			SimplificationSourceLOD = Options.ApproximationSourceLOD;
+			NumSimplifiedLODs = Options.NumSimplifiedLODs;
+			SimplifyBaseTolerance = Options.SimplifyBaseTolerance;
+			SimplifyLODLevelToleranceScale = Options.SimplifyLODLevelToleranceScale;
+			bSimplifyPreserveUVs = Options.bSimplifyPreserveUVs;
+			bSimplifyPreserveVertexColors = Options.bSimplifyPreserveVertexColors;
+			bSimplifyPreserveCorners = Options.bSimplifyPreserveCorners;
+			SimplifySharpEdgeAngleDeg = Options.SimplifySharpEdgeAngleDeg;
+			SimplifyMinSalientDimension = Options.SimplifyMinSalientDimension;
+			HardNormalAngleDeg = Options.HardNormalAngleDeg;
+			ApproximationSourceLOD = Options.ApproximationSourceLOD;
+			SetApproximationLODs(FMath::Max(0, Options.NumLODs - Options.NumCopiedLODs - Options.NumSimplifiedLODs - Options.NumCoarseLODs), Options.OptimizeBaseTriCost, Options.OptimizeLODLevelTriCostScale);
+			MaxAllowableApproximationDeviation = Options.MaxAllowableApproximationDeviation;
+			bRetriangulateSourceLODs = Options.bRetriangulateSourceLODs;
+			StartRetriangulateSourceLOD = Options.StartRetriangulateSourceLOD;
+			PreserveUVLODLevel = Options.PreserveUVLODLevel;
+		}
+	};
+
+	// Options that are specific to a given part / more likely to be changed per-part
+	struct FComputePartMeshesSinglePartOptions
+	{
+		FComputePartMeshesSinglePartOptions() = default;
+		FComputePartMeshesSinglePartOptions(EApproximationType ApproximationConstraint, bool bPreserveUVs = false) : ApproximationConstraint(ApproximationConstraint), bPreserveUVs(bPreserveUVs) {}
+
+		// Whether to constrain this part to only use specific approximation methods
+		EApproximationType ApproximationConstraint = EApproximationType::NoConstraint;
+
+		// Whether to always preserve UVs for this part. If false, UVs may still be preserved based on the general settings
+		bool bPreserveUVs = false;
+	};
+
+	// Generate a set of coarse simplification and/or approximation LODs from a source array of mesh descriptions per LOD
+	// @param SourceMeshLODs Source mesh LODs as mesh descriptions
+	// @param GeneralOptions Sets the number and settings for simplified and approximated meshes
+	// @param PartOptions Options more likely to vary on a per-part basis
+	// @param ResultMeshes Dynamic meshes of the source LODs as well as simplified and approximated versions
+	virtual void ComputeSinglePartMeshSet(
+		TConstArrayView<const FMeshDescription*> SourceMeshLODs,
+		const FComputePartMeshesOptions& GeneralOptions,
+		const FComputePartMeshesSinglePartOptions& PartOptions, 
+		FSinglePartMeshSet& ResultMeshes
+	)
+	{
+		check(false); // not implemented in base class
+	}
+
+	// Generate a set of coarse simplification and/or approximation LODs from a source static mesh
+	// @param SourceMeshLODs Source mesh LODs as mesh descriptions
+	// @param GeneralOptions Sets the number and settings for simplified and approximated meshes
+	// @param PartOptions Options more likely to vary on a per-part basis
+	// @param ResultMeshes Dynamic meshes of the source LODs as well as simplified and approximated versions
+	virtual void ComputeSinglePartMeshSet(
+		UStaticMesh* SourceMesh,
+		const FComputePartMeshesOptions& GeneralOptions,
+		const FComputePartMeshesSinglePartOptions& PartOptions,
+		FSinglePartMeshSet& ResultPartMeshSet
+	)
+	{
+		check(false); // not implemented in base class
+	}
+
+	// Generate a set of coarse simplification and/or approximation LODs for all source meshes in the given MeshInstances
+	// @param MeshInstances Instances to generate part meshes for. Indices into the ResultMeshSets array will also be assigned per instance
+	// @param Options Sets the number and settings for simplified and approximated meshes
+	// @param bKeepExistingPartMeshes If false will clear the existing part mesh cache, otherwise will only generate part meshes for those instance that do not already have them.
+	// @param ResultMeshSets Computed mesh sets
+	virtual void ComputePartMeshSets(
+		FSourceInstanceList& MeshInstances,
+		const FComputePartMeshesOptions& Options,
+		bool bKeepExistingPartMeshes,
+		TArray<FSinglePartMeshSet>& ResultMeshSets
+	)
+	{
+		check(false); // not implemented in base class
+	}
+
+	// A version of CombineMeshInstances that takes per-LOD settings and uses pre-computed part meshes (computed via ComputePartMeshSets, above)
+	virtual void CombineMeshInstances(
+		const FSourceInstanceList& MeshInstances,
+		TConstArrayView<FSinglePartMeshSet> CachedPartMeshes,
+		const FCombineMeshInstancesOptionsGeneral& AllLODOptions,
+		TConstArrayView<FCombineMeshInstancesOptionsPerLOD> PerLODOptions,
+		FResults& ResultsOut
+	)
+	{
+		check(false);		// not implemented in base class
+	}
 
 
 	// Modular feature name to register for retrieval during runtime
