Skeleton Modifier: restore retargeting modes when commiting changes

#rnx

diff --git a/Engine/Plugins/Experimental/MeshModelingToolsetExp/Source/SkeletalMeshModifiers/Private/SkeletonModifier.cpp b/Engine/Plugins/Experimental/MeshModelingToolsetExp/Source/SkeletalMeshModifiers/Private/SkeletonModifier.cpp
--- a/Engine/Plugins/Experimental/MeshModelingToolsetExp/Source/SkeletalMeshModifiers/Private/SkeletonModifier.cpp
+++ b/Engine/Plugins/Experimental/MeshModelingToolsetExp/Source/SkeletalMeshModifiers/Private/SkeletonModifier.cpp
@@ -365,12 +365,26 @@ bool USkeletonModifier::CommitSkeletonToSkeletalMesh()
 		}
 	}
 
+	// store retargeting modes
+	USkeleton* Skeleton = SkeletalMesh->GetSkeleton();
+	
+	TArray<EBoneTranslationRetargetingMode::Type> RetargetingModes;
+	RetargetingModes.Init(EBoneTranslationRetargetingMode::Animation, BoneInfos.Num());
+	
+	for (int32 OldBoneIndex = 0, NumOldBones = BoneIndexTracker.Num(); OldBoneIndex < NumOldBones; ++OldBoneIndex)
+	{
+		const int32 NewBoneIndex = BoneIndexTracker[OldBoneIndex];
+		if (RetargetingModes.IsValidIndex(NewBoneIndex))
+		{
+			RetargetingModes[NewBoneIndex] = Skeleton->GetBoneTranslationRetargetingMode(OldBoneIndex);	
+		}
+	}
+
 	// update skeletal mesh
 	FlushRenderingCommands();
 
 	// call modify on the skeleton first as post undo will re-register components so it must be done once both
 	// skeletal mesh and skeleton are up to date, so it must be done after the skeletal mesh has been undone 
-	USkeleton* Skeleton = SkeletalMesh->GetSkeleton();
 	Skeleton->Modify();
 
 	SkeletalMesh->SetFlags(RF_Transactional);
@@ -389,7 +403,13 @@ bool USkeletonModifier::CommitSkeletonToSkeletalMesh()
 	// update skeleton
 	if (Skeleton->RecreateBoneTree(SkeletalMesh.Get()))
 	{
-		Skeleton->MarkPackageDirty();	
+		// restore retargeting modes
+		for (int32 BoneIndex = 0, NumBones = BoneInfos.Num(); BoneIndex < NumBones; ++BoneIndex)
+		{
+			Skeleton->SetBoneTranslationRetargetingMode(BoneIndex, RetargetingModes[BoneIndex]);
+		}
+		
+		Skeleton->MarkPackageDirty();
 	}
 	
 	// must be done once the skeleton is up to date
