Skelmesh Editing: Added a shortcut for transfering weights between LODs/profiles on the same mesh.

diff --git a/Engine/Plugins/Experimental/Animation/SkeletalMeshModelingTools/Source/SkeletalMeshModelingTools/Private/DetailCustomization/SkinWeightDetailCustomization.cpp b/Engine/Plugins/Experimental/Animation/SkeletalMeshModelingTools/Source/SkeletalMeshModelingTools/Private/DetailCustomization/SkinWeightDetailCustomization.cpp
--- a/Engine/Plugins/Experimental/Animation/SkeletalMeshModelingTools/Source/SkeletalMeshModelingTools/Private/DetailCustomization/SkinWeightDetailCustomization.cpp
+++ b/Engine/Plugins/Experimental/Animation/SkeletalMeshModelingTools/Source/SkeletalMeshModelingTools/Private/DetailCustomization/SkinWeightDetailCustomization.cpp
@@ -1181,7 +1181,8 @@ void FSkinWeightDetailCustomization::AddTransferUI(IDetailLayoutBuilder& DetailB
 				.HAlign(HAlign_Center)
 				.Text(LOCTEXT("TransferWeightsButtonLabel", "Transfer Weights"))
 				.ToolTipText(LOCTEXT("TransferButtonTooltip",
-					"Weights are transfered from the source skeletal mesh using inpainting.\n"
+					"Weights are transferred from the source skeletal mesh using in-painting.\n"
+					"If no source skeletal mesh is specified, weights can be transferred from the existing mesh between profiles and LODs.\n"
 					"This command can operate on selected components when in Mesh mode."))
 				.OnClicked_Lambda([this]()
 				{
@@ -1190,7 +1191,7 @@ void FSkinWeightDetailCustomization::AddTransferUI(IDetailLayoutBuilder& DetailB
 				})
 				.IsEnabled_Lambda([this]()
 				{
-					return Tool->GetSourceTarget() != nullptr;
+					return Tool->CanTransferWeights();
 				})
 			]
 		]
diff --git a/Engine/Plugins/Experimental/MeshModelingToolsetExp/Source/MeshModelingToolsEditorOnlyExp/Private/SkeletalMesh/SkinWeightsPaintTool.cpp b/Engine/Plugins/Experimental/MeshModelingToolsetExp/Source/MeshModelingToolsEditorOnlyExp/Private/SkeletalMesh/SkinWeightsPaintTool.cpp
--- a/Engine/Plugins/Experimental/MeshModelingToolsetExp/Source/MeshModelingToolsEditorOnlyExp/Private/SkeletalMesh/SkinWeightsPaintTool.cpp
+++ b/Engine/Plugins/Experimental/MeshModelingToolsetExp/Source/MeshModelingToolsEditorOnlyExp/Private/SkeletalMesh/SkinWeightsPaintTool.cpp
@@ -3052,139 +3052,197 @@ void USkinWeightsPaintTool::TransferWeights()
 	using namespace UE::Geometry;
 	using namespace UE::AnimationCore;
 
-	if (!SourceTarget)
-	{
-		const FText NotificationText = LOCTEXT("NoSourceInTransfer", "No source skeletal mesh was specified. No weights were transferred.");
-		ShowEditorMessage(ELogVerbosity::Error, NotificationText);
-		return;
-	}
-
 	const EMeshLODIdentifier TargetLODId = GetLODId(WeightToolProperties->ActiveLOD);
     const FGetMeshParameters TargetParams(true, TargetLODId);
 	FDynamicMesh3 TargetMesh = UE::ToolTarget::GetDynamicMeshCopy(Target, TargetParams);
 
 	const EMeshLODIdentifier SourceLODId = GetLODId(WeightToolProperties->SourceLOD);
-	const FGetMeshParameters SourceParams(true, SourceLODId);
-	const FDynamicMesh3 SourceMesh = UE::ToolTarget::GetDynamicMeshCopy(SourceTarget, SourceParams);
-	
-	if (!SourceMesh.HasAttributes() || !SourceMesh.Attributes()->HasBones())
-	{
-    	const FText NotificationText = LOCTEXT("NoWeightsFoundInTransfer", "No skin weights were found in the source skeletal mesh. No weights were transferred.");
-    	ShowEditorMessage(ELogVerbosity::Error, NotificationText);
-		return;
-	}
-	
-	if (SourceMesh.Attributes()->GetNumBones() == 0)
-	{
-		const FText NotificationText = LOCTEXT("NoBonesFoundInTransfer", "No bones were found in the source skeletal mesh. No weights were transferred.");
-		ShowEditorMessage(ELogVerbosity::Error, NotificationText);
-		return;
-	}
-	
-	FTransferBoneWeights TransferBoneWeights(&SourceMesh, WeightToolProperties->SourceSkinWeightProfile);
-	TransferBoneWeights.TransferMethod = FTransferBoneWeights::ETransferBoneWeightsMethod::InpaintWeights;
+
+	const FName TargetProfile = WeightToolProperties->GetActiveSkinWeightProfile();
+	FDynamicMeshVertexSkinWeightsAttribute* TransferredSkinWeights = nullptr;
+
+	TArray<int32> VertexSubset;
+	bool bCreateTransaction = false;
 	
-	if (!TargetMesh.HasAttributes())
+	// If a source mesh was given, or if the LODs differ, then do a weight transfer using Rinat's algo.
+	// Otherwise, do an exact copy from a different profile.
+	if (SourceTarget || TargetLODId != SourceLODId)
 	{
-		TargetMesh.EnableAttributes();
-	}
+		const FGetMeshParameters SourceParams(true, SourceLODId);
+		const FDynamicMesh3 SourceMesh = UE::ToolTarget::GetDynamicMeshCopy(SourceTarget ? SourceTarget : Target, SourceParams);
+		
+		if (!SourceMesh.HasAttributes() || !SourceMesh.Attributes()->HasBones())
+		{
+			const FText NotificationText = LOCTEXT("NoWeightsFoundInTransfer", "No skin weights were found in the source skeletal mesh. No weights were transferred.");
+			ShowEditorMessage(ELogVerbosity::Error, NotificationText);
+			return;
+		}
+		
+		if (SourceMesh.Attributes()->GetNumBones() == 0)
+		{
+			const FText NotificationText = LOCTEXT("NoBonesFoundInTransfer", "No bones were found in the source skeletal mesh. No weights were transferred.");
+			ShowEditorMessage(ELogVerbosity::Error, NotificationText);
+			return;
+		}
+		
+		FTransferBoneWeights TransferBoneWeights(&SourceMesh, WeightToolProperties->SourceSkinWeightProfile);
+		TransferBoneWeights.TransferMethod = FTransferBoneWeights::ETransferBoneWeightsMethod::InpaintWeights;
+		
+		if (!TargetMesh.HasAttributes())
+		{
+			TargetMesh.EnableAttributes();
+		}
 
-	FDynamicMeshAttributeSet* TargetAttributes = TargetMesh.Attributes();
-	if (!TargetAttributes->HasBones())
-	{
-		TargetAttributes->CopyBoneAttributes(*SourceMesh.Attributes());
+		FDynamicMeshAttributeSet* TargetAttributes = TargetMesh.Attributes();
+		if (!TargetAttributes->HasBones())
+		{
+			TargetAttributes->CopyBoneAttributes(*SourceMesh.Attributes());
+		}
+		else
+		{
+			const USkeletalMeshComponent* Component = GetSkeletalMeshComponent(Target);
+			if (ensure(Component))
+			{
+				const FReferenceSkeleton& RefSkeleton = Component->GetSkeletalMeshAsset()->GetRefSkeleton();
+				ensure(TargetAttributes->GetNumBones() == RefSkeleton.GetRawBoneNum());
+			}
+		}
+
+		// NOTE should we expose all the options?
+		// 	TransferBoneWeights.NormalThreshold;
+		// 	TransferBoneWeights.SearchRadius
+		// 	TransferBoneWeights.NumSmoothingIterations;
+		// 	TransferBoneWeights.SmoothingStrength;
+		// 	TransferBoneWeights.LayeredMeshSupport;
+		// 	TransferBoneWeights.ForceInpaintWeightMapName;
+		
+		if (TransferBoneWeights.Validate() != EOperationValidationResult::Ok)
+		{
+			return;
+		}
+
+		if (WeightToolProperties->EditingMode == EWeightEditMode::Mesh)
+		{
+			TransferBoneWeights.TargetVerticesSubset = GetSelectedVertices();
+		}
+
+		if (!TransferBoneWeights.TransferWeightsToMesh(TargetMesh, TargetProfile))
+		{
+			// notify user
+			const FText NotificationText = LOCTEXT("TransferWeightsUnknownIssue", "Transfer weights operation encountered an unknown issue. No weights were transferred.");
+			ShowEditorMessage(ELogVerbosity::Error, NotificationText);
+			return;
+		}
+		
+		TransferredSkinWeights = TargetAttributes->GetSkinWeightsAttribute(TargetProfile);
+		
+		VertexSubset = MoveTemp(TransferBoneWeights.TargetVerticesSubset);
+		bCreateTransaction = true;
 	}
 	else
 	{
-		const USkeletalMeshComponent* Component = GetSkeletalMeshComponent(Target);
-		if (ensure(Component))
+		FDynamicMeshAttributeSet* TargetAttributes = TargetMesh.Attributes();
+		TransferredSkinWeights = TargetAttributes->GetSkinWeightsAttribute(TargetProfile);
+		if (!TransferredSkinWeights)
 		{
-			const FReferenceSkeleton& RefSkeleton = Component->GetSkeletalMeshAsset()->GetRefSkeleton();
-			ensure(TargetAttributes->GetNumBones() == RefSkeleton.GetRawBoneNum());
+			TransferredSkinWeights = new FDynamicMeshVertexSkinWeightsAttribute(&TargetMesh);
+			TargetMesh.Attributes()->AttachSkinWeightsAttribute(TargetProfile, TransferredSkinWeights);
 		}
-	}
+		
+		const FDynamicMeshVertexSkinWeightsAttribute* SourceAttributes = TargetAttributes->GetSkinWeightsAttribute(WeightToolProperties->SourceSkinWeightProfile);
+		check(SourceAttributes != nullptr);
 
-	// NOTE should we expose all the options?
-	// 	TransferBoneWeights.NormalThreshold;
-	// 	TransferBoneWeights.SearchRadius
-	// 	TransferBoneWeights.NumSmoothingIterations;
-	// 	TransferBoneWeights.SmoothingStrength;
-	// 	TransferBoneWeights.LayeredMeshSupport;
-	// 	TransferBoneWeights.ForceInpaintWeightMapName;
-	
-	if (TransferBoneWeights.Validate() != EOperationValidationResult::Ok)
-	{
-		return;
-	}
+		if (SourceAttributes != TransferredSkinWeights)
+		{
+			if (WeightToolProperties->EditingMode == EWeightEditMode::Mesh)
+			{
+				VertexSubset = GetSelectedVertices();
+			}
 
-	if (WeightToolProperties->EditingMode == EWeightEditMode::Mesh)
-	{
-		TransferBoneWeights.TargetVerticesSubset = GetSelectedVertices();
+			if (VertexSubset.IsEmpty())
+			{
+				TransferredSkinWeights->Copy(*SourceAttributes);
+			}
+			else
+			{
+				for (int32 VertexID: VertexSubset)
+				{
+					FBoneWeights BoneWeights;
+					SourceAttributes->GetValue(VertexID, BoneWeights);
+					TransferredSkinWeights->SetValue(VertexID, BoneWeights);
+				}
+			}
+			bCreateTransaction = true;
+		}
 	}
 
-	const FName TargetProfile = WeightToolProperties->GetActiveSkinWeightProfile();
-	if (!TransferBoneWeights.TransferWeightsToMesh(TargetMesh, TargetProfile))
+	if (bCreateTransaction)
 	{
-		// notify user
-		const FText NotificationText = LOCTEXT("TransferWeightsUnknownIssue", "Transfer weights operation encountered an unknown issue. No weights were transferred.");
-		ShowEditorMessage(ELogVerbosity::Error, NotificationText);
-		return;
-	}
-	
-	// weight edits for transaction
-	FMultiBoneWeightEdits WeightEdits;
-
-	// spin through all the transferred skin weights and record a weight edit to apply as a transaction
-	FDynamicMeshVertexSkinWeightsAttribute* TransferedSkinWeights = TargetAttributes->GetSkinWeightsAttribute(TargetProfile);
-	check(TransferedSkinWeights);
-	static constexpr float ZeroWeight = 0.f;
-	const bool bUseSubset = !TransferBoneWeights.TargetVerticesSubset.IsEmpty();
-	const int32 NumVertices = bUseSubset ? TransferBoneWeights.TargetVerticesSubset.Num() : TargetMesh.VertexCount();
-	const FNonManifoldMappingSupport NonManifoldMappingSupport(TargetMesh);
-	for (int32 VertexIndex = 0; VertexIndex < NumVertices; ++VertexIndex)
-	{
-		const int32 VertexID = bUseSubset ? TransferBoneWeights.TargetVerticesSubset[VertexIndex] : VertexIndex;
-		const int32 SrcVertexID = NonManifoldMappingSupport.GetOriginalNonManifoldVertexID(VertexID);
+		// weight edits for transaction
+		FMultiBoneWeightEdits WeightEdits;
 		
-		// remove all weight on vertex
-		const VertexWeights& VertexBoneWeights = Weights.PreChangeWeights[SrcVertexID];
-		if (!VertexBoneWeights.IsEmpty())
+		// spin through all the transferred skin weights and record a weight edit to apply as a transaction
+		check(TransferredSkinWeights);
+		static constexpr float ZeroWeight = 0.f;
+		const bool bUseSubset = !VertexSubset.IsEmpty();
+		const int32 NumVertices = bUseSubset ? VertexSubset.Num() : TargetMesh.VertexCount();
+		const FNonManifoldMappingSupport NonManifoldMappingSupport(TargetMesh);
+		for (int32 VertexIndex = 0; VertexIndex < NumVertices; ++VertexIndex)
 		{
-			for (const FVertexBoneWeight& BoneWeight : VertexBoneWeights)
+			const int32 VertexID = bUseSubset ? VertexSubset[VertexIndex] : VertexIndex;
+			const int32 SrcVertexID = NonManifoldMappingSupport.GetOriginalNonManifoldVertexID(VertexID);
+			
+			// remove all weight on vertex
+			const VertexWeights& VertexBoneWeights = Weights.PreChangeWeights[SrcVertexID];
+			if (!VertexBoneWeights.IsEmpty())
 			{
-				const float OldWeight = BoneWeight.Weight;
-				WeightEdits.MergeSingleEdit(BoneWeight.BoneID, SrcVertexID, OldWeight, ZeroWeight);
+				for (const FVertexBoneWeight& BoneWeight : VertexBoneWeights)
+				{
+					const float OldWeight = BoneWeight.Weight;
+					WeightEdits.MergeSingleEdit(BoneWeight.BoneID, SrcVertexID, OldWeight, ZeroWeight);
+				}
+			}
+			else
+			{
+				WeightEdits.MergeSingleEdit(0, SrcVertexID, 1.f, ZeroWeight);
 			}
-		}
-		else
-		{
-			WeightEdits.MergeSingleEdit(0, SrcVertexID, 1.f, ZeroWeight);
-		}
 
-		// update with new weight
-		FBoneWeights TransferedBoneWeights;
-		TransferedSkinWeights->GetValue(VertexID, TransferedBoneWeights);
-		for (const FBoneWeight& BoneWeight: TransferedBoneWeights)
-		{
-			const int32 BoneIndex = BoneWeight.GetBoneIndex();					
-			const float OldWeight = Weights.GetWeightOfBoneOnVertex(BoneIndex, SrcVertexID, Weights.PreChangeWeights);
-			const float NewWeight = BoneWeight.GetWeight();
-			WeightEdits.MergeSingleEdit(BoneIndex, SrcVertexID, OldWeight, NewWeight);
+			// update with new weight
+			FBoneWeights TransferedBoneWeights;
+			TransferredSkinWeights->GetValue(VertexID, TransferedBoneWeights);
+			for (const FBoneWeight& BoneWeight: TransferedBoneWeights)
+			{
+				const int32 BoneIndex = BoneWeight.GetBoneIndex();					
+				const float OldWeight = Weights.GetWeightOfBoneOnVertex(BoneIndex, SrcVertexID, Weights.PreChangeWeights);
+				const float NewWeight = BoneWeight.GetWeight();
+				WeightEdits.MergeSingleEdit(BoneIndex, SrcVertexID, OldWeight, NewWeight);
+			}
 		}
-	}
 
-	// apply the changes as a transaction
-	const FText TransactionLabel = LOCTEXT("TransferWeightsChange", "Transfer skin weights.");
-	ApplyWeightEditsAsTransaction(WeightEdits, TransactionLabel);
-
-	// put the mesh back in it's current pose
-	Weights.Deformer.SetAllVerticesToBeUpdated();
+		// apply the changes as a transaction
+		const FText TransactionLabel = LOCTEXT("TransferWeightsChange", "Transfer skin weights.");
+		ApplyWeightEditsAsTransaction(WeightEdits, TransactionLabel);
+		
+		// put the mesh back in it's current pose
+		Weights.Deformer.SetAllVerticesToBeUpdated();
+	}
 
 	// notify user that weights were transferred.
 	const FText NotificationText = LOCTEXT("WeightsTransferred", "Skin weights transferred.");
 	ShowEditorMessage(ELogVerbosity::Log, NotificationText);
 }
 
+bool USkinWeightsPaintTool::CanTransferWeights() const
+{
+	if (GetSourceTarget() != nullptr)
+	{
+		return true; 
+	}
+
+	return GetLODId(WeightToolProperties->ActiveLOD) != GetLODId(WeightToolProperties->SourceLOD) ||
+		   WeightToolProperties->GetActiveSkinWeightProfile() != WeightToolProperties->SourceSkinWeightProfile; 		
+}
+
 const FString USkinWeightsPaintTool::CopyPasteWeightsIdentifier = TEXT("UNREAL_VERTEX_WEIGHTS:");
 
 void USkinWeightsPaintTool::CopyWeights()
diff --git a/Engine/Plugins/Experimental/MeshModelingToolsetExp/Source/MeshModelingToolsEditorOnlyExp/Public/SkeletalMesh/SkinWeightsPaintTool.h b/Engine/Plugins/Experimental/MeshModelingToolsetExp/Source/MeshModelingToolsEditorOnlyExp/Public/SkeletalMesh/SkinWeightsPaintTool.h
--- a/Engine/Plugins/Experimental/MeshModelingToolsetExp/Source/MeshModelingToolsEditorOnlyExp/Public/SkeletalMesh/SkinWeightsPaintTool.h
+++ b/Engine/Plugins/Experimental/MeshModelingToolsetExp/Source/MeshModelingToolsEditorOnlyExp/Public/SkeletalMesh/SkinWeightsPaintTool.h
@@ -502,6 +502,7 @@ public:
 	void NormalizeWeights();
 	void HammerWeights();
 	void TransferWeights();
+	bool CanTransferWeights() const;
 
 	// copy paste
 	void CopyWeights();
