Skin Weight Editor: Fixes ensure and unnormalized weights when using Redo after Undo'ing a Prune

diff --git a/Engine/Plugins/Experimental/MeshModelingToolsetExp/Source/MeshModelingToolsEditorOnlyExp/Private/SkeletalMesh/SkinWeightsPaintTool.cpp b/Engine/Plugins/Experimental/MeshModelingToolsetExp/Source/MeshModelingToolsEditorOnlyExp/Private/SkeletalMesh/SkinWeightsPaintTool.cpp
--- a/Engine/Plugins/Experimental/MeshModelingToolsetExp/Source/MeshModelingToolsEditorOnlyExp/Private/SkeletalMesh/SkinWeightsPaintTool.cpp
+++ b/Engine/Plugins/Experimental/MeshModelingToolsetExp/Source/MeshModelingToolsEditorOnlyExp/Private/SkeletalMesh/SkinWeightsPaintTool.cpp
@@ -404,15 +404,64 @@ FSkinWeightBrushConfig& USkinWeightsPaintToolProperties::GetBrushConfig()
 void FMultiBoneWeightEdits::MergeSingleEdit(
 	const int32 BoneIndex,
 	const int32 VertexID,
-	const float OldWeight,
-	const float NewWeight)
+	const float NewWeight,
+	bool bPruneInfluence,
+	const TArray<VertexWeights>& InPreChangeWeights)
 {
-	ensure(BoneIndex!=INDEX_NONE);
+	if (!ensure(BoneIndex!=INDEX_NONE))
+	{
+		return;
+	}
 	
+	if (bPruneInfluence)
+	{
+		// should never be pruning an influence while also trying to add weight to it
+		if (!ensure(FMath::IsNearlyEqual(NewWeight, 0.f)))
+		{
+			return;
+		}
+	}
+
+	// get the old weight of this influence and check whether it was already influences this vertex
+	float OldWeight = 0.f;
+	bool bWasAlreadyAnInfluence = false;
+	const VertexWeights& VertexWeights = InPreChangeWeights[VertexID];
+	for (const FVertexBoneWeight& BoneWeight : VertexWeights)
+	{
+		if (BoneWeight.BoneID == BoneIndex)
+		{
+			OldWeight = BoneWeight.Weight;
+			bWasAlreadyAnInfluence = true;
+			break;
+		}
+	}
+
+	// don't prune influence unless they are actually connected to the vertex
+	if (bPruneInfluence && !bWasAlreadyAnInfluence)
+	{
+		bPruneInfluence = false;
+	}
+
+	// record a modification of this vertex weight
 	FSingleBoneWeightEdits& BoneWeightEdit = PerBoneWeightEdits.FindOrAdd(BoneIndex);
 	BoneWeightEdit.BoneIndex = BoneIndex;
 	BoneWeightEdit.NewWeights.Add(VertexID, NewWeight);
 	BoneWeightEdit.OldWeights.FindOrAdd(VertexID, OldWeight);
+
+	if (bPruneInfluence)
+	{
+		// record that this influence was REMOVED from this vertex
+		BoneWeightEdit.VerticesRemovedFrom.AddUnique(VertexID);
+		// remove from the AddedTo list in the off chance that vertex was added by a prior edit (last edit takes precedence)
+		BoneWeightEdit.VerticesAddedTo.Remove(VertexID);
+	}
+	else if (!bWasAlreadyAnInfluence && NewWeight > MinimumWeightThreshold)
+	{
+		// record that this influence was ADDED to this vertex
+		BoneWeightEdit.VerticesAddedTo.AddUnique(VertexID);
+		// remove from the pruned list in the off chance that vertex was pruned by a prior edit (last edit takes precedence)
+		BoneWeightEdit.VerticesRemovedFrom.Remove(VertexID);
+	}
 }
 
 void FMultiBoneWeightEdits::MergeEdits(const FSingleBoneWeightEdits& BoneWeightEdits)
@@ -453,13 +502,6 @@ void FMultiBoneWeightEdits::AddEditedVerticesToSet(TSet<int32>& OutEditedVertexS
 	}
 }
 
-void FMultiBoneWeightEdits::AddPruneBoneEdit(
-	const VertexIndex VertexToPruneFrom,
-	const BoneIndex BoneToPrune)
-{
-	PrunedInfluences.Emplace(VertexToPruneFrom, BoneToPrune);
-}
-
 void FSkinToolDeformer::Initialize(const USkeletalMeshComponent* InSkelMeshComponent, const FMeshDescription* InMeshDescription)
 {
 	// get all bone transforms in the reference pose store a copy in component space
@@ -728,6 +770,9 @@ void FSkinToolWeights::CreateWeightEditForVertex(
 	float NewWeightValue,
 	FMultiBoneWeightEdits& WeightEdits)
 {
+	// this operation should never prune weights
+	constexpr bool bPruneInfluence = false;
+	
 	// clamp new weight
 	NewWeightValue = FMath::Clamp(NewWeightValue, 0.0f, 1.0f);
 
@@ -757,21 +802,15 @@ void FSkinToolWeights::CreateWeightEditForVertex(
 	if (FMath::IsNearlyEqual(NewWeightValue, 1.f))
 	{
 		// in this case normalization is trivial, just assign the full weight directly and zero all others
-		const float PrevWeight = GetWeightOfBoneOnVertex(BoneToHoldIndex, VertexID, PreChangeWeights);
 		constexpr float FullWeight = 1.0f;
-		WeightEdits.MergeSingleEdit(
-			BoneToHoldIndex,
-			VertexID,
-			PrevWeight,
-			FullWeight);
+		WeightEdits.MergeSingleEdit(BoneToHoldIndex, VertexID, FullWeight, bPruneInfluence, PreChangeWeights);
 
 		// zero all others
 		for (int32 i=0; i<ValuesToNormalize.Num(); ++i)
 		{
 			const int32 BoneIndex = RecordedBonesOnVertex[i];
-			const float OldWeight = ValuesToNormalize[i];
 			constexpr float NewWeight = 0.f;
-			WeightEdits.MergeSingleEdit(BoneIndex, VertexID, OldWeight, NewWeight);
+			WeightEdits.MergeSingleEdit(BoneIndex, VertexID, NewWeight, bPruneInfluence, PreChangeWeights);
 		}
 
 		return;
@@ -802,18 +841,12 @@ void FSkinToolWeights::CreateWeightEditForVertex(
 			for (int32 i=0; i<ValuesToNormalize.Num(); ++i)
 			{
 				const int32 BoneIndex = RecordedBonesOnVertex[i];
-				const float OldWeight = ValuesToNormalize[i];
 				const float NewWeight = WeightToDistribute;
-				WeightEdits.MergeSingleEdit(BoneIndex, VertexID, OldWeight, NewWeight);
+				WeightEdits.MergeSingleEdit(BoneIndex, VertexID, NewWeight, bPruneInfluence, PreChangeWeights);
 			}
 
 			// set current bone value to user assigned weight
-			const float PrevWeight = GetWeightOfBoneOnVertex(BoneToHoldIndex, VertexID, PreChangeWeights);
-			WeightEdits.MergeSingleEdit(
-				BoneToHoldIndex,
-				VertexID,
-				PrevWeight,
-				NewWeightValue);
+			WeightEdits.MergeSingleEdit(BoneToHoldIndex, VertexID, NewWeightValue, bPruneInfluence, PreChangeWeights);
 		}
 		else
 		{
@@ -830,54 +863,44 @@ void FSkinToolWeights::CreateWeightEditForVertex(
 				// this could only happen if user is trying to remove weight from the ONLY bone in the whole skeleton
 				// in this case just assign the full weight to the bone (there's no other valid configuration)
 				// this is a "do nothing" operation, but at least it generates an undo transaction to let user know the input was received
-				const float PrevWeight = GetWeightOfBoneOnVertex(BoneToHoldIndex, VertexID, PreChangeWeights);
 				constexpr float FullWeight = 1.0f;
-				WeightEdits.MergeSingleEdit(
-					BoneToHoldIndex,
-					VertexID,
-					PrevWeight,
-					FullWeight);
+				WeightEdits.MergeSingleEdit(BoneToHoldIndex, VertexID, FullWeight, bPruneInfluence, PreChangeWeights);
 			}
 			else
 			{
 				// assign remaining weight to parent
-				constexpr float OldParentWeight = 0.f;
 				const float NewParentWeight = 1.0f - NewWeightValue;
-				WeightEdits.MergeSingleEdit(ParentBoneIndex, VertexID, OldParentWeight, NewParentWeight);
+				WeightEdits.MergeSingleEdit(ParentBoneIndex, VertexID, NewParentWeight, bPruneInfluence, PreChangeWeights);
 				// and assign user requested weight to the current bone
-				const float OldWeight = GetWeightOfBoneOnVertex(BoneToHoldIndex, VertexID, PreChangeWeights);
-				const float NewWeight = NewWeightValue;
-				WeightEdits.MergeSingleEdit(BoneToHoldIndex, VertexID, OldWeight, NewWeight);
+				WeightEdits.MergeSingleEdit(BoneToHoldIndex, VertexID, NewWeightValue, bPruneInfluence, PreChangeWeights);
 			}
 		}
 		
 		return;
 	}
 
-	// calculate amount we have to spread across the other bones affecting this vertex
-	const float AvailableTotal = 1.0f - NewWeightValue;
-
-	// normalize weights into available space not set by current bone
-	for (int32 i=0; i<ValuesToNormalize.Num(); ++i)
+	// a normal weight edit where we assign the weight as requested
+	// and split the remainder amongst the recorded influences
 	{
-		float NormalizedValue = 0.f;
-		if (AvailableTotal > MinimumWeightThreshold && Total > KINDA_SMALL_NUMBER)
+		// calculate amount we have to spread across the other bones affecting this vertex
+		const float AvailableTotal = 1.0f - NewWeightValue;
+
+		// normalize weights into available space not set by current bone
+		for (int32 i=0; i<ValuesToNormalize.Num(); ++i)
 		{
-			NormalizedValue = (ValuesToNormalize[i] / Total) * AvailableTotal;	
+			float NormalizedValue = 0.f;
+			if (AvailableTotal > MinimumWeightThreshold && Total > KINDA_SMALL_NUMBER)
+			{
+				NormalizedValue = (ValuesToNormalize[i] / Total) * AvailableTotal;	
+			}
+			const int32 BoneIndex = RecordedBonesOnVertex[i];
+			const float NewWeight = NormalizedValue;
+			WeightEdits.MergeSingleEdit(BoneIndex, VertexID, NewWeight, bPruneInfluence, PreChangeWeights);
 		}
-		const int32 BoneIndex = RecordedBonesOnVertex[i];
-		const float OldWeight = ValuesToNormalize[i];
-		const float NewWeight = NormalizedValue;
-		WeightEdits.MergeSingleEdit(BoneIndex, VertexID, OldWeight, NewWeight);
-	}
 
-	// record current bone edit
-	const float PrevWeight = GetWeightOfBoneOnVertex(BoneToHoldIndex, VertexID, PreChangeWeights);
-	WeightEdits.MergeSingleEdit(
-		BoneToHoldIndex,
-		VertexID,
-		PrevWeight,
-		NewWeightValue);
+		// record current bone edit
+		WeightEdits.MergeSingleEdit(BoneToHoldIndex, VertexID, NewWeightValue, bPruneInfluence, PreChangeWeights);
+	}
 }
 
 void FSkinToolWeights::ApplyCurrentWeightsToMeshDescription(FMeshDescription* MeshDescription)
@@ -929,7 +952,7 @@ float FSkinToolWeights::GetWeightOfBoneOnVertex(
 			return BoneWeight.Weight;
 		}
 	}
-
+	
 	return 0.f;
 }
 
@@ -1176,15 +1199,21 @@ void FMeshSkinWeightsChange::Apply(UObject* Object)
 {
 	USkinWeightsPaintTool* Tool = CastChecked<USkinWeightsPaintTool>(Object);
 
-	// apply weight edits
 	Tool->ExternalUpdateSkinWeightLayer(LOD, SkinWeightProfile);
-	for (TTuple<int32, FSingleBoneWeightEdits>& Pair : AllWeightEdits.PerBoneWeightEdits)
+	
+	// apply weight edits
+	for (TTuple<BoneIndex, FSingleBoneWeightEdits>& Pair : AllWeightEdits.PerBoneWeightEdits)
 	{
-		Tool->ExternalUpdateWeights(Pair.Key, Pair.Value.NewWeights);
-	}
+		const BoneIndex BoneID = Pair.Key;
+		
+		// add and/or remove this bone from vertices
+		// (any weight removed by this influence is presumed to be added back in a normalized fashion by ExternalUpdateWeights)
+		Tool->ExternalRemoveInfluenceFromVertices(BoneID, Pair.Value.VerticesRemovedFrom);
+		Tool->ExternalAddInfluenceToVertices(BoneID, Pair.Value.VerticesAddedTo);
 
-	// remove pruned influences (if there are any)
-	Tool->ExternalRemoveInfluences(AllWeightEdits.PrunedInfluences);
+		// update the weights on this bone to use the new weights
+		Tool->ExternalUpdateWeights(BoneID, Pair.Value.NewWeights);
+	}
 }
 
 void FMeshSkinWeightsChange::Revert(UObject* Object)
@@ -1193,14 +1222,19 @@ void FMeshSkinWeightsChange::Revert(UObject* Object)
 
 	// update the skin weight profile
 	Tool->ExternalUpdateSkinWeightLayer(LOD, SkinWeightProfile);
-
-	// apply prune edits (restores pruned influences if there are any)
-	Tool->ExternalAddInfluences(AllWeightEdits.PrunedInfluences);
 	
 	// apply weight edits
 	for (TTuple<int32, FSingleBoneWeightEdits>& Pair : AllWeightEdits.PerBoneWeightEdits)
 	{
-		Tool->ExternalUpdateWeights(Pair.Key, Pair.Value.OldWeights);
+		const BoneIndex BoneID = Pair.Key;
+		
+		// add back vertices that this bone was removed from
+		Tool->ExternalAddInfluenceToVertices(BoneID, Pair.Value.VerticesRemovedFrom);
+		// remove vertices that this bone was added to
+		Tool->ExternalRemoveInfluenceFromVertices(BoneID, Pair.Value.VerticesAddedTo);
+		
+		// set the weights back to what they were before this change
+		Tool->ExternalUpdateWeights(BoneID, Pair.Value.OldWeights);
 	}
 
 	// notify dependent systems
@@ -1230,6 +1264,18 @@ void FMeshSkinWeightsChange::StoreBoneWeightEdit(const FSingleBoneWeightEdits& B
 		}
 		RemappedBoneWeightEdits.OldWeights = RemappedWeights;
 
+		// remap vertex indices that had this influence added to them
+		for (int32 VertArrayIndex=0; VertArrayIndex<RemappedBoneWeightEdits.VerticesAddedTo.Num(); ++VertArrayIndex)
+		{
+			RemappedBoneWeightEdits.VerticesAddedTo[VertArrayIndex] = VertexIndexConverter(VertArrayIndex);
+		}
+
+		// remap vertex indices that had this influence removed from them
+		for (int32 VertArrayIndex=0; VertArrayIndex<RemappedBoneWeightEdits.VerticesRemovedFrom.Num(); ++VertArrayIndex)
+		{
+			RemappedBoneWeightEdits.VerticesRemovedFrom[VertArrayIndex] = VertexIndexConverter(VertArrayIndex);
+		}
+
 		// merge the weight edits
 		AllWeightEdits.MergeEdits(RemappedBoneWeightEdits);
 		return;
@@ -1239,11 +1285,6 @@ void FMeshSkinWeightsChange::StoreBoneWeightEdit(const FSingleBoneWeightEdits& B
 	AllWeightEdits.MergeEdits(BoneWeightEdit);
 }
 
-void FMeshSkinWeightsChange::StorePruneBoneEdit(const VertexIndex VertexToPruneFrom, const BoneIndex BoneToPrune)
-{
-	AllWeightEdits.PrunedInfluences.Emplace(VertexToPruneFrom, BoneToPrune);
-}
-
 void FMeshSkinWeightsChange::StoreMultipleWeightEdits(const FMultiBoneWeightEdits& WeightEdits, const TFunction<int32(int32)>& VertexIndexConverter)
 {
 	// store weight edits
@@ -1251,12 +1292,6 @@ void FMeshSkinWeightsChange::StoreMultipleWeightEdits(const FMultiBoneWeightEdit
 	{
 		StoreBoneWeightEdit(BoneWeightEdits.Value, VertexIndexConverter);
 	}
-
-	// store pruned influences
-	for (const TPair<VertexIndex, BoneIndex>& PrunedInfluence : WeightEdits.PrunedInfluences)
-	{
-		StorePruneBoneEdit(PrunedInfluence.Key, PrunedInfluence.Value);
-	}
 }
 
 /*
@@ -2237,18 +2272,19 @@ void USkinWeightsPaintTool::CreateWeightEditsToRelaxVertices(
 
 			TMap<int32, float> FinalWeights;
 			const bool bSmoothSuccess = SmoothWeightsOp->SmoothWeightsAtVertex(VertexID, UseFalloff, FinalWeights);
+			if (!ensure(bSmoothSuccess))
+			{
+				continue;
+			}
 
-			if (ensure(bSmoothSuccess))
+			// apply weight edits
+			for (const TTuple<BoneIndex, float>& FinalWeight : FinalWeights)
 			{
-				// apply weight edits
-				for (const TTuple<BoneIndex, float>& FinalWeight : FinalWeights)
-				{
-					// record an edit for this vertex, for this bone
-					const int32 BoneIndex = FinalWeight.Key;
-					const float NewWeight = FinalWeight.Value;
-					const float OldWeight = Weights.GetWeightOfBoneOnVertex(BoneIndex, VertexID, Weights.PreChangeWeights);
-					InOutWeightEdits.MergeSingleEdit(BoneIndex, VertexID, OldWeight, NewWeight);
-				}
+				// record an edit for this vertex, for this bone
+				const int32 BoneIndex = FinalWeight.Key;
+				const float NewWeight = FinalWeight.Value;
+				constexpr bool bPruneInfluence = false;
+				InOutWeightEdits.MergeSingleEdit(BoneIndex, VertexID, NewWeight, bPruneInfluence, Weights.PreChangeWeights);
 			}
 		}
 	}
@@ -2493,22 +2529,26 @@ void USkinWeightsPaintTool::ExternalUpdateSkinWeightLayer(const EMeshLODIdentifi
 	}
 }
 
-void USkinWeightsPaintTool::ExternalAddInfluences(const TArray<TPair<VertexIndex, BoneIndex>>& InfluencesToAdd)
+void USkinWeightsPaintTool::ExternalAddInfluenceToVertices(
+	const BoneIndex InfluenceToAdd,
+	const TArray<VertexIndex>& Vertices)
 {
-	for (const TPair<VertexIndex, BoneIndex>& ToAdd : InfluencesToAdd)
+	for (const VertexIndex VertexID : Vertices)
 	{
 		constexpr float DefaultWeight = 0.f;
-		Weights.AddNewInfluenceToVertex(ToAdd.Key, ToAdd.Value, DefaultWeight, Weights.CurrentWeights);
-		Weights.AddNewInfluenceToVertex(ToAdd.Key, ToAdd.Value, DefaultWeight, Weights.PreChangeWeights);
+		Weights.AddNewInfluenceToVertex(VertexID, InfluenceToAdd, DefaultWeight, Weights.CurrentWeights);
+		Weights.AddNewInfluenceToVertex(VertexID, InfluenceToAdd, DefaultWeight, Weights.PreChangeWeights);
 	}
 }
 
-void USkinWeightsPaintTool::ExternalRemoveInfluences(const TArray<TPair<VertexIndex, BoneIndex>>& InfluencesToRemove)
+void USkinWeightsPaintTool::ExternalRemoveInfluenceFromVertices(
+	const BoneIndex InfluenceToRemove,
+	const TArray<VertexIndex>& Vertices)
 {
-	for (const TPair<VertexIndex, BoneIndex>& ToRemove : InfluencesToRemove)
+	for (const VertexIndex VertexID : Vertices)
 	{
-		Weights.RemoveInfluenceFromVertex(ToRemove.Key, ToRemove.Value, Weights.CurrentWeights);
-		Weights.RemoveInfluenceFromVertex(ToRemove.Key, ToRemove.Value, Weights.PreChangeWeights);
+		Weights.RemoveInfluenceFromVertex(VertexID, InfluenceToRemove, Weights.CurrentWeights);
+		Weights.RemoveInfluenceFromVertex(VertexID, InfluenceToRemove, Weights.PreChangeWeights);
 	}
 }
 
@@ -2710,18 +2750,18 @@ void USkinWeightsPaintTool::MirrorWeights(EAxis::Type Axis, EMirrorDirection Dir
 		// remove all weight on vertex
 		for (const FVertexBoneWeight& TargetBoneWeight : Weights.PreChangeWeights[TargetVertexID])
 		{
-			const float OldWeight = TargetBoneWeight.Weight;
+			constexpr bool bPruneInfluence = false;
 			constexpr float NewWeight = 0.f;
-			WeightEditsFromMirroring.MergeSingleEdit(TargetBoneWeight.BoneID, TargetVertexID, OldWeight, NewWeight);
+			WeightEditsFromMirroring.MergeSingleEdit(TargetBoneWeight.BoneID, TargetVertexID, NewWeight, bPruneInfluence, Weights.PreChangeWeights);
 		}
 
 		// copy source weights, but with mirrored bones
 		for (const FVertexBoneWeight& SourceBoneWeight : Weights.PreChangeWeights[SourceVertexID])
 		{
 			const BoneIndex MirroredBoneIndex = BoneMap[SourceBoneWeight.BoneID];
-			const float OldWeight = Weights.GetWeightOfBoneOnVertex(MirroredBoneIndex, TargetVertexID, Weights.PreChangeWeights);
+			constexpr bool bPruneInfluence = false;
 			const float NewWeight = SourceBoneWeight.Weight;
-			WeightEditsFromMirroring.MergeSingleEdit(MirroredBoneIndex, TargetVertexID, OldWeight, NewWeight);
+			WeightEditsFromMirroring.MergeSingleEdit(MirroredBoneIndex, TargetVertexID, NewWeight, bPruneInfluence, Weights.PreChangeWeights);
 		}
 	}
 
@@ -2801,17 +2841,16 @@ void USkinWeightsPaintTool::PruneWeights(float Threshold, const TArray<BoneIndex
 				InfluencesToPrune.Add(BoneWeight.BoneID);
 
 				// store a weight edit to remove this weight
-				WeightEditsFromPrune.MergeSingleEdit(BoneWeight.BoneID, VertexID, BoneWeight.Weight, 0.f);
+				constexpr float NewWeight = 0.f;
+				constexpr bool bPruneInfluence = true;
+				WeightEditsFromPrune.MergeSingleEdit(BoneWeight.BoneID, VertexID, NewWeight, bPruneInfluence, Weights.PreChangeWeights);
 			}
 		}
 
-		// actually prune the influences from the vert
+		// remove the influence from the vertex in the CURRENT weights
+		// this prevents ApplyWeightEdits() from ever using this influence as a dumping ground for normalization
 		for (const BoneIndex InfluenceToPrune : InfluencesToPrune)
 		{
-			// store this in the transaction
-			WeightEditsFromPrune.AddPruneBoneEdit(VertexID, InfluenceToPrune);
-			
-			// remove the influence from the vertex to prevent subsequent weight editing from normalizing weight back onto it
 			Weights.RemoveInfluenceFromVertex(VertexID, InfluenceToPrune, Weights.CurrentWeights);
 		}
 
@@ -2820,9 +2859,9 @@ void USkinWeightsPaintTool::PruneWeights(float Threshold, const TArray<BoneIndex
 		{
 			// we pruned ALL influences from a vertex, so dump all weight on root
 			constexpr BoneIndex RootBoneIndex = 0;
-			const float OldWeight = Weights.GetWeightOfBoneOnVertex(RootBoneIndex, VertexID, Weights.PreChangeWeights);
+			constexpr bool bPruneInfluence = false;
 			constexpr float NewWeight = 1.f;
-			WeightEditsFromPrune.MergeSingleEdit(RootBoneIndex, VertexID, OldWeight, NewWeight);
+			WeightEditsFromPrune.MergeSingleEdit(RootBoneIndex, VertexID, NewWeight, bPruneInfluence, Weights.PreChangeWeights);
 		}
 		else
 		{
@@ -2840,9 +2879,9 @@ void USkinWeightsPaintTool::PruneWeights(float Threshold, const TArray<BoneIndex
 			// record weight edits to normalize the weight across the remaining influences
 			for (const FVertexBoneWeight& BoneWeight : VertexWeights)
 			{
-				const float OldWeight = BoneWeight.Weight;
+				constexpr bool bPruneInfluence = false;
 				const float NewWeight = bNoOtherWeights ? EvenlySplitWeight : BoneWeight.Weight / TotalWeight;
-				WeightEditsFromPrune.MergeSingleEdit(BoneWeight.BoneID, VertexID, OldWeight, NewWeight);
+				WeightEditsFromPrune.MergeSingleEdit(BoneWeight.BoneID, VertexID, NewWeight, bPruneInfluence, Weights.PreChangeWeights);
 			}
 		}
 	}
@@ -2880,9 +2919,9 @@ void USkinWeightsPaintTool::AverageWeights(const float Strength)
 			{
 				if (!AveragedWeights.Contains(BoneWeight.BoneID))
 				{
-					const float OldWeight = BoneWeight.Weight;
+					constexpr bool bPruneInfluence = false;
 					constexpr float NewWeight = 0.f;
-					WeightEditsFromAveraging.MergeSingleEdit(BoneWeight.BoneID, VertexID, OldWeight, NewWeight);
+					WeightEditsFromAveraging.MergeSingleEdit(BoneWeight.BoneID, VertexID, NewWeight, bPruneInfluence, Weights.PreChangeWeights);
 				}
 			}
 
@@ -2890,9 +2929,9 @@ void USkinWeightsPaintTool::AverageWeights(const float Strength)
 			for (const TTuple<BoneIndex, float>& AveragedWeight : AveragedWeights)
 			{
 				const BoneIndex IndexOfBone = AveragedWeight.Key;
-				const float OldWeight = Weights.GetWeightOfBoneOnVertex(IndexOfBone, VertexID, Weights.PreChangeWeights);
+				constexpr bool bPruneInfluence = false;
 				const float NewWeight = AveragedWeight.Value;
-				WeightEditsFromAveraging.MergeSingleEdit(IndexOfBone, VertexID, OldWeight, NewWeight);
+				WeightEditsFromAveraging.MergeSingleEdit(IndexOfBone, VertexID, NewWeight, bPruneInfluence, Weights.PreChangeWeights);
 			}
 		}
 	}
@@ -2933,9 +2972,9 @@ void USkinWeightsPaintTool::AverageWeights(const float Strength)
 			for (const TTuple<BoneIndex, float>& BlendedWeight : BlendedWeights)
 			{
 				const BoneIndex BoneID = BlendedWeight.Key;
-				const float OldWeight = Weights.GetWeightOfBoneOnVertex(BoneID, VertexID, Weights.PreChangeWeights);
+				constexpr bool bPruneInfluence = false;
 				const float NewWeight = BlendedWeight.Value;
-				WeightEditsFromAveraging.MergeSingleEdit(BoneID, VertexID, OldWeight, NewWeight);
+				WeightEditsFromAveraging.MergeSingleEdit(BoneID, VertexID, NewWeight, bPruneInfluence, Weights.PreChangeWeights);
 			}
 		}
 	}
@@ -3022,20 +3061,20 @@ void USkinWeightsPaintTool::HammerWeights()
 		}
 		const int32 ClosestVertex = VertexPath.Last();
 
-		// remove all current weights
+		// remove all current weights (but don't prune)
 		for (const FVertexBoneWeight& BoneWeight : Weights.PreChangeWeights[SelectedVertex])
 		{
-			const float OldWeight = BoneWeight.Weight;
 			constexpr float NewWeight = 0.f;
-			HammerWeightEdits.MergeSingleEdit(BoneWeight.BoneID, SelectedVertex, OldWeight, NewWeight);
+			constexpr bool bPruneInfluence = false;
+			HammerWeightEdits.MergeSingleEdit(BoneWeight.BoneID, SelectedVertex, NewWeight, bPruneInfluence, Weights.PreChangeWeights);
 		}
 
 		// add weights from closest vertex
 		for (const FVertexBoneWeight& BoneWeight : Weights.PreChangeWeights[ClosestVertex])
 		{
-			const float OldWeight = Weights.GetWeightOfBoneOnVertex(BoneWeight.BoneID, SelectedVertex, Weights.PreChangeWeights);
 			const float NewWeight = BoneWeight.Weight;
-			HammerWeightEdits.MergeSingleEdit(BoneWeight.BoneID, SelectedVertex, OldWeight, NewWeight);
+			constexpr bool bPruneInfluence = false;
+			HammerWeightEdits.MergeSingleEdit(BoneWeight.BoneID, SelectedVertex, NewWeight, bPruneInfluence, Weights.PreChangeWeights);
 		}
 	}
 	
@@ -3198,13 +3237,17 @@ void USkinWeightsPaintTool::TransferWeights()
 			{
 				for (const FVertexBoneWeight& BoneWeight : VertexBoneWeights)
 				{
-					const float OldWeight = BoneWeight.Weight;
-					WeightEdits.MergeSingleEdit(BoneWeight.BoneID, SrcVertexID, OldWeight, ZeroWeight);
+					// when transferring weights we do prune influences because we would prefer the results to be identical to the source
+					constexpr bool bPruneInfluence = true;
+					WeightEdits.MergeSingleEdit(BoneWeight.BoneID, SrcVertexID, ZeroWeight, bPruneInfluence, Weights.PreChangeWeights);
 				}
 			}
 			else
 			{
-				WeightEdits.MergeSingleEdit(0, SrcVertexID, 1.f, ZeroWeight);
+				// in the unlikely event that the target vertex has no weight, "fake" remove it from the root so that undo will put it back there
+				constexpr bool bPruneInfluence = false;
+				BoneIndex RootBoneIndex = 0;
+				WeightEdits.MergeSingleEdit(RootBoneIndex, SrcVertexID, ZeroWeight, bPruneInfluence, Weights.PreChangeWeights);
 			}
 
 			// update with new weight
@@ -3213,9 +3256,9 @@ void USkinWeightsPaintTool::TransferWeights()
 			for (const FBoneWeight& BoneWeight: TransferedBoneWeights)
 			{
 				const int32 BoneIndex = BoneWeight.GetBoneIndex();					
-				const float OldWeight = Weights.GetWeightOfBoneOnVertex(BoneIndex, SrcVertexID, Weights.PreChangeWeights);
 				const float NewWeight = BoneWeight.GetWeight();
-				WeightEdits.MergeSingleEdit(BoneIndex, SrcVertexID, OldWeight, NewWeight);
+				constexpr bool bPruneInfluence = false;
+				WeightEdits.MergeSingleEdit(BoneIndex, SrcVertexID, NewWeight, bPruneInfluence, Weights.PreChangeWeights);
 			}
 		}
 
@@ -3381,18 +3424,19 @@ void USkinWeightsPaintTool::PasteWeights()
 		// remove all current weights
 		for (const FVertexBoneWeight& BoneWeight : Weights.PreChangeWeights[SelectedVertex])
 		{
-			const float OldWeight = BoneWeight.Weight;
-			constexpr float NewWeight = 0.f;
-			PasteWeightEdits.MergeSingleEdit(BoneWeight.BoneID, SelectedVertex, OldWeight, NewWeight);
+			// when pasting weights, we want a complete replacement, so we prune everything first
+			constexpr bool bPruneInfluence = true;
+			constexpr float ZeroWeight = 0.f;
+			PasteWeightEdits.MergeSingleEdit(BoneWeight.BoneID, SelectedVertex, ZeroWeight, bPruneInfluence, Weights.PreChangeWeights);
 		}
 
 		// add weights from clipboard
 		for (const TPair<BoneIndex, float>& LoadedWeight : LoadedWeights)
 		{
-			BoneIndex BoneID = LoadedWeight.Key;
-			const float OldWeight = Weights.GetWeightOfBoneOnVertex(BoneID, SelectedVertex, Weights.PreChangeWeights);
+			const BoneIndex BoneID = LoadedWeight.Key;
 			const float NewWeight = LoadedWeight.Value;
-			PasteWeightEdits.MergeSingleEdit(BoneID, SelectedVertex, OldWeight, NewWeight);
+			constexpr bool bPruneInfluence = false;
+			PasteWeightEdits.MergeSingleEdit(BoneID, SelectedVertex, NewWeight, bPruneInfluence, Weights.PreChangeWeights);
 		}
 	}
 	
diff --git a/Engine/Plugins/Experimental/MeshModelingToolsetExp/Source/MeshModelingToolsEditorOnlyExp/Public/SkeletalMesh/SkinWeightsPaintTool.h b/Engine/Plugins/Experimental/MeshModelingToolsetExp/Source/MeshModelingToolsEditorOnlyExp/Public/SkeletalMesh/SkinWeightsPaintTool.h
--- a/Engine/Plugins/Experimental/MeshModelingToolsetExp/Source/MeshModelingToolsEditorOnlyExp/Public/SkeletalMesh/SkinWeightsPaintTool.h
+++ b/Engine/Plugins/Experimental/MeshModelingToolsetExp/Source/MeshModelingToolsEditorOnlyExp/Public/SkeletalMesh/SkinWeightsPaintTool.h
@@ -143,25 +143,30 @@ namespace SkinPaintTool
 	struct FSingleBoneWeightEdits
 	{
 		int32 BoneIndex;
+		
 		TMap<VertexIndex, float> OldWeights;
 		TMap<VertexIndex, float> NewWeights;
+		
+		TArray<VertexIndex> VerticesAddedTo;
+		TArray<VertexIndex> VerticesRemovedFrom;
 	};
 
 	// store a sparse set of modifications to a set of vertex weights for a SET of bones
 	// with support for merging edits. these are used for transaction history undo/redo.
 	struct FMultiBoneWeightEdits
 	{
-		void MergeSingleEdit(const int32 BoneIndex, const int32 VertexID, const float OldWeight, const float NewWeight);
+		void MergeSingleEdit(
+			const int32 BoneIndex,
+			const int32 VertexID,
+			const float NewWeight,
+			bool bPruneInfluence,
+			const TArray<VertexWeights>& InPreChangeWeights);
 		void MergeEdits(const FSingleBoneWeightEdits& BoneWeightEdits);
 		float GetVertexDeltaFromEdits(const int32 BoneIndex, const int32 VertexIndex);
 		void AddEditedVerticesToSet(TSet<int32>& OutEditedVertexSet) const;
-		void AddPruneBoneEdit(const VertexIndex VertexToPruneFrom, const BoneIndex BoneToPrune);
-
+		
 		// map of bone indices to weight edits made to that bone
 		TMap<BoneIndex, FSingleBoneWeightEdits> PerBoneWeightEdits;
-
-		// influences to prune as part of these edits
-		TArray<TPair<VertexIndex, BoneIndex>> PrunedInfluences;
 	};
 	
 	class FMeshSkinWeightsChange : public FToolCommandChange
@@ -183,9 +188,7 @@ namespace SkinPaintTool
 		virtual void Revert(UObject* Object) override;
 
 		void StoreBoneWeightEdit(const FSingleBoneWeightEdits& BoneWeightEdit, const TFunction<int32(int32)>& VertexIndexConverter);
-
-		void StorePruneBoneEdit(const VertexIndex VertexToPruneFrom, const BoneIndex BoneToPrune);
-
+		
 		void StoreMultipleWeightEdits(const FMultiBoneWeightEdits& WeightEdits, const TFunction<int32(int32)>& VertexIndexConverter);
 
 	private:
@@ -218,6 +221,12 @@ namespace SkinPaintTool
 			const int32 VertexID,
 			const TArray<VertexWeights>& InVertexWeights);
 
+		static void FillWeightEdit(
+			const int32 BoneIndex,
+			const int32 VertexID,
+			const float NewWeight,
+			const TArray<VertexWeights>& InVertexWeights);
+
 		void SetWeightOfBoneOnVertex(
 			const int32 BoneIndex,
 			const int32 VertexID,
@@ -492,8 +501,8 @@ public:
 	// using when ToolChange is applied via Undo/Redo
 	void ExternalUpdateWeights(const int32 BoneIndex, const TMap<int32, float>& IndexValues);
 	void ExternalUpdateSkinWeightLayer(const EMeshLODIdentifier InLOD, const FName InSkinWeightProfile);
-	void ExternalAddInfluences(const TArray<TPair<VertexIndex, BoneIndex>>& InfluencesToAdd);
-	void ExternalRemoveInfluences(const TArray<TPair<VertexIndex, BoneIndex>>& InfluencesToRemove);
+	void ExternalAddInfluenceToVertices(const BoneIndex InfluenceToAdd, const TArray<VertexIndex>& Vertices);
+	void ExternalRemoveInfluenceFromVertices(const BoneIndex InfluenceToRemove, const TArray<VertexIndex>& Vertices);
 
 	// weight editing operations (selection based)
 	void MirrorWeights(EAxis::Type Axis, EMirrorDirection Direction);
