make edit pivot's snap drag support snapping just rotation, in addition to the previous just position or position + rotation options
+ fix snap not applying on initial click of click+drag

diff --git a/Engine/Plugins/Experimental/MeshModelingToolsetExp/Source/MeshModelingToolsExp/Private/EditPivotTool.cpp b/Engine/Plugins/Experimental/MeshModelingToolsetExp/Source/MeshModelingToolsExp/Private/EditPivotTool.cpp
--- a/Engine/Plugins/Experimental/MeshModelingToolsetExp/Source/MeshModelingToolsExp/Private/EditPivotTool.cpp
+++ b/Engine/Plugins/Experimental/MeshModelingToolsetExp/Source/MeshModelingToolsExp/Private/EditPivotTool.cpp
@@ -364,7 +364,7 @@ void UEditPivotTool::ResetActiveGizmos()
 // does not make sense that CanBeginClickDragSequence() returns a RayHit? Needs to be an out-argument...
 FInputRayHit UEditPivotTool::CanBeginClickDragSequence(const FInputDeviceRay& PressPos)
 {
-	if (TransformProps->bEnableSnapDragging == false || ActiveGizmos.Num() == 0)
+	if ((!TransformProps->bSnapDragPosition && !TransformProps->bSnapDragRotation) || ActiveGizmos.Num() == 0)
 	{
 		return FInputRayHit();
 	}
@@ -389,12 +389,15 @@ void UEditPivotTool::OnClickPress(const FInputDeviceRay& PressPos)
 	FEditPivotTarget& ActiveTarget = ActiveGizmos[0];
 	USceneComponent* GizmoComponent = ActiveTarget.TransformGizmo->GetGizmoActor()->GetRootComponent();
 	StartDragTransform = GizmoComponent->GetComponentToWorld();
+
+	// Apply the drag logic as well so that the snap-drag position/orientation update is applied on first click
+	OnClickDrag(PressPos);
 }
 
 
 void UEditPivotTool::OnClickDrag(const FInputDeviceRay& DragPos)
 {
-	bool bRotate = (TransformProps->RotationMode != EEditPivotSnapDragRotationMode::Ignore);
+	bool bRotate = TransformProps->bSnapDragRotation;
 	float NormalSign = (TransformProps->RotationMode == EEditPivotSnapDragRotationMode::AlignFlipped) ? -1.0f : 1.0f;
 
 	FHitResult Result;
@@ -413,7 +416,10 @@ void UEditPivotTool::OnClickDrag(const FInputDeviceRay& DragPos)
 
 	FTransform NewTransform = StartDragTransform;
 	NewTransform.SetRotation((FQuat)AlignRotation);
-	NewTransform.SetTranslation(HitPos);
+	if (TransformProps->bSnapDragPosition)
+	{
+		NewTransform.SetTranslation(HitPos);
+	}
 
 	FEditPivotTarget& ActiveTarget = ActiveGizmos[0];
 	ActiveTarget.TransformGizmo->SetNewGizmoTransform(NewTransform);
diff --git a/Engine/Plugins/Experimental/MeshModelingToolsetExp/Source/MeshModelingToolsExp/Public/EditPivotTool.h b/Engine/Plugins/Experimental/MeshModelingToolsetExp/Source/MeshModelingToolsExp/Public/EditPivotTool.h
--- a/Engine/Plugins/Experimental/MeshModelingToolsetExp/Source/MeshModelingToolsExp/Public/EditPivotTool.h
+++ b/Engine/Plugins/Experimental/MeshModelingToolsetExp/Source/MeshModelingToolsExp/Public/EditPivotTool.h
@@ -40,9 +40,6 @@ protected:
 UENUM()
 enum class EEditPivotSnapDragRotationMode : uint8
 {
-	/** Snap-Drag only translates, ignoring Normals */
-	Ignore = 0 UMETA(DisplayName = "Ignore"),
-
 	/** Snap-Drag aligns the pivot Z Axis and Target Normals to point in the same direction */
 	Align = 1 UMETA(DisplayName = "Align"),
 
@@ -68,11 +65,15 @@ public:
 	bool bApplyToAllLODs = true;
 
 	/** When enabled, click-drag to reposition the Pivot */
-	UPROPERTY(EditAnywhere, Category = Options)
-	bool bEnableSnapDragging = false;
-
-	/** When Snap-Dragging, align source and target normals */
-	UPROPERTY(EditAnywhere, Category = Options, meta = (EditCondition = "bEnableSnapDragging == true"))
+	UPROPERTY(EditAnywhere, Category = SnapDragging)
+	bool bSnapDragPosition = false;
+	
+	/** When enabled, click-drag to reorient the Pivot */
+	UPROPERTY(EditAnywhere, Category = SnapDragging)
+	bool bSnapDragRotation = false;
+
+	/** When snap-dragging rotation, how to align source and target normals */
+	UPROPERTY(EditAnywhere, Category = SnapDragging, meta = (EditCondition = "bSnapDragRotation"))
 	EEditPivotSnapDragRotationMode RotationMode = EEditPivotSnapDragRotationMode::Align;
 };
 
