Skin Weights Paint Tool: make sure to reset bShowNewProfileName to false when switching back to the default profile

diff --git a/Engine/Plugins/Experimental/MeshModelingToolsetExp/Source/MeshModelingToolsEditorOnlyExp/Private/SkeletalMesh/SkinWeightsPaintTool.cpp b/Engine/Plugins/Experimental/MeshModelingToolsetExp/Source/MeshModelingToolsEditorOnlyExp/Private/SkeletalMesh/SkinWeightsPaintTool.cpp
--- a/Engine/Plugins/Experimental/MeshModelingToolsetExp/Source/MeshModelingToolsEditorOnlyExp/Private/SkeletalMesh/SkinWeightsPaintTool.cpp
+++ b/Engine/Plugins/Experimental/MeshModelingToolsetExp/Source/MeshModelingToolsEditorOnlyExp/Private/SkeletalMesh/SkinWeightsPaintTool.cpp
@@ -1684,6 +1684,7 @@ void USkinWeightsPaintTool::PostEditMeshInitialization(
 	if (!IsProfileValid(WeightToolProperties->GetActiveSkinWeightProfile()))
 	{
 		WeightToolProperties->ActiveSkinWeightProfile = FSkeletalMeshAttributesShared::DefaultSkinWeightProfileName;
+		WeightToolProperties->bShowNewProfileName = false;
 	}
 	Weights.Profile = WeightToolProperties->GetActiveSkinWeightProfile();
 	Weights.InitializeSkinWeights(InComponent, &InMeshDescription);
