remove extra mesh copy in lattice deformer setup

diff --git a/Engine/Plugins/Experimental/MeshModelingToolsetExp/Source/MeshModelingToolsExp/Private/LatticeDeformerTool.cpp b/Engine/Plugins/Experimental/MeshModelingToolsetExp/Source/MeshModelingToolsExp/Private/LatticeDeformerTool.cpp
--- a/Engine/Plugins/Experimental/MeshModelingToolsetExp/Source/MeshModelingToolsExp/Private/LatticeDeformerTool.cpp
+++ b/Engine/Plugins/Experimental/MeshModelingToolsetExp/Source/MeshModelingToolsExp/Private/LatticeDeformerTool.cpp
@@ -180,17 +180,17 @@ bool ULatticeDeformerTool::CanAccept() const
 
 void ULatticeDeformerTool::InitializeLattice(TArray<FVector3d>& OutLatticePoints, TArray<FVector2i>& OutLatticeEdges)
 {
-	TUniquePtr<UE::Geometry::FDynamicMesh3> MeshToDeform = nullptr;
+	UE::Geometry::FDynamicMesh3* MeshToDeform = nullptr;
 	
 	if (bHasSelection && Submesh)
 	{
-		MeshToDeform = MakeUnique<FDynamicMesh3>(Submesh->GetSubmesh());
+		MeshToDeform = &Submesh->GetSubmesh();
 	}
 	else
 	{
-		MeshToDeform = MakeUnique<FDynamicMesh3>(*OriginalMesh.Get());
+		MeshToDeform = OriginalMesh.Get();
 	}
-	Lattice = MakeShared<FFFDLattice, ESPMode::ThreadSafe>(GetLatticeResolution(), *MeshToDeform.Get(), Settings->Padding);
+	Lattice = MakeShared<FFFDLattice, ESPMode::ThreadSafe>(GetLatticeResolution(), *MeshToDeform, Settings->Padding);
 
 
 	Lattice->GenerateInitialLatticePositions(OutLatticePoints);
