initialize cube grid initial grid origin based on element selection (if available)

diff --git a/Engine/Plugins/Experimental/MeshModelingToolsetExp/Source/MeshModelingToolsExp/Private/CubeGridTool.cpp b/Engine/Plugins/Experimental/MeshModelingToolsetExp/Source/MeshModelingToolsExp/Private/CubeGridTool.cpp
--- a/Engine/Plugins/Experimental/MeshModelingToolsetExp/Source/MeshModelingToolsExp/Private/CubeGridTool.cpp
+++ b/Engine/Plugins/Experimental/MeshModelingToolsetExp/Source/MeshModelingToolsExp/Private/CubeGridTool.cpp
@@ -31,6 +31,7 @@
 #include "Properties/MeshMaterialProperties.h"
 #include "PropertySets/CreateMeshObjectTypeProperties.h"
 #include "SceneView.h"
+#include "Selection/StoredMeshSelectionUtil.h"
 #include "Selection/ToolSelectionUtil.h"
 #include "ToolContextInterfaces.h"
 #include "ToolHostCustomizationAPI.h"
@@ -748,6 +749,20 @@ UInteractiveTool* UCubeGridToolBuilder::BuildTool(const FToolBuilderState& Scene
 
 	TObjectPtr<UToolTarget> Target = SceneState.TargetManager->BuildFirstSelectedTargetable(SceneState, GetTargetRequirements());
 	NewTool->SetTarget(Target); // May be null
+	if (Target)
+	{
+		// if there is an element selection, use its frame origin as an initial pivot
+		FFrame3d Frame;
+		FAxisAlignedBox3d Bounds;
+		bool bIsElementSelection;
+		if (GetCurrentSelectionWorldFrameBounds(SceneState, Frame, Bounds, bIsElementSelection))
+		{
+			if (bIsElementSelection)
+			{
+				NewTool->SetInitialGridPivot(Frame.Origin);
+			}
+		}
+	}
 	NewTool->SetWorld(SceneState.World);
 
 	return NewTool;
@@ -998,6 +1013,10 @@ void UCubeGridTool::Setup()
 	UpdateOpMaterials();
 
 	CubeGrid = MakeShared<FCubeGrid>();
+	if (bHasInitialGridPivot)
+	{
+		Settings->GridFrameOrigin = InitialGridPivot;
+	}
 	CubeGrid->SetGridFrame(FFrame3d(Settings->GridFrameOrigin, Settings->GridFrameOrientation.Quaternion()));
 	CubeGrid->SetGridPowerMode(Settings->bPowerOfTwoBlockSizes ? FCubeGrid::EPowerMode::PowerOfTwo : FCubeGrid::EPowerMode::FiveAndTen);
 	CubeGrid->SetGridPower(Settings->GridPower);
diff --git a/Engine/Plugins/Experimental/MeshModelingToolsetExp/Source/MeshModelingToolsExp/Public/CubeGridTool.h b/Engine/Plugins/Experimental/MeshModelingToolsetExp/Source/MeshModelingToolsExp/Public/CubeGridTool.h
--- a/Engine/Plugins/Experimental/MeshModelingToolsetExp/Source/MeshModelingToolsExp/Public/CubeGridTool.h
+++ b/Engine/Plugins/Experimental/MeshModelingToolsetExp/Source/MeshModelingToolsExp/Public/CubeGridTool.h
@@ -419,6 +419,12 @@ public:
 	virtual bool SupportsWorldSpaceFocusBox() { return bHaveSelection; }
 	virtual FBox GetWorldSpaceFocusBox() override;
 
+	void SetInitialGridPivot(FVector3d PivotPos)
+	{
+		InitialGridPivot = PivotPos;
+		bHasInitialGridPivot = true;
+	}
+
 protected:
 
 	UPROPERTY()
@@ -614,4 +620,8 @@ protected:
 
 	void OutputCurrentResults(bool bSetSelection);
 	void AcceptToolAndStartNew();
+
+private:
+	bool bHasInitialGridPivot = false;
+	FVector3d InitialGridPivot;
 };
