Skeleton Modifier: avoid brute forcing skeleton / skeletal mesh changes in the case of modifications likely to cause breakage
    - this workflow follows the interchange path, with the option of duplicating the skeleton before submitting modifications, so that it is not destructive

diff --git a/Engine/Plugins/Experimental/MeshModelingToolsetExp/Source/SkeletalMeshModifiers/Private/SkeletonModifier.cpp b/Engine/Plugins/Experimental/MeshModelingToolsetExp/Source/SkeletalMeshModifiers/Private/SkeletonModifier.cpp
--- a/Engine/Plugins/Experimental/MeshModelingToolsetExp/Source/SkeletalMeshModifiers/Private/SkeletonModifier.cpp
+++ b/Engine/Plugins/Experimental/MeshModelingToolsetExp/Source/SkeletalMeshModifiers/Private/SkeletonModifier.cpp
@@ -3,6 +3,7 @@
 
 #include "SkeletonModifier.h"
 
+#include "AssetNotifications.h"
 #include "BoneWeights.h"
 #include "MeshDescription.h"
 #include "Animation/Skeleton.h"
@@ -16,11 +17,152 @@
 #include "Engine/SkeletalMesh.h"
 #include "RenderingThread.h"
 #include "Animation/MirrorDataTable.h"
+#include "PropertyEditorModule.h"
+#include "Dialog/SCustomDialog.h"
+#include "Misc/MessageDialog.h"
+#include "Widgets/Layout/SBox.h"
+#include "Widgets/Text/STextBlock.h"
 
 
 namespace USkeletonModifierLocals
 {
 	static constexpr int32 LODIndex = 0;
+
+	/**
+	* FReferenceSkeletonCompatibilityChecker is a utility class to check whether two skeletons are compatible.
+	* It follows the USkeleton compatibility pattern.
+	* NOTE: there should be consistency between this, USkeleton and FSkeletonHelper so it should be merged at some point
+	*/
+
+	class FReferenceSkeletonCompatibilityChecker
+	{
+	public:
+		FReferenceSkeletonCompatibilityChecker(const FReferenceSkeleton& InRefSkeleton)
+			: ReferenceSkeleton(InRefSkeleton)
+		{}
+
+		bool DoesParentChainMatch(const int32 InStartBoneIndex, const FReferenceSkeleton& InRefSkeleton) const
+		{
+			// if start is root bone
+			if ( InStartBoneIndex == 0 )
+			{
+				// verify name of root bone matches
+				return (ReferenceSkeleton.GetBoneName(0) == InRefSkeleton.GetBoneName(0));
+			}
+
+			int32 SkeletonBoneIndex = InStartBoneIndex;
+			// If skeleton bone is not found in mesh, fail.
+			int32 OtherBoneIndex = InRefSkeleton.FindBoneIndex( ReferenceSkeleton.GetBoneName(SkeletonBoneIndex) );
+			if( OtherBoneIndex == INDEX_NONE )
+			{
+				return false;
+			}
+	
+			do
+			{
+				// verify if parent name matches
+				int32 ParentSkeletonBoneIndex = ReferenceSkeleton.GetParentIndex(SkeletonBoneIndex);
+				int32 ParentOtherBoneIndex = InRefSkeleton.GetParentIndex(OtherBoneIndex);
+
+				// if one of the parents doesn't exist, make sure both end. Otherwise fail.
+				if( (ParentSkeletonBoneIndex == INDEX_NONE) || (ParentOtherBoneIndex == INDEX_NONE) )
+				{
+					return (ParentSkeletonBoneIndex == ParentOtherBoneIndex);
+				}
+
+				// If parents are not named the same, fail.
+				if( ReferenceSkeleton.GetBoneName(ParentSkeletonBoneIndex) != InRefSkeleton.GetBoneName(ParentOtherBoneIndex) )
+				{
+					UE_LOG(LogAnimation, Warning, TEXT("%s : Hierarchy does not match %s - %s."), *ReferenceSkeleton.GetBoneName(SkeletonBoneIndex).ToString(), *ReferenceSkeleton.GetBoneName(ParentSkeletonBoneIndex).ToString(), *InRefSkeleton.GetBoneName(ParentOtherBoneIndex).ToString());
+					return false;
+				}
+
+				// move up
+				SkeletonBoneIndex = ParentSkeletonBoneIndex;
+				OtherBoneIndex = ParentOtherBoneIndex;
+			} while ( true );
+
+			return true;	
+		}
+		
+		bool IsCompatibleReferenceSkeleton(const FReferenceSkeleton& InRefSkeleton, bool bDoParentChainCheck=true) const
+		{
+			// at least % of bone should match 
+			int32 NumOfBoneMatches = 0;
+
+			const int32 OtherNumBones = InRefSkeleton.GetRawBoneNum();
+
+			// first ensure the parent exists for each bone
+			for (int32 OtherBoneIndex=0; OtherBoneIndex<OtherNumBones; OtherBoneIndex++)
+			{
+				const FName OtherBoneName = InRefSkeleton.GetBoneName(OtherBoneIndex);
+				// See if Mesh bone exists in Skeleton.
+				int32 SkeletonBoneIndex = ReferenceSkeleton.FindBoneIndex( OtherBoneName );
+
+				// if found, increase num of bone matches count
+				if( SkeletonBoneIndex != INDEX_NONE )
+				{
+					++NumOfBoneMatches;
+
+					// follow the parent chain to verify the chain is same
+					if(bDoParentChainCheck && !DoesParentChainMatch(SkeletonBoneIndex, InRefSkeleton))
+					{
+						UE_LOG(LogAnimation, Warning, TEXT("%s : Hierarchy does not match."), *OtherBoneName.ToString());
+						return false;
+					}
+				}
+				else
+				{
+					int32 CurrentBoneId = OtherBoneIndex;
+					// if not look for parents that matches
+					while (SkeletonBoneIndex == INDEX_NONE && CurrentBoneId != INDEX_NONE)
+					{
+						// find Parent one see exists
+						const int32 ParentOtherBoneIndex = InRefSkeleton.GetParentIndex(CurrentBoneId);
+						if ( ParentOtherBoneIndex != INDEX_NONE )
+						{
+							// @TODO: make sure RefSkeleton's root ParentIndex < 0 if not, I'll need to fix this by checking TreeBoneIdx
+							FName ParentBoneName = InRefSkeleton.GetBoneName(ParentOtherBoneIndex);
+							SkeletonBoneIndex = ReferenceSkeleton.FindBoneIndex(ParentBoneName);
+						}
+
+						// root is reached
+						if( ParentOtherBoneIndex == 0 )
+						{
+							break;
+						}
+						else
+						{
+							CurrentBoneId = ParentOtherBoneIndex;
+						}
+					}
+
+					// still no match, return false, no parent to look for
+					if( SkeletonBoneIndex == INDEX_NONE )
+					{
+						UE_LOG(LogAnimation, Warning, TEXT("%s : Missing joint on skeleton. Make sure to assign to the skeleton."), *OtherBoneName.ToString());
+						return false;
+					}
+
+					// second follow the parent chain to verify the chain is same
+					if (bDoParentChainCheck && !DoesParentChainMatch(SkeletonBoneIndex, InRefSkeleton))
+					{
+						UE_LOG(LogAnimation, Warning, TEXT("%s : Hierarchy does not match."), *OtherBoneName.ToString());
+						return false;
+					}
+				}
+			}
+
+			// originally we made sure at least matches more than 50% 
+			// but then follower components can't play since they're only partial
+			// if the hierarchy matches, and if it's more then 1 bone, we allow
+			return (NumOfBoneMatches > 0);
+		}
+
+	private:
+		const FReferenceSkeleton& ReferenceSkeleton;
+	};
+
 }
 
 FTransform FMirrorOptions::MirrorTransform(const FTransform& InTransform) const
@@ -179,36 +321,6 @@
 		return false;
 	}
 	
-	// check assets using this skeleton ?
-	{
-		// TODO avoid certain changes when the skeleton is referenced by other assets (i.e. changing the skeletal mesh's
-		// reference skeleton poses is fine, re-parenting/removing bones, etc. is not)
-		static const TArray<FTopLevelAssetPath> AssetPaths({
-		   UAnimSequence::StaticClass()->GetClassPathName(),
-		   UAnimMontage::StaticClass()->GetClassPathName(),
-		   UPoseAsset::StaticClass()->GetClassPathName(),
-		   USkeletalMesh::StaticClass()->GetClassPathName()});
-
-		FARFilter Filter;
-		Filter.ClassPaths = AssetPaths;
-
-		const IAssetRegistry& AssetRegistry = FModuleManager::LoadModuleChecked<FAssetRegistryModule>("AssetRegistry").Get();
-		TArray<FAssetData> Assets; AssetRegistry.GetAssets(Filter, Assets);
-
-		const FAssetData SkeletonAssetData(Skeleton);
-		const FString SkeletonPath = SkeletonAssetData.GetExportTextName();
-	
-		static const FName Tag("Skeleton");
-		for (const FAssetData& AssetData: Assets)
-		{
-			const FString TagValue = AssetData.GetTagValueRef<FString>(Tag);
-			if (TagValue == SkeletonPath)
-			{
-				UE_LOG(LogAnimation, Warning, TEXT("%s references that skeleton."), *AssetData.GetExportTextName());
-			}
-		}
-	}
-
 	// store pointer to mesh and instantiate a mesh description for commiting changes
 	SkeletalMesh = InSkeletalMesh;
 
@@ -236,10 +348,11 @@
 		BoneIndexTracker.Add(Index);
 	}
 	return true;
+	
 #else
 	ensureMsgf(false, TEXT("Skeleton Modifier is an editor only feature."));
-#endif
 	return false;
+#endif
 }
 
 bool USkeletonModifier::IsReferenceSkeletonValid(const bool bLog) const
@@ -255,15 +368,8 @@
 	return true;
 }
 
-bool USkeletonModifier::CommitSkeletonToSkeletalMesh()
+ESkeletalMeshModificationType USkeletonModifier::PreCommitSkeletalMesh()
 {
-	if (!SkeletalMesh.IsValid() || !ReferenceSkeleton || !MeshDescription)
-	{
-		UE_LOG(LogAnimation, Error, TEXT("Skeleton Modifier: No mesh loaded. Cannot apply skeleton edits."));
-		return false;
-	}
-
-#if WITH_EDITORONLY_DATA
 	// before commiting, we have to reparent non-root bones with no parent as the animation pipeline
 	// doesn't support multi-roots
 	const TArray<FMeshBoneInfo>& BoneInfos = ReferenceSkeleton->GetRawRefBoneInfo();
@@ -288,6 +394,218 @@
 			ParentBones(BonesToParent, {BoneInfos[0].Name});
 		}
 	}
+
+	// check topological change from this modifier
+	auto HasBoneIndexesChanged = [this]()
+	{
+		for (int32 Index = 0; Index < BoneIndexTracker.Num(); Index++)
+		{
+			if (BoneIndexTracker[Index] != Index)
+			{
+				return true;
+			}
+		}
+		return false;
+	};
+
+	auto BasicModificationCheck = [this, HasBoneIndexesChanged](const FReferenceSkeleton& InRefSkeleton)
+	{
+		const TArray<FMeshBoneInfo>& OtherBoneInfos = InRefSkeleton.GetRawRefBoneInfo();
+		const TArray<FTransform>& OtherBonePoses = InRefSkeleton.GetRawRefBonePose();
+		const int32 NumOtherBones = OtherBoneInfos.Num();
+
+		const TArray<FMeshBoneInfo>& NewBoneInfos = ReferenceSkeleton->GetRawRefBoneInfo();
+		const TArray<FTransform>& NewBonePoses = ReferenceSkeleton->GetRawRefBonePose();
+		const int32 NumNewBones = NewBoneInfos.Num();
+
+		ESkeletalMeshModificationType Modifications = ESkeletalMeshModificationType::None;
+		if (NumNewBones > NumOtherBones)
+		{
+			EnumAddFlags(Modifications, ESkeletalMeshModificationType::BonesAdded);
+			if (HasBoneIndexesChanged())
+			{
+				EnumAddFlags(Modifications, ESkeletalMeshModificationType::HierarchyChanged);
+			}
+		}
+		else if (NumNewBones < NumOtherBones)
+		{
+			EnumAddFlags(Modifications, ESkeletalMeshModificationType::BonesRemoved);
+
+			for (int32 NewBoneIndex = 0; NewBoneIndex < NumNewBones; NewBoneIndex++)
+			{
+				// check names
+				const FName NewBoneName = ReferenceSkeleton->GetBoneName(NewBoneIndex);
+				const int32 OtherBoneIndex = InRefSkeleton.FindBoneIndex(NewBoneName);
+				if (OtherBoneIndex == INDEX_NONE)
+				{
+					EnumAddFlags(Modifications, ESkeletalMeshModificationType::BonesRenamed);
+				}
+				else
+				{
+					// check parents names
+					const FMeshBoneInfo& NewBoneInfo = NewBoneInfos[NewBoneIndex];
+					const FMeshBoneInfo& OtherBoneInfo = OtherBoneInfos[OtherBoneIndex];
+
+					const int32 NewParentIndex = NewBoneInfo.ParentIndex;
+					const int32 OldParentIndex = OtherBoneInfo.ParentIndex;
+					const FName NewParentName = NewParentIndex != INDEX_NONE ? NewBoneInfos[NewParentIndex].Name : NAME_None; 
+					const FName OldParentName = OldParentIndex != INDEX_NONE ? OtherBoneInfos[OldParentIndex].Name : NAME_None;
+					
+					if ((NewParentIndex != INDEX_NONE || OldParentIndex != INDEX_NONE) && NewParentName != OldParentName)
+					{
+						EnumAddFlags(Modifications, ESkeletalMeshModificationType::HierarchyChanged);
+					}
+				}
+			}
+		}
+		else
+		{
+			for (int32 NewBoneIndex = 0; NewBoneIndex < NumNewBones; NewBoneIndex++)
+			{
+				// check names
+				const FName NewBoneName = ReferenceSkeleton->GetBoneName(NewBoneIndex);
+				const int32 OtherBoneIndex = InRefSkeleton.FindBoneIndex(NewBoneName);
+				if (OtherBoneIndex == INDEX_NONE)
+				{
+					EnumAddFlags(Modifications, ESkeletalMeshModificationType::BonesRenamed);
+				}
+				else
+				{
+					// check index
+					if (OtherBoneIndex != NewBoneIndex)
+					{
+						EnumAddFlags(Modifications, ESkeletalMeshModificationType::HierarchyChanged);
+					}
+					
+					// check parents
+					const FMeshBoneInfo& NewBoneInfo = NewBoneInfos[NewBoneIndex];
+					const FMeshBoneInfo& OtherBoneInfo = OtherBoneInfos[OtherBoneIndex];
+					if (NewBoneInfo.ParentIndex != OtherBoneInfo.ParentIndex)
+					{
+						EnumAddFlags(Modifications, ESkeletalMeshModificationType::HierarchyChanged);
+					}
+
+					// check transforms
+					const FTransform& NewBoneTransform = NewBonePoses[NewBoneIndex];
+					const FTransform& OtherBoneTransform = OtherBonePoses[OtherBoneIndex];
+					if (!NewBoneTransform.Equals(OtherBoneTransform))
+					{
+						EnumAddFlags(Modifications, ESkeletalMeshModificationType::TransformChanged);
+					}
+				}
+			}
+		}
+		
+		return Modifications;
+	};
 	
+	return BasicModificationCheck(SkeletalMesh->GetRefSkeleton());
+}
+
+ESkeletonModificationType USkeletonModifier::PreCommitSkeleton(const ESkeletalMeshModificationType InSkeletalMeshModifications) const
+{
+	const bool bNeedSkeletonUpdate = EnumHasAnyFlags(InSkeletalMeshModifications, ESkeletalMeshModificationType::SkeletonUpdated);
+	if (!bNeedSkeletonUpdate)
+	{
+		return ESkeletonModificationType::None;
+	}
+
+	USkeleton* Skeleton = SkeletalMesh->GetSkeleton();
+	USkeletonModifierLocals::FReferenceSkeletonCompatibilityChecker Checker(Skeleton->GetReferenceSkeleton());
+	if (Checker.IsCompatibleReferenceSkeleton(*ReferenceSkeleton))
+	{
+		// skeleton is compatible
+		return ESkeletonModificationType::SimpleMerge;
+	}
+
+	USkeletalMeshMergeOptions* Options = NewObject<USkeletalMeshMergeOptions>();
+	FPropertyEditorModule& PropertyEditorModule = FModuleManager::GetModuleChecked<FPropertyEditorModule>("PropertyEditor");
+	
+	FDetailsViewArgs DetailsViewArgs;
+	DetailsViewArgs.bAllowSearch = false;
+	DetailsViewArgs.NameAreaSettings = FDetailsViewArgs::HideNameArea;
+	TSharedRef<IDetailsView> DetailsView = PropertyEditorModule.CreateDetailView(DetailsViewArgs);
+	static constexpr bool bForceRefresh = true;
+	DetailsView->SetObject(Options, bForceRefresh);
+
+	TSharedRef<SCustomDialog> OptionsDialog = SNew(SCustomDialog)
+		.Title(NSLOCTEXT("SkeletonModifier", "SkeletonModifierMergeDialog", "SkeletonModifier Merge Options"))
+		.Content()
+		[
+			SNew(SVerticalBox)
+
+			+ SVerticalBox::Slot()
+			.Padding(4.f, 2.f)
+			[
+				SNew(STextBlock)
+				.Text(NSLOCTEXT("SkeletonModifier", "SkeletonModifierMergeText",
+				"The current changes to the edited bone hierarchy are incompatible with the assigned skeleton asset.\n"
+				"'Commit' to commit the current changes using the merge type below.\n"
+				"'Cancel' to cancel the current changes and revert to the previous skeleton.\n"))
+			]
+
+			+ SVerticalBox::Slot()
+			.AutoHeight()
+			.Padding(4.f, 2.f)
+			[
+				SNew(SBox)
+				.MinDesiredWidth(450)
+				[
+					DetailsView
+				]
+			]
+		]
+		.Buttons({
+			SCustomDialog::FButton(NSLOCTEXT("SkeletonModifier", "CommitButton", "Commit")),
+			SCustomDialog::FButton(NSLOCTEXT("SkeletonModifier", "CancelButton", "Cancel"))
+		});
+
+	const int32 Choice = OptionsDialog->ShowModal();
+	if (Choice == 1)
+	{
+		return ESkeletonModificationType::Cancel;
+	}
+
+	switch (Options->MergeType)
+	{
+		case ESKeletalMeshMergeType::New:
+			return ESkeletonModificationType::DuplicateAndMerge;
+		case ESKeletalMeshMergeType::Merge:
+			return Options->bMergeAll ? ESkeletonModificationType::FullMergeAll : ESkeletonModificationType::FullMerge;
+		default:
+			break;
+	}
+
+	return ESkeletonModificationType::None;
+}
+
+bool USkeletonModifier::CommitSkeletonToSkeletalMesh()
+{
+	if (!SkeletalMesh.IsValid() || !ReferenceSkeleton || !MeshDescription)
+	{
+		UE_LOG(LogAnimation, Error, TEXT("Skeleton Modifier: No mesh loaded. Cannot apply skeleton edits."));
+		return false;
+	}
+
+#if WITH_EDITORONLY_DATA
+	
+	// check modifications at the skeletal mesh level
+	const ESkeletalMeshModificationType Modifications = PreCommitSkeletalMesh();
+	if (Modifications == ESkeletalMeshModificationType::None)
+	{
+		UE_LOG(LogAnimation, Warning, TEXT("Skeleton Modifier: No modification needed."));
+		return false;
+	}
+
+	// check modifications at the skeleton level
+	const ESkeletonModificationType SkeletonModifications = PreCommitSkeleton(Modifications);
+	if (SkeletonModifications == ESkeletonModificationType::Cancel)
+	{
+		UE_LOG(LogAnimation, Warning, TEXT("Skeleton Modifier: Skeleton can't be modified."));
+		return false;
+	}
+	
+	const TArray<FMeshBoneInfo>& BoneInfos = ReferenceSkeleton->GetRawRefBoneInfo();
+	
 	// update mesh description
 	{
@@ -316,18 +634,7 @@
 		}
 		
 		// update skin data if needed
-		auto HasBoneIndexesChanged = [&]()
+		if (EnumHasAnyFlags(Modifications, ESkeletalMeshModificationType::IndicesUpdated))
 		{
-			for (int32 Index = 0; Index < BoneIndexTracker.Num(); Index++)
-			{
-				if (BoneIndexTracker[Index] != Index)
-				{
-					return true;
-				}
-			}
-			return false;
-		};
-		if (HasBoneIndexesChanged())
-		{
 			using namespace UE::AnimationCore;
 			FBoneWeightsSettings BoneSettings; BoneSettings.SetNormalizeType(EBoneWeightNormalizeType::None);
@@ -384,8 +691,11 @@
 	FlushRenderingCommands();
 
 	// call modify on the skeleton first as post undo will re-register components so it must be done once both
-	// skeletal mesh and skeleton are up to date, so it must be done after the skeletal mesh has been undone 
-	Skeleton->Modify();
+	// skeletal mesh and skeleton are up to date, so it must be done after the skeletal mesh has been undone
+	if (EnumHasAnyFlags(SkeletonModifications, ESkeletonModificationType::DoUpdate))
+	{
+		Skeleton->Modify();
+	}
 
 	SkeletalMesh->SetFlags(RF_Transactional);
 	SkeletalMesh->Modify();
@@ -401,15 +711,49 @@
 	SkeletalMesh->CommitMeshDescription(USkeletonModifierLocals::LODIndex);
 
 	// update skeleton
-	if (Skeleton->RecreateBoneTree(SkeletalMesh.Get()))
+	if (EnumHasAnyFlags(SkeletonModifications, ESkeletonModificationType::DoUpdate))
 	{
-		// restore retargeting modes
-		for (int32 BoneIndex = 0, NumBones = BoneInfos.Num(); BoneIndex < NumBones; ++BoneIndex)
+		NotifyFromSkeletonChanges();
+		
+		auto UpdateSkeleton = [this, SkeletonModifications, &RetargetingModes]()
 		{
-			Skeleton->SetBoneTranslationRetargetingMode(BoneIndex, RetargetingModes[BoneIndex]);
+			USkeleton* Skeleton = SkeletalMesh->GetSkeleton();
+			if (SkeletonModifications == ESkeletonModificationType::SimpleMerge)
+			{
+				return Skeleton->MergeAllBonesToBoneTree(SkeletalMesh.Get());
+			}
+		
+			bool bSkeletonModified = false;	
+			if (EnumHasAnyFlags(SkeletonModifications, ESkeletonModificationType::DeepMerge))
+			{
+				bSkeletonModified = Skeleton->RecreateBoneTree(SkeletalMesh.Get());
+			}
+
+			if (!bSkeletonModified)
+			{
+				return false;
+			}
+
+			// restore retargeting modes
+			const TArray<FMeshBoneInfo>& BoneInfos = ReferenceSkeleton->GetRawRefBoneInfo();
+			for (int32 BoneIndex = 0, NumBones = BoneInfos.Num(); BoneIndex < NumBones; ++BoneIndex)
+			{
+				Skeleton->SetBoneTranslationRetargetingMode(BoneIndex, RetargetingModes[BoneIndex]);
+			}
+
+			return true;
+		};
+	
+		if (UpdateSkeleton())
+		{
+			PostCommitSkeleton(SkeletonModifications);
+			Skeleton->MarkPackageDirty();
+			FAssetNotifications::SkeletonNeedsToBeSaved(Skeleton);
 		}
-		
-		Skeleton->MarkPackageDirty();
+		else
+		{
+			ensure(false);
+		}
 	}
 	
 	// must be done once the skeleton is up to date
@@ -422,6 +766,81 @@
 	return false;
 }
 
+void USkeletonModifier::PostCommitSkeleton(const ESkeletonModificationType InSkeletonModifications) const
+{
+	if (InSkeletonModifications != ESkeletonModificationType::FullMergeAll)
+	{
+		return;
+	}
+
+	USkeleton* Skeleton = SkeletalMesh->GetSkeleton();
+	
+	TArray<const USkeletalMesh*> OtherSkeletalMeshUsingSkeleton;
+
+	const IAssetRegistry& AssetRegistry = FModuleManager::LoadModuleChecked<FAssetRegistryModule>("AssetRegistry").Get();
+	
+	FARFilter ARFilter;
+	ARFilter.ClassPaths.Add(USkeletalMesh::StaticClass()->GetClassPathName());
+	ARFilter.TagsAndValues.Add(TEXT("Skeleton"), FAssetData(Skeleton).GetExportTextName());
+
+	TArray<FAssetData> SkeletalMeshAssetData;
+	if (AssetRegistry.GetAssets(ARFilter, SkeletalMeshAssetData))
+	{
+		for (const FAssetData& AssetData: SkeletalMeshAssetData)
+		{
+			const USkeletalMesh* ExtraSkeletalMesh = Cast<USkeletalMesh>(AssetData.GetAsset());
+			if (IsValid(ExtraSkeletalMesh) && ExtraSkeletalMesh != SkeletalMesh)
+			{
+				OtherSkeletalMeshUsingSkeleton.Add(ExtraSkeletalMesh);
+			}
+		}
+	}
+	
+	for (const USkeletalMesh* ExtraSkeletalMesh : OtherSkeletalMeshUsingSkeleton)
+	{
+		// merge still can fail
+		if (!Skeleton->MergeAllBonesToBoneTree(ExtraSkeletalMesh))
+		{
+			FMessageDialog::Open(EAppMsgType::Ok,
+				FText::Format(NSLOCTEXT("SkeletonModifier", "SkeletonModifier_RemergingBones", "Failed to merge SkeletalMesh '{0}'."), FText::FromString(ExtraSkeletalMesh->GetName())));
+		}
+	}
+}
+
+void USkeletonModifier::NotifyFromSkeletonChanges() const
+{
+	// check assets using this skeleton ?
+	if (bDebug)
+	{
+		// TODO avoid certain changes when the skeleton is referenced by other assets (i.e. changing the skeletal mesh's
+		// reference skeleton poses is fine, re-parenting/removing bones, etc. is not)
+		static const TArray<FTopLevelAssetPath> AssetPaths({
+		   UAnimSequence::StaticClass()->GetClassPathName(),
+		   UAnimMontage::StaticClass()->GetClassPathName(),
+		   UPoseAsset::StaticClass()->GetClassPathName(),
+		   USkeletalMesh::StaticClass()->GetClassPathName()});
+	
+		FARFilter Filter;
+		Filter.ClassPaths = AssetPaths;
+	
+		const IAssetRegistry& AssetRegistry = FModuleManager::LoadModuleChecked<FAssetRegistryModule>("AssetRegistry").Get();
+		TArray<FAssetData> Assets; AssetRegistry.GetAssets(Filter, Assets);
+	
+		const FAssetData SkeletonAssetData(SkeletalMesh->GetSkeleton());
+		const FString SkeletonPath = SkeletonAssetData.GetExportTextName();
+	
+		static const FName Tag("Skeleton");
+		for (const FAssetData& AssetData: Assets)
+		{
+			const FString TagValue = AssetData.GetTagValueRef<FString>(Tag);
+			if (TagValue == SkeletonPath)
+			{
+				UE_LOG(LogAnimation, Warning, TEXT("%s references that skeleton."), *AssetData.GetExportTextName());
+			}
+		}
+	}
+}
+
 bool USkeletonModifier::AddBone(const FName InBoneName, const FName InParentName, const FTransform& InTransform)
 {
 	if (InBoneName == NAME_None)
diff --git a/Engine/Plugins/Experimental/MeshModelingToolsetExp/Source/SkeletalMeshModifiers/Public/SkeletonModifier.h b/Engine/Plugins/Experimental/MeshModelingToolsetExp/Source/SkeletalMeshModifiers/Public/SkeletonModifier.h
--- a/Engine/Plugins/Experimental/MeshModelingToolsetExp/Source/SkeletalMeshModifiers/Public/SkeletonModifier.h
+++ b/Engine/Plugins/Experimental/MeshModelingToolsetExp/Source/SkeletalMeshModifiers/Public/SkeletonModifier.h
@@ -7,6 +7,7 @@
 #include "Templates/UniquePtr.h"
 #include "MeshDescription.h"
 #include "ReferenceSkeleton.h"
+#include "Misc/EnumClassFlags.h"
 
 #include "SkeletonModifier.generated.h"
 
@@ -16,6 +17,38 @@ struct FMeshDescription;
 struct FReferenceSkeleton;
 struct FTransformComposer;
 
+enum class ESkeletalMeshModificationType : uint32
+{
+	None = 0x000,
+
+	BonesAdded = 0x001,
+	BonesRemoved = 0x002,
+	BonesRenamed = 0x004,
+	TransformChanged = 0x008,
+	HierarchyChanged = 0x010,
+
+	IndicesUpdated = HierarchyChanged | BonesRemoved,
+
+	SkeletonUpdated = BonesAdded | BonesRenamed | HierarchyChanged
+};
+ENUM_CLASS_FLAGS(ESkeletalMeshModificationType)
+
+enum class ESkeletonModificationType : uint32
+{
+	Cancel = 0x000,
+	None = 0x001,
+	
+	SimpleMerge = 0x002,
+	DuplicateAndMerge = 0x004,
+	FullMerge = 0x008,
+	FullMergeAll = 0x010,
+
+	DeepMerge = DuplicateAndMerge | FullMerge | FullMergeAll,
+	
+	DoUpdate = SimpleMerge | DeepMerge
+};
+ENUM_CLASS_FLAGS(ESkeletonModificationType)
+
 /**
  * FMirrorOptions
  */
@@ -230,6 +263,14 @@ private:
 	// orient function
 	void GetBonesToOrient(const TArray<FName>& InBonesName, const FOrientOptions& InOptions, TArray<int32>& OutBonesToOrient) const;
 
+	// pre/post commit
+	ESkeletalMeshModificationType PreCommitSkeletalMesh();
+	ESkeletonModificationType PreCommitSkeleton(const ESkeletalMeshModificationType InSkeletalMeshModifications) const;
+	void PostCommitSkeleton(const ESkeletonModificationType InSkeletonModifications) const;
+
+	// notification
+	void NotifyFromSkeletonChanges() const;
+	
 	UPROPERTY(transient)
 	TWeakObjectPtr<USkeletalMesh> SkeletalMesh = nullptr;
 	
@@ -238,6 +279,11 @@ private:
 	TUniquePtr<FTransformComposer> TransformComposer;
 	TArray<int32> BoneIndexTracker;
 	bool bRestricted = false;
+
+public:
+	
+	// UPROPERTY(EditAnywhere, Category = Debug)
+	bool bDebug = false;
 };
 
 struct FTransformComposer
@@ -295,3 +341,25 @@ struct FTransformComposer
 	TBitArray<> TransformCached;
 };
 
+UENUM()
+enum class ESKeletalMeshMergeType : uint8
+{
+	New		UMETA(ToolTip = "Create a new skeleton asset from the current changes."),
+	Merge	UMETA(ToolTip = "Merge the current changes to the existing skeleton asset. Note that this may invalidate dependent assets, such as animation clips, pose assets and animation blueprints."),
+};
+
+UCLASS(HideCategories=Object, MinimalAPI)
+class USkeletalMeshMergeOptions : public UObject
+{
+	GENERATED_BODY()
+
+public:
+
+	/** Changes merge type (New or Merge) */
+	UPROPERTY(EditAnywhere, Category="Merge")
+	ESKeletalMeshMergeType MergeType = ESKeletalMeshMergeType::New;
+
+	/** Also apply the changes to the skeletal meshes referencing the same skeleton */
+	UPROPERTY(EditAnywhere, Category="Merge", meta=(EditCondition="MergeType==ESKeletalMeshMergeType::Merge"))
+	bool bMergeAll = false;
+};
diff --git a/Engine/Plugins/Experimental/MeshModelingToolsetExp/Source/SkeletalMeshModifiers/SkeletalMeshModifiers.build.cs b/Engine/Plugins/Experimental/MeshModelingToolsetExp/Source/SkeletalMeshModifiers/SkeletalMeshModifiers.build.cs
--- a/Engine/Plugins/Experimental/MeshModelingToolsetExp/Source/SkeletalMeshModifiers/SkeletalMeshModifiers.build.cs
+++ b/Engine/Plugins/Experimental/MeshModelingToolsetExp/Source/SkeletalMeshModifiers/SkeletalMeshModifiers.build.cs
@@ -26,7 +26,11 @@ public class SkeletalMeshModifiers : ModuleRules
 				"RenderCore",
 				"CoreUObject",
 				"Engine",
-				"ApplicationCore"
+				"ApplicationCore",
+				"SlateCore",
+				"Slate",
+				"UnrealEd",
+				"ToolWidgets"
 			}
 			);
 
