Skin Weights Paint Tool: remove useless bShowSourcePreview property

#rnx

diff --git a/Engine/Plugins/Experimental/MeshModelingToolsetExp/Source/MeshModelingToolsEditorOnlyExp/Private/SkeletalMesh/SkinWeightsPaintTool.cpp b/Engine/Plugins/Experimental/MeshModelingToolsetExp/Source/MeshModelingToolsEditorOnlyExp/Private/SkeletalMesh/SkinWeightsPaintTool.cpp
--- a/Engine/Plugins/Experimental/MeshModelingToolsetExp/Source/MeshModelingToolsEditorOnlyExp/Private/SkeletalMesh/SkinWeightsPaintTool.cpp
+++ b/Engine/Plugins/Experimental/MeshModelingToolsetExp/Source/MeshModelingToolsEditorOnlyExp/Private/SkeletalMesh/SkinWeightsPaintTool.cpp
@@ -1563,7 +1563,7 @@ FBox USkinWeightsPaintTool::GetWorldSpaceFocusBox()
 	static constexpr bool bNonColliding = true;
 	FBox PreviewBox = PreviewMesh->GetActor()->GetComponentsBoundingBox(bNonColliding);
 	
-	if (WeightToolProperties->bShowSourcePreview && SourcePreviewMesh)
+	if (SourcePreviewMesh)
 	{
 		if (AActor* SourceActor = SourcePreviewMesh->GetActor())
 		{
@@ -4192,14 +4192,6 @@ void USkinWeightsPaintTool::OnPropertyModified(UObject* ModifiedObject, FPropert
 		}
 	}
 
-	if (ModifiedProperty->GetName() == GET_MEMBER_NAME_STRING_CHECKED(USkinWeightsPaintToolProperties, bShowSourcePreview))
-	{
-		if (SourcePreviewMesh)
-		{
-			SourcePreviewMesh->SetVisible(WeightToolProperties->bShowSourcePreview);
-		}
-	}
-
 	if (ModifiedProperty->GetName() == GET_MEMBER_NAME_STRING_CHECKED(USkinWeightsPaintToolProperties, SourcePreviewOffset))
 	{
 		if (SourcePreviewMesh)
@@ -4253,8 +4245,6 @@ void USkinWeightsPaintTool::ResetSourceForTransfer(USkeletalMesh* InSkeletalMesh
 
 		const FComponentMaterialSet MaterialSet = UE::ToolTarget::GetMaterialSet(SourceTarget);
 		SourcePreviewMesh->SetMaterials(MaterialSet.Materials);
-
-		SourcePreviewMesh->SetVisible(WeightToolProperties->bShowSourcePreview);
 	}
 }
 
diff --git a/Engine/Plugins/Experimental/MeshModelingToolsetExp/Source/MeshModelingToolsEditorOnlyExp/Public/SkeletalMesh/SkinWeightsPaintTool.h b/Engine/Plugins/Experimental/MeshModelingToolsetExp/Source/MeshModelingToolsEditorOnlyExp/Public/SkeletalMesh/SkinWeightsPaintTool.h
--- a/Engine/Plugins/Experimental/MeshModelingToolsetExp/Source/MeshModelingToolsEditorOnlyExp/Public/SkeletalMesh/SkinWeightsPaintTool.h
+++ b/Engine/Plugins/Experimental/MeshModelingToolsetExp/Source/MeshModelingToolsEditorOnlyExp/Public/SkeletalMesh/SkinWeightsPaintTool.h
@@ -445,9 +445,6 @@ public:
 	UPROPERTY(EditAnywhere, Category = "WeightTransfer|SkinWeightLayer", meta = (DisplayName = "Source Profile", GetOptions = GetSourceSkinWeightProfilesFunc))
 	FName SourceSkinWeightProfile = FSkeletalMeshAttributesShared::DefaultSkinWeightProfileName;
 	
-	UPROPERTY(EditAnywhere, Transient, Category = "WeightTransfer|Preview")
-	bool bShowSourcePreview = false;
-	
 	UPROPERTY(EditAnywhere, Transient, Category = "WeightTransfer|Preview")
 	FTransform SourcePreviewOffset = FTransform::Identity;
 	
