Skin weight transfer: Add support for selecting vertices on the source mesh

diff --git a/Engine/Plugins/Experimental/Animation/SkeletalMeshModelingTools/Source/SkeletalMeshModelingTools/Private/DetailCustomization/SkinWeightDetailCustomization.cpp b/Engine/Plugins/Experimental/Animation/SkeletalMeshModelingTools/Source/SkeletalMeshModelingTools/Private/DetailCustomization/SkinWeightDetailCustomization.cpp
--- a/Engine/Plugins/Experimental/Animation/SkeletalMeshModelingTools/Source/SkeletalMeshModelingTools/Private/DetailCustomization/SkinWeightDetailCustomization.cpp
+++ b/Engine/Plugins/Experimental/Animation/SkeletalMeshModelingTools/Source/SkeletalMeshModelingTools/Private/DetailCustomization/SkinWeightDetailCustomization.cpp
@@ -141,23 +141,25 @@ void FSkinWeightDetailCustomization::CustomizeDetails(IDetailLayoutBuilder& Deta
 	IDetailCategoryBuilder& SkinWeightLayerCategory = DetailBuilder.EditCategory("SkinWeightLayer", FText::GetEmpty(), ECategoryPriority::Important);
 	SkinWeightLayerCategory.InitiallyCollapsed(true);
 
-	// hide all base brush properties that have been customized
-	const TSharedRef<IPropertyHandle> BrushModeHandle = DetailBuilder.GetProperty(GET_MEMBER_NAME_CHECKED(USkinWeightsPaintToolProperties, BrushMode));
-	DetailBuilder.HideProperty(BrushModeHandle);
-	const TSharedRef<IPropertyHandle> BrushSizeHandle = DetailBuilder.GetProperty(GET_MEMBER_NAME_CHECKED(UBrushBaseProperties, BrushSize), UBrushBaseProperties::StaticClass());
-	DetailBuilder.HideProperty(BrushSizeHandle);
-	const TSharedRef<IPropertyHandle> BrushStrengthHandle = DetailBuilder.GetProperty(GET_MEMBER_NAME_CHECKED(UBrushBaseProperties, BrushStrength), UBrushBaseProperties::StaticClass());
-	DetailBuilder.HideProperty(BrushStrengthHandle);
-	const TSharedRef<IPropertyHandle> BrushFalloffHandle = DetailBuilder.GetProperty(GET_MEMBER_NAME_CHECKED(UBrushBaseProperties, BrushFalloffAmount), UBrushBaseProperties::StaticClass());
-	DetailBuilder.HideProperty(BrushFalloffHandle);
-	const TSharedRef<IPropertyHandle> BrushRadiusHandle = DetailBuilder.GetProperty(GET_MEMBER_NAME_CHECKED(UBrushBaseProperties, BrushRadius), UBrushBaseProperties::StaticClass());
-	DetailBuilder.HideProperty(BrushRadiusHandle);
-	const TSharedRef<IPropertyHandle> SpecifyRadiusHandle = DetailBuilder.GetProperty(GET_MEMBER_NAME_CHECKED(UBrushBaseProperties, bSpecifyRadius), UBrushBaseProperties::StaticClass());
-	DetailBuilder.HideProperty(SpecifyRadiusHandle);
-	const TSharedRef<IPropertyHandle> EditModePropHandle = DetailBuilder.GetProperty(GET_MEMBER_NAME_CHECKED(USkinWeightsPaintToolProperties, EditingMode));
-	DetailBuilder.HideProperty(EditModePropHandle);
-	const TSharedRef<IPropertyHandle> ColorModePropHandle = DetailBuilder.GetProperty(GET_MEMBER_NAME_CHECKED(USkinWeightsPaintToolProperties, ColorMode));
-	DetailBuilder.HideProperty(ColorModePropHandle);
+	// hide all properties that have been customized
+	TArray<FName> PropertiesToHide;
+	PropertiesToHide.Add(GET_MEMBER_NAME_CHECKED(USkinWeightsPaintToolProperties, BrushMode));
+	PropertiesToHide.Add(GET_MEMBER_NAME_CHECKED(UBrushBaseProperties, BrushSize));
+	PropertiesToHide.Add(GET_MEMBER_NAME_CHECKED(UBrushBaseProperties, BrushStrength));
+	PropertiesToHide.Add(GET_MEMBER_NAME_CHECKED(UBrushBaseProperties, BrushFalloffAmount));
+	PropertiesToHide.Add(GET_MEMBER_NAME_CHECKED(UBrushBaseProperties, BrushRadius));
+	PropertiesToHide.Add(GET_MEMBER_NAME_CHECKED(UBrushBaseProperties, bSpecifyRadius));
+	PropertiesToHide.Add(GET_MEMBER_NAME_CHECKED(USkinWeightsPaintToolProperties, EditingMode));
+	PropertiesToHide.Add(GET_MEMBER_NAME_CHECKED(USkinWeightsPaintToolProperties, ColorMode));
+	PropertiesToHide.Add(GET_MEMBER_NAME_CHECKED(USkinWeightsPaintToolProperties, MeshSelectMode));
+	PropertiesToHide.Add(GET_MEMBER_NAME_CHECKED(USkinWeightsPaintToolProperties, SourceSkeletalMesh));
+	PropertiesToHide.Add(GET_MEMBER_NAME_CHECKED(USkinWeightsPaintToolProperties, SourceLOD));
+	PropertiesToHide.Add(GET_MEMBER_NAME_CHECKED(USkinWeightsPaintToolProperties, SourceSkinWeightProfile));
+	for (const FName PropertyToHide : PropertiesToHide)
+	{
+		TSharedRef<IPropertyHandle> Property = DetailBuilder.GetProperty(PropertyToHide);
+		DetailBuilder.HideProperty(Property);
+	}
 }
 
 void FSkinWeightDetailCustomization::AddBrushUI(IDetailLayoutBuilder& DetailBuilder) const
@@ -322,7 +324,7 @@ void FSkinWeightDetailCustomization::AddBrushUI(IDetailLayoutBuilder& DetailBuil
 
 void FSkinWeightDetailCustomization::AddSelectionUI(IDetailLayoutBuilder& DetailBuilder) const
 {
-	// custom display of weight editing tools
+	// custom display of selection editing tools
 	IDetailCategoryBuilder& EditSelectionCategory = DetailBuilder.EditCategory("Edit Selection", FText::GetEmpty(), ECategoryPriority::Important);
 	EditSelectionCategory.InitiallyCollapsed(true);
 
@@ -375,7 +377,7 @@ void FSkinWeightDetailCustomization::AddSelectionUI(IDetailLayoutBuilder& Detail
 	ToolbarBuilder.EndBlockGroup();
 	ToolbarBuilder.EndSection();
 
-	// GROW/SHRINK/FLOOD Selection category
+	// edit selection category
 	EditSelectionCategory.AddCustomRow(LOCTEXT("EditSelectionRow", "Edit Selection"), false)
 	.WholeRowContent()
 	[
@@ -385,6 +387,7 @@ void FSkinWeightDetailCustomization::AddSelectionUI(IDetailLayoutBuilder& Detail
 		[
 			SNew(SHorizontalBox)
 
+			// ISOLATE SELECTION
 			+SHorizontalBox::Slot()
 			.HAlign(HAlign_Center)
 			.VAlign(VAlign_Center)
@@ -405,10 +408,12 @@ void FSkinWeightDetailCustomization::AddSelectionUI(IDetailLayoutBuilder& Detail
 					.HAlign(HAlign_Center)
 					.ToolTipText(LOCTEXT("IsolateSelectedTooltip",
 							"Shows only the selected faces in the viewport.\n"
-							"Weight editing operations will not affect hidden vertices.\n "))
+							"Weight editing operations will not affect hidden vertices.\n"
+							"NOTE: This only works on the target (main) mesh."))
 					.IsEnabled_Lambda([this]()
 					{
-						const bool bHasSelection = Tool->IsAnyComponentSelected();
+						// isolated selection only available on main mesh (for now)
+						const bool bHasSelection = Tool->GetMainMeshSelector()->IsAnyComponentSelected();
 						const bool bAlreadyIsolatingSelection = Tool->IsSelectionIsolated();
 						return bHasSelection ||  bAlreadyIsolatingSelection;
 					})
@@ -434,7 +439,7 @@ void FSkinWeightDetailCustomization::AddSelectionUI(IDetailLayoutBuilder& Detail
 						{
 							if (Tool->IsSelectionIsolated())
 							{
-								return LOCTEXT("ShowAllButtonLabel", "Show All");
+								return LOCTEXT("ShowAllButtonLabel", "Show Full Mesh");
 							}
 								
 							return LOCTEXT("IsolateButtonLabel", "Isolate Selected");
@@ -444,6 +449,7 @@ void FSkinWeightDetailCustomization::AddSelectionUI(IDetailLayoutBuilder& Detail
 			]
 		]
 
+		// GROW / SHRINK / FLOOD
 		+SVerticalBox::Slot()
 		[
 			SNew(SHorizontalBox)
@@ -452,7 +458,7 @@ void FSkinWeightDetailCustomization::AddSelectionUI(IDetailLayoutBuilder& Detail
 			.Padding(2.f, WeightEditVerticalPadding)
 			[
 				SNew(SButton)
-				.IsEnabled_Lambda([this]{ return ToolSettings->WeightTool->HasSelectedVertices(); })
+				.IsEnabled_Lambda([this]{ return Tool->GetActiveMeshSelector()->IsAnyComponentSelected(); })
 				.HAlign(HAlign_Center)
 				.VAlign(VAlign_Center)
 				.Text(LOCTEXT("GrowSelectionButtonLabel", "Grow"))
@@ -460,7 +466,7 @@ void FSkinWeightDetailCustomization::AddSelectionUI(IDetailLayoutBuilder& Detail
 						"Grow the current selection by adding connected neighbors to current selection.\n"))
 				.OnClicked_Lambda([this]()
 				{
-					Tool->GrowSelection();
+					Tool->GetActiveMeshSelector()->GrowSelection();
 					return FReply::Handled();
 				})
 			]
@@ -469,7 +475,7 @@ void FSkinWeightDetailCustomization::AddSelectionUI(IDetailLayoutBuilder& Detail
 			.Padding(2.f, WeightEditVerticalPadding)
 			[
 				SNew(SButton)
-				.IsEnabled_Lambda([this]{ return ToolSettings->WeightTool->HasSelectedVertices(); })
+				.IsEnabled_Lambda([this]{ return Tool->GetActiveMeshSelector()->IsAnyComponentSelected(); })
 				.HAlign(HAlign_Center)
 				.VAlign(VAlign_Center)
 				.Text(LOCTEXT("ShrinkSelectionButtonLabel", "Shrink"))
@@ -477,7 +483,7 @@ void FSkinWeightDetailCustomization::AddSelectionUI(IDetailLayoutBuilder& Detail
 						"Shrink the current selection by removing components on the border of the current selection.\n"))
 				.OnClicked_Lambda([this]()
 				{
-					Tool->ShrinkSelection();
+					Tool->GetActiveMeshSelector()->ShrinkSelection();
 					return FReply::Handled();
 				})
 			]
@@ -486,7 +492,7 @@ void FSkinWeightDetailCustomization::AddSelectionUI(IDetailLayoutBuilder& Detail
 			.Padding(2.f, WeightEditVerticalPadding)
 			[
 				SNew(SButton)
-				.IsEnabled_Lambda([this]{ return ToolSettings->WeightTool->HasSelectedVertices(); })
+				.IsEnabled_Lambda([this]{ return Tool->GetActiveMeshSelector()->IsAnyComponentSelected(); })
 				.HAlign(HAlign_Center)
 				.VAlign(VAlign_Center)
 				.Text(LOCTEXT("FloodSelectionButtonLabel", "Flood"))
@@ -494,12 +500,13 @@ void FSkinWeightDetailCustomization::AddSelectionUI(IDetailLayoutBuilder& Detail
 						"Flood the current selection by adding all connected components to the current selection.\n"))
 				.OnClicked_Lambda([this]()
 				{
-					Tool->FloodSelection();
+					Tool->GetActiveMeshSelector()->FloodSelection();
 					return FReply::Handled();
 				})
 			]
 		]
 
+		// SELECT AFFECTED VERTICES
 		+SVerticalBox::Slot()
 		[
 			SNew(SHorizontalBox)
@@ -508,6 +515,11 @@ void FSkinWeightDetailCustomization::AddSelectionUI(IDetailLayoutBuilder& Detail
 			.Padding(2.f, WeightEditVerticalPadding)
 			[
 				SNew(SButton)
+				.IsEnabled_Lambda([this]
+				{
+					// only allow selecting affected vertices on the target/main mesh
+					return ToolSettings->MeshSelectMode == EMeshTransferOption::Target;
+				})
 				.HAlign(HAlign_Center)
 				.VAlign(VAlign_Center)
 				.Text(LOCTEXT("AffectedSelectionButtonLabel", "Affected"))
@@ -525,21 +537,20 @@ void FSkinWeightDetailCustomization::AddSelectionUI(IDetailLayoutBuilder& Detail
 			.Padding(2.f, WeightEditVerticalPadding)
 			[
 				SNew(SButton)
-				.IsEnabled_Lambda([this]{ return ToolSettings->WeightTool->HasSelectedVertices(); })
+				.IsEnabled_Lambda([this]{ return Tool->GetActiveMeshSelector()->IsAnyComponentSelected(); })
 				.HAlign(HAlign_Center)
 				.VAlign(VAlign_Center)
 				.Text(LOCTEXT("BorderSelectionButtonLabel", "Convert to Border"))
 				.ToolTipText(LOCTEXT("BorderSelectionTooltip", "Select vertices on the border of the current selection."))
 				.OnClicked_Lambda([this]()
 				{
-					Tool->SelectBorder();
+					Tool->GetActiveMeshSelector()->SelectBorder();
 					return FReply::Handled();
 				})
 			]
 		]
 	];
-
-	// custom display of weight editing tools
+	
 	IDetailCategoryBuilder& EditWeightsCategory = DetailBuilder.EditCategory("Edit Weights", FText::GetEmpty(), ECategoryPriority::Important);
 	EditWeightsCategory.InitiallyCollapsed(true);
 
@@ -550,11 +561,11 @@ void FSkinWeightDetailCustomization::AddSelectionUI(IDetailLayoutBuilder& Detail
 		SNew(STextBlock)
 		.Visibility_Lambda([this]()
 		{
-			return Tool->GetSelectedVertices().IsEmpty() ? EVisibility::Visible : EVisibility::Collapsed;
+			return !Tool->HasActiveSelectionOnMainMesh() ? EVisibility::Visible : EVisibility::Collapsed;
 		})
 		.Text_Lambda([this]()
 		{
-			return LOCTEXT("NothingSelectedLabel", "Select vertices to edit weights...");
+			return LOCTEXT("NothingSelectedLabel", "Select vertices on target mesh to edit weights...");
 		})
 	];
 
@@ -564,7 +575,7 @@ void FSkinWeightDetailCustomization::AddSelectionUI(IDetailLayoutBuilder& Detail
 	.WholeRowContent()
 	[
 		SNew(SHorizontalBox)
-		.IsEnabled_Lambda([this]{ return ToolSettings->WeightTool->HasSelectedVertices(); })
+		.IsEnabled_Lambda([this]{ return Tool->HasActiveSelectionOnMainMesh(); })
 
 		+SHorizontalBox::Slot()
 		[
@@ -642,7 +653,7 @@ void FSkinWeightDetailCustomization::AddSelectionUI(IDetailLayoutBuilder& Detail
 							Value,
 							0, // iterations
 							ToolSettings->DirectEditState.EditMode,
-							Tool->GetSelectedVertices(),
+							Tool->GetMainMeshSelector()->GetSelectedVertices(),
 							bShouldTransact);
 					}
 				})
@@ -656,7 +667,7 @@ void FSkinWeightDetailCustomization::AddSelectionUI(IDetailLayoutBuilder& Detail
 							NewValue,
 							0, // iterations
 							ToolSettings->DirectEditState.EditMode,
-							Tool->GetSelectedVertices(),
+							Tool->GetMainMeshSelector()->GetSelectedVertices(),
 							bShouldTransact);
 					}
 					ToolSettings->DirectEditState.bInTransaction = false;
@@ -699,7 +710,7 @@ void FSkinWeightDetailCustomization::AddSelectionUI(IDetailLayoutBuilder& Detail
 	.WholeRowContent()
 	[
 		SNew(SVerticalBox)
-		.IsEnabled_Lambda([this]{ return ToolSettings->WeightTool->HasSelectedVertices(); })
+		.IsEnabled_Lambda([this]{ return Tool->HasActiveSelectionOnMainMesh(); })
 
 		+ SVerticalBox::Slot()
 		.Padding(WeightEditHorizontalPadding, WeightEditVerticalPadding)
@@ -723,7 +734,7 @@ void FSkinWeightDetailCustomization::AddSelectionUI(IDetailLayoutBuilder& Detail
 							ToolSettings->AddStrength,
 							0, // iterations 
 							EWeightEditOperation::Add,
-							Tool->GetSelectedVertices(),
+							Tool->GetMainMeshSelector()->GetSelectedVertices(),
 							bShouldTransact);
 						return FReply::Handled();
 					})
@@ -757,7 +768,7 @@ void FSkinWeightDetailCustomization::AddSelectionUI(IDetailLayoutBuilder& Detail
 	.WholeRowContent()
 	[
 		SNew(SVerticalBox)
-		.IsEnabled_Lambda([this]{ return ToolSettings->WeightTool->HasSelectedVertices(); })
+		.IsEnabled_Lambda([this]{ return Tool->HasActiveSelectionOnMainMesh(); })
 
 		+ SVerticalBox::Slot()
 		.Padding(WeightEditHorizontalPadding, WeightEditVerticalPadding)
@@ -782,7 +793,7 @@ void FSkinWeightDetailCustomization::AddSelectionUI(IDetailLayoutBuilder& Detail
 							ToolSettings->ReplaceValue,
 							0, // iterations
 							EWeightEditOperation::Replace,
-							Tool->GetSelectedVertices(),
+							Tool->GetMainMeshSelector()->GetSelectedVertices(),
 							bShouldTransact);
 						return FReply::Handled();
 					})
@@ -816,7 +827,7 @@ void FSkinWeightDetailCustomization::AddSelectionUI(IDetailLayoutBuilder& Detail
 	.WholeRowContent()
 	[
 		SNew(SVerticalBox)
-		.IsEnabled_Lambda([this]{ return ToolSettings->WeightTool->HasSelectedVertices(); })
+		.IsEnabled_Lambda([this]{ return Tool->HasActiveSelectionOnMainMesh(); })
 
 		+ SVerticalBox::Slot()
 		.Padding(WeightEditHorizontalPadding, WeightEditVerticalPadding)
@@ -867,7 +878,7 @@ void FSkinWeightDetailCustomization::AddSelectionUI(IDetailLayoutBuilder& Detail
 	.WholeRowContent()
 	[
 		SNew(SVerticalBox)
-		.IsEnabled_Lambda([this]{ return ToolSettings->WeightTool->HasSelectedVertices(); })
+		.IsEnabled_Lambda([this]{ return Tool->HasActiveSelectionOnMainMesh(); })
 
 		+ SVerticalBox::Slot()
 		.Padding(WeightEditHorizontalPadding, WeightEditVerticalPadding)
@@ -892,7 +903,7 @@ void FSkinWeightDetailCustomization::AddSelectionUI(IDetailLayoutBuilder& Detail
 							ToolSettings->RelaxStrength,
 							DefaultRelaxIterations, 
 							EWeightEditOperation::Relax,
-							Tool->GetSelectedVertices(),
+							Tool->GetMainMeshSelector()->GetSelectedVertices(),
 							bShouldTransact);
 						return FReply::Handled();
 					})
@@ -925,7 +936,11 @@ void FSkinWeightDetailCustomization::AddSelectionUI(IDetailLayoutBuilder& Detail
 	EditWeightsCategory.AddCustomRow(LOCTEXT("VertexEditorRow", "Component Editor"), false)
 	.WholeRowContent()
 	[
-		SNew(SVertexWeightEditor, ToolSettings->WeightTool)
+		SNew(SBox)
+		.IsEnabled_Lambda([this]{ return Tool->HasActiveSelectionOnMainMesh(); })
+		[
+			SNew(SVertexWeightEditor, ToolSettings->WeightTool)
+		]
 	];
 	
 	// MIRROR WEIGHTS category
@@ -933,7 +948,7 @@ void FSkinWeightDetailCustomization::AddSelectionUI(IDetailLayoutBuilder& Detail
 	.WholeRowContent()
 	[
 		SNew(SVerticalBox)
-		.IsEnabled_Lambda([this]{ return ToolSettings->WeightTool->HasSelectedVertices(); })
+		.IsEnabled_Lambda([this]{ return Tool->HasActiveSelectionOnMainMesh(); })
 		
 		+ SVerticalBox::Slot()
 		.Padding(WeightEditHorizontalPadding, WeightEditVerticalPadding)
@@ -1022,7 +1037,7 @@ void FSkinWeightDetailCustomization::AddSelectionUI(IDetailLayoutBuilder& Detail
 	.WholeRowContent()
 	[
 		SNew(SHorizontalBox)
-		.IsEnabled_Lambda([this]{ return ToolSettings->WeightTool->HasSelectedVertices(); })
+		.IsEnabled_Lambda([this]{ return Tool->HasActiveSelectionOnMainMesh(); })
 
 		+SHorizontalBox::Slot()
 		.Padding(2.f, WeightEditVerticalPadding)
@@ -1072,7 +1087,7 @@ void FSkinWeightDetailCustomization::AddSelectionUI(IDetailLayoutBuilder& Detail
 	.WholeRowContent()
 	[
 		SNew(SHorizontalBox)
-		.IsEnabled_Lambda([this]{ return ToolSettings->WeightTool->HasSelectedVertices(); })
+		.IsEnabled_Lambda([this]{ return Tool->HasActiveSelectionOnMainMesh(); })
 
 		+SHorizontalBox::Slot()
 		.Padding(2.f, WeightEditVerticalPadding)
@@ -1118,7 +1133,7 @@ void FSkinWeightDetailCustomization::AddSelectionUI(IDetailLayoutBuilder& Detail
 	.WholeRowContent()
 	[
 		SNew(SHorizontalBox)
-		.IsEnabled_Lambda([this]{ return ToolSettings->WeightTool->HasSelectedVertices(); })
+		.IsEnabled_Lambda([this]{ return Tool->HasActiveSelectionOnMainMesh(); })
 
 		+SHorizontalBox::Slot()
 		.Padding(2.f, WeightEditVerticalPadding)
@@ -1167,35 +1182,83 @@ void FSkinWeightDetailCustomization::AddTransferUI(IDetailLayoutBuilder& DetailB
 	IDetailCategoryBuilder& TransferWeightsCategory = DetailBuilder.EditCategory("WeightTransfer", FText::GetEmpty(), ECategoryPriority::Important);
 	TransferWeightsCategory.InitiallyCollapsed(true);
 
+	// TRANSFER BUTTON
 	TransferWeightsCategory.AddCustomRow(LOCTEXT("TransferWeightsRow", "Transfer Weights"), false)
 	.WholeRowContent()
 	[
-		SNew(SVerticalBox)
-
-		+SVerticalBox::Slot()
+		SNew(SBox)
 		.Padding(0.f, WeightEditVerticalPadding)
 		[
-			SNew(SBox)
-			[
-				SNew(SButton)
-				.HAlign(HAlign_Center)
-				.Text(LOCTEXT("TransferWeightsButtonLabel", "Transfer Weights"))
-				.ToolTipText(LOCTEXT("TransferButtonTooltip",
-					"Weights are transferred from the source skeletal mesh using in-painting.\n"
-					"If no source skeletal mesh is specified, weights can be transferred from the existing mesh between profiles and LODs.\n"
-					"This command can operate on selected components when in Mesh mode."))
-				.OnClicked_Lambda([this]()
-				{
-					Tool->TransferWeights();
-					return FReply::Handled();
-				})
-				.IsEnabled_Lambda([this]()
-				{
-					return Tool->CanTransferWeights();
-				})
-			]
+			SNew(SButton)
+			.HAlign(HAlign_Center)
+			.Text(LOCTEXT("TransferWeightsButtonLabel", "Transfer Weights"))
+			.ToolTipText(LOCTEXT("TransferButtonTooltip",
+						"Weights are transferred from the source skeletal mesh.\n"
+						"Vertices may be selected on the source and/or target mesh to filter which parts to copy from and which parts to copy to.\n"
+						"If either mesh has no vertices selected, the whole mesh is considered.\n"))
+			.OnClicked_Lambda([this]()
+			{
+				Tool->GetWeightTransferManager()->TransferWeights();
+				return FReply::Handled();
+			})
+			.IsEnabled_Lambda([this]()
+			{
+				return Tool->GetWeightTransferManager()->CanTransferWeights();
+			})
 		]
 	];
+
+	// SKELETAL MESH ASSET INPUT
+	const TSharedRef<IPropertyHandle> SourceSkeletalMeshHandle = DetailBuilder.GetProperty(GET_MEMBER_NAME_CHECKED(USkinWeightsPaintToolProperties, SourceSkeletalMesh), USkinWeightsPaintToolProperties::StaticClass());
+	TransferWeightsCategory.AddProperty(SourceSkeletalMeshHandle);
+	
+	// LOD
+	const TSharedRef<IPropertyHandle> LODHandle = DetailBuilder.GetProperty(GET_MEMBER_NAME_CHECKED(USkinWeightsPaintToolProperties, SourceLOD), USkinWeightsPaintToolProperties::StaticClass());
+	TransferWeightsCategory.AddProperty(LODHandle);
+
+	// PROFILE
+	const TSharedRef<IPropertyHandle> SourceProfileHandle = DetailBuilder.GetProperty(GET_MEMBER_NAME_CHECKED(USkinWeightsPaintToolProperties, SourceSkinWeightProfile), USkinWeightsPaintToolProperties::StaticClass());
+	TransferWeightsCategory.AddProperty(SourceProfileHandle);
+
+	// MESH SELECTION OPTION (SOURCE OR TARGET)
+	const TSharedRef<IPropertyHandle> TransferSelectModeHandle = DetailBuilder.GetProperty(GET_MEMBER_NAME_CHECKED(USkinWeightsPaintToolProperties, MeshSelectMode), USkinWeightsPaintToolProperties::StaticClass());
+	FDetailWidgetRow& TransferSelectModeRow = TransferWeightsCategory.AddCustomRow(LOCTEXT("SelectionModeRow", "Selection Mode"), false)
+	.NameContent()
+	[
+		TransferSelectModeHandle->CreatePropertyNameWidget()
+	]
+	.ValueContent()
+	[
+		SNew(SBox)
+		.Padding(2.0f)
+		.HAlign(HAlign_Center)
+		[
+			SNew(SSegmentedControl<EMeshTransferOption>)
+			.ToolTipText(LOCTEXT("SelectionSourceTooltip",
+					"Choose which mesh to select components on (vertices/edges/faces).\n"
+					"Weights will be transferred from selected components on the source to selected components on the target.\n"
+					"If no components are selected on either the source or target, the whole mesh will be considered.\n"
+					"Source: The mesh to copy weights FROM.\n"
+					"Target: The mesh to copy weights TO (the main mesh in the tool)."))
+			.Value_Lambda([this]()
+			{
+				return ToolSettings->MeshSelectMode;
+			})
+			.OnValueChanged_Lambda([this](EMeshTransferOption Mode)
+			{
+				ToolSettings->MeshSelectMode = Mode;
+				ToolSettings->WeightTool->GetWeightTransferManager()->UpdateSelectionAndVisibility();
+			})
+			+ SSegmentedControl<EMeshTransferOption>::Slot(EMeshTransferOption::Source)
+			.Text(LOCTEXT("SourceMode", "Source"))
+			+SSegmentedControl<EMeshTransferOption>::Slot(EMeshTransferOption::Target)
+			.Text(LOCTEXT("TargetMode", "Target"))
+		]
+	];
+
+	// PREVIEW OFFSET
+	const TSharedRef<IPropertyHandle> PreviewOffsetHandle = DetailBuilder.GetProperty(GET_MEMBER_NAME_CHECKED(USkinWeightsPaintToolProperties, SourcePreviewOffset), USkinWeightsPaintToolProperties::StaticClass());
+	TransferWeightsCategory.AddProperty(PreviewOffsetHandle);
 }
 
 void SVertexWeightItem::Construct(const FArguments& InArgs, const TSharedRef<STableViewBase>& OwnerTableView)
@@ -1228,7 +1291,7 @@ TSharedRef<SWidget> SVertexWeightItem::GenerateWidgetForColumn(const FName& Colu
 				return ValueDuringSlide;
 			}
 			
-			return ParentTable->Tool->GetAverageWeightOnBone(Element->BoneIndex, ParentTable->Tool->GetSelectedVertices());
+			return ParentTable->Tool->GetAverageWeightOnBone(Element->BoneIndex, ParentTable->Tool->GetMainMeshSelector()->GetSelectedVertices());
 		})
 		.OnValueChanged_Lambda([this](float NewValue)
 		{
@@ -1246,7 +1309,7 @@ TSharedRef<SWidget> SVertexWeightItem::GenerateWidgetForColumn(const FName& Colu
 					RelativeScale,
 					0, // iterations
 					EWeightEditOperation::RelativeScale,
-					ParentTable->Tool->GetSelectedVertices(),
+					ParentTable->Tool->GetMainMeshSelector()->GetSelectedVertices(),
 					bShouldTransact);
 			}
 		})
@@ -1260,7 +1323,7 @@ TSharedRef<SWidget> SVertexWeightItem::GenerateWidgetForColumn(const FName& Colu
 					NewValue,
 					0, // iterations,
 					EWeightEditOperation::Replace,
-					ParentTable->Tool->GetSelectedVertices(),
+					ParentTable->Tool->GetMainMeshSelector()->GetSelectedVertices(),
 					bShouldTransact);
 			}
 			bInTransaction = false;
@@ -1268,7 +1331,7 @@ TSharedRef<SWidget> SVertexWeightItem::GenerateWidgetForColumn(const FName& Colu
 		.OnBeginSliderMovement_Lambda([this]()
 		{
 			ParentTable->Tool->BeginChange();
-			ValueAtStartOfSlide = ParentTable->Tool->GetAverageWeightOnBone(Element->BoneIndex, ParentTable->Tool->GetSelectedVertices());
+			ValueAtStartOfSlide = ParentTable->Tool->GetAverageWeightOnBone(Element->BoneIndex, ParentTable->Tool->GetMainMeshSelector()->GetSelectedVertices());
 			ValueDuringSlide = ValueAtStartOfSlide;
 			bInTransaction = true;
 		})
@@ -1335,7 +1398,7 @@ void SVertexWeightEditor::Construct(const FArguments& InArgs, USkinWeightsPaintT
 			SAssignNew( ListView, SWeightEditorListViewType )
 			.Visibility_Lambda([this]()
 			{
-				return Tool->GetSelectedVertices().IsEmpty() ? EVisibility::Collapsed : EVisibility::Visible;
+				return Tool->GetMainMeshSelector()->GetSelectedVertices().IsEmpty() ? EVisibility::Collapsed : EVisibility::Visible;
 			})
 			.SelectionMode(ESelectionMode::Single)
 			.ListItemsSource( &ListViewItems )
@@ -1375,7 +1438,7 @@ void SVertexWeightEditor::RefreshView()
 	
 	// get all bones affecting the selected vertices
 	TArray<int32> Influences;
-	Tool->GetInfluences(Tool->GetSelectedVertices(), Influences);
+	Tool->GetInfluences(Tool->GetMainMeshSelector()->GetSelectedVertices(), Influences);
 
 	// generate list view items
 	ListViewItems.Reset();
diff --git a/Engine/Plugins/Experimental/MeshModelingToolsetExp/Source/MeshModelingToolsEditorOnlyExp/Private/SkeletalMesh/SkinWeightsPaintTool.cpp b/Engine/Plugins/Experimental/MeshModelingToolsetExp/Source/MeshModelingToolsEditorOnlyExp/Private/SkeletalMesh/SkinWeightsPaintTool.cpp
--- a/Engine/Plugins/Experimental/MeshModelingToolsetExp/Source/MeshModelingToolsEditorOnlyExp/Private/SkeletalMesh/SkinWeightsPaintTool.cpp
+++ b/Engine/Plugins/Experimental/MeshModelingToolsetExp/Source/MeshModelingToolsEditorOnlyExp/Private/SkeletalMesh/SkinWeightsPaintTool.cpp
@@ -18,6 +18,7 @@
 
 #include "Editor/Persona/Public/IPersonaEditorModeManager.h"
 #include "Editor/Persona/Public/PersonaModule.h"
+#include "EditorViewportClient.h"
 #include "Preferences/PersonaOptions.h"
 #include "PreviewProfileController.h"
 #include "BaseGizmos/BrushStampIndicator.h"
@@ -308,14 +309,18 @@ TArray<FName> GetLODs(UToolTarget* InTarget)
 {
     static TArray<FName> Dummy;
 
-    if (!ensure(InTarget))
+    if (!InTarget)
     {
     	return Dummy;
     }
 
-    bool bSupportsLODs = false;
-    const TArray<EMeshLODIdentifier> LODIDSs = UE::ToolTarget::GetMeshDescriptionLODs(InTarget, bSupportsLODs);
-    if (!ensure(bSupportsLODs))
+    bool bOutSupportsLODs = false;
+	constexpr bool bOnlyReturnDefaultLOD = false;
+	// NOTE: currently auto-generated LODs do not have USkeletalMesh::SourceModels and so the returned MeshDescription will be null
+	// for this reason, we do not allow transferring weights to/from auto-generated LODs
+	constexpr bool bExcludeAutoGeneratedLODs = true;
+    const TArray<EMeshLODIdentifier> LODIDSs = UE::ToolTarget::GetMeshDescriptionLODs(InTarget, bOutSupportsLODs, bOnlyReturnDefaultLOD, bExcludeAutoGeneratedLODs);
+    if (!ensure(bOutSupportsLODs))
     {
     	return Dummy;
     }
@@ -335,60 +340,857 @@ TArray<FName> GetLODs(UToolTarget* InTarget)
     return LODs;
 }
 
-TArray<FName> GetSkinWeightProfilesFunc(const FMeshDescription& InMeshDescription)
+TArray<FName> GetProfilesFromToolTarget(UToolTarget* InTarget, const FName InLOD)
 {
-	const FSkeletalMeshConstAttributes MeshAttribs(InMeshDescription);
-	return MeshAttribs.GetSkinWeightProfileNames();
+	if (InTarget)
+	{
+		TArray<FName> Profiles;
+		const EMeshLODIdentifier LODId = GetLODId(InLOD);
+		const FGetMeshParameters Params(true, LODId);
+		if (const FMeshDescription* MeshDescription = UE::ToolTarget::GetMeshDescription(InTarget, Params))
+		{
+			const FSkeletalMeshConstAttributes MeshAttribs(*MeshDescription);
+			return MeshAttribs.GetSkinWeightProfileNames();
+		}
+	}
+	
+	static const TArray EmptyProfiles = {FSkeletalMeshAttributesShared::DefaultSkinWeightProfileName, CreateNewName()};	
+	return EmptyProfiles;
+};
+
+}
+
+TArray<FName> USkinWeightsPaintToolProperties::GetTargetLODsFunc() const
+{
+	return SkinWeightLayer::GetLODs(WeightTool->GetTarget());
+}
+
+TArray<FName> USkinWeightsPaintToolProperties::GetTargetSkinWeightProfilesFunc() const
+{
+	TArray<FName> Profiles = SkinWeightLayer::GetProfilesFromToolTarget(WeightTool->GetTarget(), SourceLOD);
+	Profiles.Add(CreateNewName());
+	return Profiles;
+}
+
+TArray<FName> USkinWeightsPaintToolProperties::GetSourceLODsFunc() const
+{
+	return SkinWeightLayer::GetLODs(WeightTool->GetWeightTransferManager()->GetTarget());
+}
+
+TArray<FName> USkinWeightsPaintToolProperties::GetSourceSkinWeightProfilesFunc() const
+{
+	UToolTarget* SourceToolTarget = WeightTool->GetWeightTransferManager()->GetTarget();
+	return SkinWeightLayer::GetProfilesFromToolTarget(SourceToolTarget, SourceLOD);
+}
+
+void UWeightToolMeshSelector::InitialSetup(
+	UWorld* InWorld,
+	UInteractiveTool* InParentTool,
+	FEditorViewportClient* InViewportClient,
+	TFunction<void()> OnSelectionChangedFunc)
+{
+	World = InWorld;
+
+	ViewportClient = InViewportClient;
+	
+	// set up vertex selection mechanic
+	PolygonSelectionMechanic = NewObject<UPolygonSelectionMechanic>(this);
+	PolygonSelectionMechanic->bAddSelectionFilterPropertiesToParentTool = false;
+	PolygonSelectionMechanic->Setup(InParentTool);
+	PolygonSelectionMechanic->SetIsEnabled(false);
+	PolygonSelectionMechanic->OnSelectionChanged.AddLambda(OnSelectionChangedFunc );
+
+	// set up style of vertex selection
+	constexpr FLinearColor VertexSelectedPurple = FLinearColor(0.78f, 0.f, 0.78f);
+	constexpr FLinearColor VertexSelectedYellow = FLinearColor(1.f,1.f,0.f);
+	// adjust selection rendering for this context
+	PolygonSelectionMechanic->HilightRenderer.PointColor = FLinearColor::Blue;
+	PolygonSelectionMechanic->HilightRenderer.PointSize = 10.0f;
+	// vertex highlighting once selected
+	PolygonSelectionMechanic->SelectionRenderer.LineThickness = 1.0f;
+	PolygonSelectionMechanic->SelectionRenderer.PointColor = VertexSelectedYellow;
+	PolygonSelectionMechanic->SelectionRenderer.PointSize = 5.0f;
+	PolygonSelectionMechanic->SelectionRenderer.DepthBias = 2.0f;
+	// despite the name, this renders the vertices
+	PolygonSelectionMechanic->PolyEdgesRenderer.PointColor = VertexSelectedPurple;
+	PolygonSelectionMechanic->PolyEdgesRenderer.PointSize = 5.0f;
+	PolygonSelectionMechanic->PolyEdgesRenderer.DepthBias = 2.0f;
+	PolygonSelectionMechanic->PolyEdgesRenderer.LineThickness = 1.0f;
+}
+
+void UWeightToolMeshSelector::SetMesh(
+	UPreviewMesh* InPreviewMesh,
+	const FMeshDescription* InMeshDescription,
+	const FTransform3d& InMeshTransform)
+{
+	// store the mesh this is operating on
+	PreviewMesh = InPreviewMesh;
+	// copy the mesh description for selection validation
+	MeshDescription = *InMeshDescription;
+
+	if (!ensure(World))
+	{
+		return;
+	}
+
+	if (!PreviewMesh)
+	{
+		SetIsEnabled(false);
+		return;
+	}
+
+	// reset selection topology and mesh spatial data
+	static constexpr bool bAutoBuild = true;
+	const FDynamicMesh3* DynamicMesh = PreviewMesh->GetMesh();
+	SelectionTopology = MakeUnique<UE::Geometry::FTriangleGroupTopology>(DynamicMesh, bAutoBuild);
+	MeshSpatial = MakeUnique<FDynamicMeshAABBTree3>(DynamicMesh, bAutoBuild);
+
+	// initialize the selection mechanic
+	PolygonSelectionMechanic->Initialize(
+		DynamicMesh,
+		InMeshTransform,
+		World,
+		SelectionTopology.Get(),
+		[this]() { return MeshSpatial.Get(); }
+	);
+
+	// clear the selection (old selection is invalid on new topo)
+	PolygonSelectionMechanic->ClearSelection();
+
+	// selection colors
+	constexpr FLinearColor FaceSelectedOrange = FLinearColor(0.886f, 0.672f, 0.473f);
+	// configure secondary render material for selected triangles
+	// NOTE: the material returned by ToolSetupUtil::GetSelectionMaterial has a checkerboard pattern on back faces which makes it hard to use
+	if (UMaterialInterface* Material = LoadObject<UMaterial>(nullptr, TEXT("/MeshModelingToolsetExp/Materials/SculptMaterial")))
+	{
+		if (UMaterialInstanceDynamic* MatInstance = UMaterialInstanceDynamic::Create(Material, World))
+		{
+			MatInstance->SetVectorParameterValue(TEXT("Color"), FaceSelectedOrange);
+			PreviewMesh->SetSecondaryRenderMaterial(MatInstance);
+		}
+	}
+
+	// secondary triangle buffer used to render face selection
+	PreviewMesh->EnableSecondaryTriangleBuffers([this](const FDynamicMesh3* Mesh, int32 TriangleID)
+	{
+		return PolygonSelectionMechanic->GetActiveSelection().IsSelectedTriangle(Mesh, SelectionTopology.Get(), TriangleID);
+	});
+	// notify preview mesh when triangle selection has been updated
+	PolygonSelectionMechanic->OnSelectionChanged.AddWeakLambda(this, [this]()
+	{
+		PreviewMesh->FastNotifySecondaryTrianglesChanged();
+	});
+	PolygonSelectionMechanic->OnFaceSelectionPreviewChanged.AddWeakLambda(this, [this]()
+	{
+		PreviewMesh->FastNotifySecondaryTrianglesChanged();
+	});
 }
 
+void UWeightToolMeshSelector::Shutdown()
+{
+	if (PolygonSelectionMechanic)
+	{
+		PolygonSelectionMechanic->Shutdown();
+	}
+	
+	PolygonSelectionMechanic = nullptr;
 }
 
-TArray<FName> USkinWeightsPaintToolProperties::GetLODsFunc() const
+void UWeightToolMeshSelector::SetIsEnabled(bool bIsEnabled)
 {
-	static TArray<FName> Dummy;
-	if (!ensure(WeightTool && WeightTool->GetTarget()))
+	if (!PolygonSelectionMechanic)
 	{
-		return Dummy;
+		return;
 	}
+
+	// force off if there's no preview mesh
+	bIsEnabled = PreviewMesh ? bIsEnabled : false; 
 	
-	return SkinWeightLayer::GetLODs(WeightTool->GetTarget());
+	PolygonSelectionMechanic->SetIsEnabled(bIsEnabled);
 }
 
-TArray<FName> USkinWeightsPaintToolProperties::GetSkinWeightProfilesFunc() const
+void UWeightToolMeshSelector::SetComponentSelectionMode(EComponentSelectionMode InMode)
 {
-	if (const USkeletalMeshComponent* SkeletalMeshComponent = GetSkeletalMeshComponent(WeightTool->GetTarget()))
+	if (!(PolygonSelectionMechanic))
 	{
-		const EMeshLODIdentifier LODId = GetLODId(ActiveLOD);
-		const FGetMeshParameters Params(true, LODId);
-		const FMeshDescription* MeshDescription = UE::ToolTarget::GetMeshDescription(WeightTool->GetTarget(), Params);
-		TArray Profiles = SkinWeightLayer::GetSkinWeightProfilesFunc(*MeshDescription);
-		Profiles.Add(CreateNewName());
-		return Profiles;
+		return;
 	}
+	
+	PolygonSelectionMechanic->Properties->bSelectVertices = InMode == EComponentSelectionMode::Vertices;
+	PolygonSelectionMechanic->Properties->bSelectEdges = InMode == EComponentSelectionMode::Edges;
+	PolygonSelectionMechanic->Properties->bSelectFaces = InMode == EComponentSelectionMode::Faces;
+	PolygonSelectionMechanic->SetShowSelectableCorners(InMode == EComponentSelectionMode::Vertices);
+	PolygonSelectionMechanic->SetShowEdges(InMode == EComponentSelectionMode::Edges);
+}
 
-	static const TArray Profiles = {FSkeletalMeshAttributesShared::DefaultSkinWeightProfileName, CreateNewName()};
-	return Profiles;
+void UWeightToolMeshSelector::SetTransform(const FTransform3d& InTargetTransform)
+{
+	if (PolygonSelectionMechanic)
+	{
+		PolygonSelectionMechanic->SetTransform(InTargetTransform);
+	}
+}
+
+void UWeightToolMeshSelector::UpdateAfterMeshDeformation()
+{
+	MeshSpatial->Build();
+	constexpr bool bTopologyDeformed = true;
+	constexpr bool bTopologyModified = false;
+	PolygonSelectionMechanic->GetTopologySelector()->Invalidate(bTopologyDeformed, bTopologyModified);
+}
+
+void UWeightToolMeshSelector::DrawHUD(FCanvas* Canvas, IToolsContextRenderAPI* RenderAPI)
+{
+	if (!PolygonSelectionMechanic)
+	{
+		return;
+	}
+
+	PolygonSelectionMechanic->DrawHUD(Canvas, RenderAPI);
+}
+	
+void UWeightToolMeshSelector::Render(IToolsContextRenderAPI* RenderAPI)
+{
+	if (!(PolygonSelectionMechanic && PreviewMesh))
+	{
+		return;
+	}
+	
+	PolygonSelectionMechanic->Render(RenderAPI);
+}
+
+const TArray<int32>& UWeightToolMeshSelector::GetSelectedVertices()
+{
+	using UE::Geometry::FNonManifoldMappingSupport;
+	
+	SelectedVerticesInternal.Empty();
+	if (!(PolygonSelectionMechanic && PreviewMesh))
+	{
+		return SelectedVerticesInternal;
+	}
+	
+	const FGroupTopologySelection& Selection = PolygonSelectionMechanic->GetActiveSelection();
+	const FDynamicMesh3* DynamicMesh = PreviewMesh->GetMesh();
+	const FVertexArray& AllVertices = MeshDescription.Vertices();
+	const FNonManifoldMappingSupport NonManifoldMappingSupport(*DynamicMesh);
+
+	// validate and add vertices to the output array
+	auto AddVertices = [this, &AllVertices, &NonManifoldMappingSupport](const TSet<int32>& VerticesToAdd)
+	{
+		// we have to make sure that the vertex ids are safe to use as PolygonSelectionMechanic does not act on the
+		// mesh description but on the dynamic mesh that can duplicate vertices when dealing with degenerate triangles.
+		// cf. FMeshDescriptionToDynamicMesh::Convert for more details.
+		Algo::CopyIf(VerticesToAdd, SelectedVerticesInternal, [&](int32 VertexID)
+		{
+			const int32 SrcVertexID = NonManifoldMappingSupport.GetOriginalNonManifoldVertexID(VertexID);
+			return AllVertices.IsValid(SrcVertexID);
+		});
+	};
+
+	// add selected vertices
+	AddVertices(Selection.SelectedCornerIDs);
+
+	// add vertices on selected edges
+	{
+		TSet<int32> VerticesInSelectedEdges;
+		for (const int32 SelectedEdgeIndex : Selection.SelectedEdgeIDs)
+		{
+			FDynamicMesh3::FEdge CurrentEdge = DynamicMesh->GetEdge(SelectedEdgeIndex);
+			VerticesInSelectedEdges.Add(CurrentEdge.Vert.A);
+			VerticesInSelectedEdges.Add(CurrentEdge.Vert.B);
+		}
+		
+		AddVertices(VerticesInSelectedEdges);
+	}
+
+	// add vertices in selected faces
+	{
+		TSet<int32> VerticesInSelectedFaces;
+		for (const int32 SelectedFaceIndex : Selection.SelectedGroupIDs)
+		{
+			UE::Geometry::FIndex3i TriangleVertices = DynamicMesh->GetTriangleRef(SelectedFaceIndex);
+			VerticesInSelectedFaces.Add(TriangleVertices[0]);
+			VerticesInSelectedFaces.Add(TriangleVertices[1]);
+			VerticesInSelectedFaces.Add(TriangleVertices[2]);
+		}
+		
+		AddVertices(VerticesInSelectedFaces);
+	}
+	
+	return SelectedVerticesInternal;
+}
+
+bool UWeightToolMeshSelector::IsAnyComponentSelected() const
+{
+	if (!ensure(PolygonSelectionMechanic))
+	{
+		return false;
+	}
+
+	return PolygonSelectionMechanic->HasSelection();
+}
+
+void UWeightToolMeshSelector::GetSelectedTriangles(TArray<int32>& OutTriangleIndices) const
+{
+	OutTriangleIndices.Empty();
+	if (!ensure(PolygonSelectionMechanic))
+	{
+		return;
+	}
+	
+	const FGroupTopologySelection& Selection = PolygonSelectionMechanic->GetActiveSelection();
+	const FDynamicMesh3* DynamicMesh = PreviewMesh->GetMesh();
+	TSet<int32> TriangleSet;
+
+	// add triangles connected to selected vertices
+	for (const int32 VertexIndex : Selection.SelectedCornerIDs)
+	{
+		DynamicMesh->EnumerateVertexTriangles(VertexIndex, [&TriangleSet](int32 TriangleIndex)
+		{
+			TriangleSet.Add(TriangleIndex);
+		});
+	}
+
+	// add triangles connected to selected edges
+	for (const int32 EdgeIndex : Selection.SelectedEdgeIDs)
+	{
+		DynamicMesh->EnumerateEdgeTriangles(EdgeIndex, [&TriangleSet](int32 TriangleIndex)
+		{
+			TriangleSet.Add(TriangleIndex);
+		});
+	}
+
+	// add selected triangles
+	TriangleSet.Append(Selection.SelectedGroupIDs);	
+
+	OutTriangleIndices = TriangleSet.Array();
+}
+
+void UWeightToolMeshSelector::GrowSelection() const
+{
+	if (!ensure(PolygonSelectionMechanic))
+	{
+		return;
+	}
+	
+	PolygonSelectionMechanic->GrowSelection();
+}
+
+void UWeightToolMeshSelector::ShrinkSelection() const
+{
+	if (!ensure(PolygonSelectionMechanic))
+	{
+		return;
+	}
+	
+	PolygonSelectionMechanic->ShrinkSelection();
+}
+
+void UWeightToolMeshSelector::FloodSelection() const
+{
+	if (!ensure(PolygonSelectionMechanic))
+	{
+		return;
+	}
+	
+	PolygonSelectionMechanic->FloodSelection();
+}
+
+void UWeightToolMeshSelector::SelectBorder() const
+{
+	if (!ensure(PolygonSelectionMechanic))
+	{
+		return;
+	}
+	
+	PolygonSelectionMechanic->ConvertSelectionToBorder();
+}
+
+void UWeightToolTransferManager::InitialSetup(USkinWeightsPaintTool* InWeightTool, FEditorViewportClient* InViewportClient)
+{
+	WeightTool = InWeightTool;
+
+	// always reset back to target selection
+	GetToolProperties()->MeshSelectMode = EMeshTransferOption::Target;
+
+	// create the mesh selector and run initial setup
+	// NOTE: currently this must happen inside the Setup of a UInteractiveTool so that input is routed to the selection mechanic
+	MeshSelector = NewObject<UWeightToolMeshSelector>(this);
+	auto DoNothingOnSelection = [](){};
+	MeshSelector->InitialSetup(WeightTool->GetTargetWorld(), WeightTool.Get(), InViewportClient,  DoNothingOnSelection);
+}
+
+void UWeightToolTransferManager::SetSourceMesh(USkeletalMesh* InSkeletalMesh)
+{
+	SourceSkeletalMesh = InSkeletalMesh;
+
+	USkinWeightsPaintToolProperties* ToolProperties = GetToolProperties();
+	
+	// reset to prepare for new mesh (or possibly no mesh)
+	{
+		if (SourcePreviewMesh)
+		{
+			SourcePreviewMesh->SetVisible(false);
+			SourcePreviewMesh->Disconnect();
+			SourcePreviewMesh = nullptr;
+		}
+
+		if (SourceTarget)
+		{
+			SourceTarget = nullptr;
+		}
+
+		if (MeshSelector)
+		{
+			MeshSelector->SetIsEnabled(false);
+		}
+
+		ToolProperties->SourceSkinWeightProfile = FSkeletalMeshAttributesShared::DefaultSkinWeightProfileName;
+		ToolProperties->SourceLOD = "LOD0";
+	}
+	
+	if (!InSkeletalMesh)
+	{
+		UpdateSelectionAndVisibility();
+		return;
+	}
+	
+	// create the preview mesh (this creates the skeletal mesh component in the world)
+	SourcePreviewMesh = NewObject<UPreviewMesh>(this);
+	SourcePreviewMesh->CreateInWorld(WeightTool->GetTargetWorld(), FTransform::Identity);
+
+	// create a new tool target for this mesh
+	SourceTarget = WeightTool->GetTargetManager()->BuildTarget(InSkeletalMesh, FToolTargetTypeRequirements());
+	
+	// move source mesh beside the main mesh (to the left in screen space)
+	{
+		UToolTarget* MainTarget = WeightTool->GetTarget();
+		if (!MainTarget)
+		{
+			return;
+		}
+
+		const USkeletalMeshComponent* SkeletalMeshComponent = GetSkeletalMeshComponent(MainTarget);
+		if (!SkeletalMeshComponent)
+		{
+			return;
+		}
+		
+		const FBoxSphereBounds TargetBounds = SkeletalMeshComponent->GetSkeletalMeshAsset()->GetBounds();
+		const FBoxSphereBounds SourceBounds = InSkeletalMesh->GetBounds();
+		FTransform Transform = UE::ToolTarget::GetLocalToWorldTransform(MainTarget);
+		FVector Location = Transform.GetLocation();
+		Location.X -= TargetBounds.GetBoxExtrema(1).X;
+		Location.X -= 1.1 * SourceBounds.GetBoxExtrema(1).X;
+		Transform.SetLocation(Location);
+		GetToolProperties()->SourcePreviewOffset = Transform;
+		SourcePreviewMesh->SetTransform(Transform);
+	}
+
+	// replace the dynamic mesh contained in the preview mesh
+	ToolSetupUtil::ApplyRenderingConfigurationToPreview(SourcePreviewMesh, SourceTarget);
+	SourcePreviewMesh->SetTangentsMode(EDynamicMeshComponentTangentsMode::AutoCalculated);
+	const EMeshLODIdentifier SourceLODId = GetLODId(ToolProperties->SourceLOD);
+	const FGetMeshParameters SourceParams(true, SourceLODId);
+	SourcePreviewMesh->ReplaceMesh(UE::ToolTarget::GetDynamicMeshCopy(SourceTarget, SourceParams));
+
+	// setup materials and visibility of the preview mesh
+	const FComponentMaterialSet MaterialSet = UE::ToolTarget::GetMaterialSet(SourceTarget);
+	SourcePreviewMesh->SetMaterials(MaterialSet.Materials);
+
+	// setup the mesh selection for the source
+	const FMeshDescription* MeshDescription =  UE::ToolTarget::GetMeshDescription(WeightTool->GetTarget(), SourceParams);
+	MeshSelector->SetMesh(SourcePreviewMesh, MeshDescription, ToolProperties->SourcePreviewOffset);
+	MeshSelector->SetIsEnabled(ToolProperties->EditingMode == EWeightEditMode::Mesh);
+	MeshSelector->SetComponentSelectionMode(ToolProperties->ComponentSelectionMode);
+
+	UpdateSelectionAndVisibility();
+}
+
+void UWeightToolTransferManager::UpdateSelectionAndVisibility() const
+{
+	const USkinWeightsPaintToolProperties* ToolProperties = GetToolProperties();
+	const bool bHasSourceMesh = GetPreviewMesh() != nullptr;
+	const bool bWasSetToSelectSource = ToolProperties->MeshSelectMode == EMeshTransferOption::Source;
+	const bool bEnableSourceMeshSelector = bWasSetToSelectSource && bHasSourceMesh;
+	
+	if (MeshSelector)
+	{
+		MeshSelector->SetIsEnabled(bEnableSourceMeshSelector);
+		MeshSelector->SetComponentSelectionMode(ToolProperties->ComponentSelectionMode);
+	}
+
+	WeightTool->GetMainMeshSelector()->SetIsEnabled(!bEnableSourceMeshSelector);
+	WeightTool->GetMainMeshSelector()->SetComponentSelectionMode(ToolProperties->ComponentSelectionMode);
+}
+
+void UWeightToolTransferManager::Shutdown()
+{
+	SetSourceMesh(nullptr);
+	
+	if (MeshSelector)
+	{
+		MeshSelector->Shutdown();
+	}
+
+	MeshSelector = nullptr;
+	SourcePreviewMesh = nullptr;
+	SourceTarget = nullptr;
+}
+
+void UWeightToolTransferManager::Render(IToolsContextRenderAPI* RenderAPI)
+{
+	if (MeshSelector)
+	{
+		MeshSelector->Render(RenderAPI);	
+	}
+}
+
+void UWeightToolTransferManager::DrawHUD(FCanvas* Canvas, IToolsContextRenderAPI* RenderAPI)
+{
+	if (MeshSelector)
+	{
+		MeshSelector->DrawHUD(Canvas, RenderAPI);
+	}
+}
+
+void UWeightToolTransferManager::TransferWeights()
+{
+	if (!ensure(SourceTarget))
+	{
+		// to transfer weights from another mesh we need a source mesh
+		// (UI should prevent us from getting here)
+		const FText NotificationText = LOCTEXT("NoSourceMesh", "No source skeletal mesh specified. No weights were transferred.");
+		ShowEditorMessage(ELogVerbosity::Error, NotificationText);
+		return;
+	}
+
+	USkinWeightsPaintToolProperties* ToolProperties = GetToolProperties();
+
+	const USkeletalMesh* TargetSkeletalMesh = GetSkeletalMeshComponent(WeightTool->GetTarget())->GetSkeletalMeshAsset();
+	const bool bSameMesh = SourceSkeletalMesh == TargetSkeletalMesh;
+	const bool bSameLOD = GetLODId(ToolProperties->ActiveLOD) == GetLODId(ToolProperties->SourceLOD);
+	const bool bSameProfile = ToolProperties->ActiveSkinWeightProfile == ToolProperties->SourceSkinWeightProfile;
+
+	// cannot transfer between same mesh/LOD/profile (identical weights)
+	if (bSameMesh && bSameLOD && bSameProfile)
+	{
+		const FText NotificationText = LOCTEXT("NoSourceMesh", "Cannot copy weights from the same mesh, LOD and profile. No weights were transferred.");
+		ShowEditorMessage(ELogVerbosity::Error, NotificationText);
+		return;
+	}
+	
+	if (bSameMesh && bSameLOD)
+	{
+		TransferWeightsFromSameMeshAndLOD();
+	}
+	else
+	{
+		TransferWeightsFromOtherMesh();
+	}
+}
+
+void UWeightToolTransferManager::TransferWeightsFromOtherMesh()
+{
+	using UE::Geometry::FDynamicMeshAttributeSet;
+	using UE::Geometry::FDynamicMeshVertexSkinWeightsAttribute;
+	using UE::Geometry::FTransferBoneWeights;
+	using UE::AnimationCore::FBoneWeights;
+
+	if (!ensure(SourceTarget))
+	{
+		return;
+	}
+
+	USkinWeightsPaintToolProperties* ToolProperties = GetToolProperties();
+
+	// get LOD IDs
+	const EMeshLODIdentifier TargetLODId = GetLODId(ToolProperties->ActiveLOD);
+	const EMeshLODIdentifier SourceLODId = GetLODId(ToolProperties->SourceLOD);
+
+	// if transferring between the same mesh, ensure that the LODs are different
+	const USkeletalMesh* TargetSkeletalMesh = GetSkeletalMeshComponent(WeightTool->GetTarget())->GetSkeletalMeshAsset();
+	if (SourceSkeletalMesh == TargetSkeletalMesh)
+	{
+		if (TargetLODId == SourceLODId)
+		{
+			const FText NotificationText = LOCTEXT("SameLOD", "Cannot copy weights between the same LOD on the same mesh. No weights were transferred.");
+			ShowEditorMessage(ELogVerbosity::Error, NotificationText);
+			return;
+		}
+	}
+
+	// get target dynamic mesh
+	const FGetMeshParameters TargetParams(true, TargetLODId);
+	FDynamicMesh3 TargetMesh = UE::ToolTarget::GetDynamicMeshCopy(WeightTool->GetTarget(), TargetParams);
+	
+	// get the source dynamic mesh and validate it
+	const FGetMeshParameters SourceParams(true, SourceLODId);
+	FDynamicMesh3 SourceMeshOrig = UE::ToolTarget::GetDynamicMeshCopy(SourceTarget, SourceParams);
+	FDynamicMesh3 SourceMesh;
+
+	// create a submesh from the selected triangles to filter the transfer so that it only copies from the selected components
+	TArray<int32> TrianglesToIsolate;
+	MeshSelector->GetSelectedTriangles(TrianglesToIsolate);
+	if (!TrianglesToIsolate.IsEmpty())
+	{
+		UE::Geometry::FDynamicSubmesh3 PartialSubMesh = UE::Geometry::FDynamicSubmesh3(&SourceMeshOrig, TrianglesToIsolate);
+		SourceMesh = PartialSubMesh.GetSubmesh();
+		FDynamicMeshAttributeSet* SourceAttributes = SourceMesh.Attributes();
+		SourceAttributes->CopyBoneAttributes(*SourceMeshOrig.Attributes());
+	}
+	else
+	{
+		SourceMesh = MoveTemp(SourceMeshOrig);
+	}
+	
+	if (!SourceMesh.HasAttributes() || !SourceMesh.Attributes()->HasBones())
+	{
+		const FText NotificationText = LOCTEXT("NoWeightsFoundInTransfer", "No skin weights were found in the source skeletal mesh. No weights were transferred.");
+		ShowEditorMessage(ELogVerbosity::Error, NotificationText);
+		return;
+	}
+	if (SourceMesh.Attributes()->GetNumBones() == 0)
+	{
+		const FText NotificationText = LOCTEXT("NoBonesFoundInTransfer", "No bones were found in the source skeletal mesh. No weights were transferred.");
+		ShowEditorMessage(ELogVerbosity::Error, NotificationText);
+		return;
+	}
+
+	FTransferBoneWeights TransferBoneWeights(&SourceMesh, ToolProperties->SourceSkinWeightProfile);
+	TransferBoneWeights.TransferMethod = FTransferBoneWeights::ETransferBoneWeightsMethod::InpaintWeights;
+	
+	if (!TargetMesh.HasAttributes())
+	{
+		TargetMesh.EnableAttributes();
+	}
+
+	FDynamicMeshAttributeSet* TargetAttributes = TargetMesh.Attributes();
+	if (!TargetAttributes->HasBones())
+	{
+		TargetAttributes->CopyBoneAttributes(*SourceMesh.Attributes());
+	}
+	else
+	{
+		const FReferenceSkeleton& TargetRefSkeleton = TargetSkeletalMesh->GetRefSkeleton();
+		if (!ensure(TargetAttributes->GetNumBones() == TargetRefSkeleton.GetRawBoneNum()))
+		{
+			const FText NotificationText = LOCTEXT("TransferWeightsInvalidBoneCount", "Target mesh bone attributes mismatch. No weights were transferred.");
+			ShowEditorMessage(ELogVerbosity::Error, NotificationText);
+			return;
+		}
+	}
+
+	// NOTE should we expose all the options?
+	// 	TransferBoneWeights.NormalThreshold;
+	// 	TransferBoneWeights.SearchRadius
+	// 	TransferBoneWeights.NumSmoothingIterations;
+	// 	TransferBoneWeights.SmoothingStrength;
+	// 	TransferBoneWeights.LayeredMeshSupport;
+	// 	TransferBoneWeights.ForceInpaintWeightMapName;
+	
+	if (TransferBoneWeights.Validate() != UE::Geometry::EOperationValidationResult::Ok)
+	{
+		const FText NotificationText = LOCTEXT("TransferWeightsNotValid", "Transfer weights operation unable to validate meshes. No weights were transferred.");
+		ShowEditorMessage(ELogVerbosity::Error, NotificationText);
+		return;
+	}
+
+	if (ToolProperties->EditingMode == EWeightEditMode::Mesh)
+	{
+		TransferBoneWeights.TargetVerticesSubset = WeightTool->GetMainMeshSelector()->GetSelectedVertices();
+	}
+
+	const FName TargetProfile = ToolProperties->GetActiveSkinWeightProfile();
+	if (!TransferBoneWeights.TransferWeightsToMesh(TargetMesh, TargetProfile))
+	{
+		const FText NotificationText = LOCTEXT("TransferWeightsUnknownIssue", "Transfer weights operation encountered an unknown issue. No weights were transferred.");
+		ShowEditorMessage(ELogVerbosity::Error, NotificationText);
+		return;
+	}
+
+	FDynamicMeshVertexSkinWeightsAttribute* TransferredSkinWeights = nullptr;
+	TransferredSkinWeights = TargetAttributes->GetSkinWeightsAttribute(TargetProfile);
+	
+	TArray<int32> VertexSubset = MoveTemp(TransferBoneWeights.TargetVerticesSubset);
+
+	// apply the weight changes as a transaction
+	ApplyTranferredWeightsAsTransaction(TransferredSkinWeights, VertexSubset, TargetMesh);
+}
+
+void UWeightToolTransferManager::TransferWeightsFromSameMeshAndLOD()
+{
+	if (!ensure(SourceTarget))
+	{
+		return;
+	}
+	
+	using UE::Geometry::FDynamicMeshAttributeSet;
+	using UE::Geometry::FDynamicMeshVertexSkinWeightsAttribute;
+	using UE::AnimationCore::FBoneWeights;
+	
+	const USkinWeightsPaintToolProperties* ToolProperties = GetToolProperties();
+	
+	// get target dynamic mesh
+	const EMeshLODIdentifier TargetLODId = GetLODId(ToolProperties->ActiveLOD);
+	const FName TargetProfile = ToolProperties->GetActiveSkinWeightProfile();
+	const FGetMeshParameters TargetParams(true, TargetLODId);
+	FDynamicMesh3 TargetMesh = UE::ToolTarget::GetDynamicMeshCopy(WeightTool->GetTarget(), TargetParams);
+	
+	const FDynamicMeshAttributeSet* TargetAttributes = TargetMesh.Attributes();
+	FDynamicMeshVertexSkinWeightsAttribute* TransferredSkinWeights = TargetAttributes->GetSkinWeightsAttribute(TargetProfile);
+	if (!TransferredSkinWeights)
+	{
+		TransferredSkinWeights = new FDynamicMeshVertexSkinWeightsAttribute(&TargetMesh);
+		TargetMesh.Attributes()->AttachSkinWeightsAttribute(TargetProfile, TransferredSkinWeights);
+	}
+		
+	const FDynamicMeshVertexSkinWeightsAttribute* SourceAttributes = TargetAttributes->GetSkinWeightsAttribute(ToolProperties->SourceSkinWeightProfile);
+	check(SourceAttributes != nullptr);
+
+	if (!ensure(SourceAttributes == TransferredSkinWeights))
+	{
+		const FText NotificationText = LOCTEXT("TransferBetweenSame", "Mesh attribute mismatch. No weights were transferred.");
+		ShowEditorMessage(ELogVerbosity::Error, NotificationText);
+		return;
+	}
+
+	TArray<int32> VertexSubset;
+	if (ToolProperties->EditingMode == EWeightEditMode::Mesh)
+	{
+		VertexSubset = MeshSelector->GetSelectedVertices();
+	}
+
+	if (VertexSubset.IsEmpty())
+	{
+		// copy weights for ALL vertices in the mesh
+		TransferredSkinWeights->Copy(*SourceAttributes);
+	}
+	else
+	{
+		// copy vertex weights for a subset of vertices
+		for (const int32 VertexID: VertexSubset)
+		{
+			FBoneWeights BoneWeights;
+			SourceAttributes->GetValue(VertexID, BoneWeights);
+			TransferredSkinWeights->SetValue(VertexID, BoneWeights);
+		}
+	}
+
+	// apply the weight changes as a transaction
+	ApplyTranferredWeightsAsTransaction(TransferredSkinWeights, VertexSubset, TargetMesh);
+}
+
+bool UWeightToolTransferManager::CanTransferWeights() const
+{
+	return SourceTarget && WeightTool->GetTarget();
+}
+
+void UWeightToolTransferManager::OnPropertyModified(const USkinWeightsPaintToolProperties* WeightToolProperties, const FProperty* ModifiedProperty)
+{
+	if (ModifiedProperty->GetName() == GET_MEMBER_NAME_STRING_CHECKED(USkinWeightsPaintToolProperties, SourceSkeletalMesh))
+	{
+		SetSourceMesh(WeightToolProperties->SourceSkeletalMesh.Get());
+	}
+
+	if (ModifiedProperty->GetName() == GET_MEMBER_NAME_STRING_CHECKED(USkinWeightsPaintToolProperties, SourceLOD))
+	{
+		if (SourcePreviewMesh)
+		{
+			// reapply mesh (will use the new LOD)
+			SetSourceMesh(SourceSkeletalMesh);
+		}
+	}
+
+	if (ModifiedProperty->GetName() == GET_MEMBER_NAME_STRING_CHECKED(USkinWeightsPaintToolProperties, SourcePreviewOffset))
+	{
+		if (SourcePreviewMesh)
+		{
+			SourcePreviewMesh->SetTransform(WeightToolProperties->SourcePreviewOffset);
+			MeshSelector->SetTransform(WeightToolProperties->SourcePreviewOffset);
+		}
+	}
+}
+
+void UWeightToolTransferManager::ApplyTranferredWeightsAsTransaction(
+	const UE::Geometry::FDynamicMeshVertexSkinWeightsAttribute* TransferredSkinWeights,
+	const TArray<int32>& VertexSubset,
+	FDynamicMesh3& TargetMesh)
+{
+	check(TransferredSkinWeights);
+	
+	// weight edits for transaction
+	FMultiBoneWeightEdits WeightEdits;
+
+	// get the weight data (used for making edits)
+	FSkinToolWeights& Weights = WeightTool->GetWeights();
+	
+	// spin through all the transferred skin weights and record a weight edit to apply as a transaction
+	static constexpr float ZeroWeight = 0.f;
+	const bool bUseSubset = !VertexSubset.IsEmpty();
+	const int32 NumVertices = bUseSubset ? VertexSubset.Num() : TargetMesh.VertexCount();
+	const UE::Geometry::FNonManifoldMappingSupport NonManifoldMappingSupport(TargetMesh);
+	for (int32 VertexIndex = 0; VertexIndex < NumVertices; ++VertexIndex)
+	{
+		const int32 VertexID = bUseSubset ? VertexSubset[VertexIndex] : VertexIndex;
+		const int32 SrcVertexID = NonManifoldMappingSupport.GetOriginalNonManifoldVertexID(VertexID);
+		
+		// remove all weight on vertex
+		const VertexWeights& VertexBoneWeights = Weights.PreChangeWeights[SrcVertexID];
+		if (!VertexBoneWeights.IsEmpty())
+		{
+			for (const FVertexBoneWeight& BoneWeight : VertexBoneWeights)
+			{
+				// when transferring weights we do prune influences because we would prefer the results to be identical to the source
+				constexpr bool bPruneInfluence = true;
+				WeightEdits.MergeSingleEdit(BoneWeight.BoneID, SrcVertexID, ZeroWeight, bPruneInfluence, Weights.PreChangeWeights);
+			}
+		}
+		else
+		{
+			// in the unlikely event that the target vertex has no weight, "fake" remove it from the root so that undo will put it back there
+			constexpr bool bPruneInfluence = false;
+			constexpr BoneIndex RootBoneIndex = 0;
+			WeightEdits.MergeSingleEdit(RootBoneIndex, SrcVertexID, ZeroWeight, bPruneInfluence, Weights.PreChangeWeights);
+		}
+
+		// update with new weight
+		UE::AnimationCore::FBoneWeights TransferedBoneWeights;
+		TransferredSkinWeights->GetValue(VertexID, TransferedBoneWeights);
+		for (const UE::AnimationCore::FBoneWeight& BoneWeight: TransferedBoneWeights)
+		{
+			const int32 BoneIndex = BoneWeight.GetBoneIndex();					
+			const float NewWeight = BoneWeight.GetWeight();
+			constexpr bool bPruneInfluence = false;
+			WeightEdits.MergeSingleEdit(BoneIndex, SrcVertexID, NewWeight, bPruneInfluence, Weights.PreChangeWeights);
+		}
+	}
+
+	// apply the changes as a transaction
+	const FText TransactionLabel = LOCTEXT("TransferWeightsChange", "Transfer skin weights.");
+	WeightTool->ApplyWeightEditsAsTransaction(WeightEdits, TransactionLabel);
+	
+	// put the mesh back in it's current pose
+	Weights.Deformer.SetAllVerticesToBeUpdated();
+
+	// notify user that weights were transferred.
+	const FText NotificationText = LOCTEXT("WeightsTransferred", "Skin weights transferred.");
+	ShowEditorMessage(ELogVerbosity::Log, NotificationText);
 }
 
-TArray<FName> USkinWeightsPaintToolProperties::GetSourceLODsFunc() const
+USkinWeightsPaintToolProperties* UWeightToolTransferManager::GetToolProperties() const
 {
-	if (WeightTool->GetSourceTarget())
-	{
-		return SkinWeightLayer::GetLODs(WeightTool->GetSourceTarget());	
-	}
-	return GetLODsFunc();	
-}
+	TArray<UObject*> AllToolProperties = WeightTool->GetToolProperties(false/*bEnabledOnly*/);
 
-TArray<FName> USkinWeightsPaintToolProperties::GetSourceSkinWeightProfilesFunc() const
-{
-	if (USkeletalMesh* SrcSkeletalMesh = SourceSkeletalMesh.Get())
+	for (UObject* PropertyObject : AllToolProperties)
 	{
-		const EMeshLODIdentifier LODId = GetLODId(SourceLOD);
-		const FGetMeshParameters Params(true, LODId);
-		const FMeshDescription* MeshDescription = UE::ToolTarget::GetMeshDescription(WeightTool->GetSourceTarget(), Params);
-		return SkinWeightLayer::GetSkinWeightProfilesFunc(*MeshDescription);
+		if (PropertyObject->IsA(USkinWeightsPaintToolProperties::StaticClass()))
+		{
+			return CastChecked<USkinWeightsPaintToolProperties>(PropertyObject);
+		}
 	}
-	return GetSkinWeightProfilesFunc();
+
+	return nullptr;
 }
 
 FName USkinWeightsPaintToolProperties::GetActiveSkinWeightProfile() const
@@ -682,8 +1484,7 @@ void FSkinToolDeformer::UpdateVertexDeformation(
 	{
 		// update AABB Tree for vertex selection
 		TRACE_CPUPROFILER_EVENT_SCOPE(SkinTool::UpdateAABBTree);
-		Tool->MeshSpatial->Build();
-		Tool->PolygonSelectionMechanic->UMeshTopologySelectionMechanic::GetTopologySelector()->Invalidate(true, false);
+		Tool->MeshSelector->UpdateAfterMeshDeformation();
 	}
 
 	// empty queue of vertices to update
@@ -1379,55 +2180,15 @@ void USkinWeightsPaintTool::Setup()
 	PreviewMesh->SetTangentsMode(EDynamicMeshComponentTangentsMode::AutoCalculated);
 	PreviewMesh->SetShadowsEnabled(false);
 
-	// selection colors
-	constexpr FLinearColor FaceSelectedOrange = FLinearColor(0.886f, 0.672f, 0.473f);
-	constexpr FLinearColor VertexSelectedPurple = FLinearColor(0.78f, 0.f, 0.78f);
-	constexpr FLinearColor VertexSelectedYellow = FLinearColor(1.f,1.f,0.f);
-	// configure secondary render material for selected triangles
-	// NOTE: the material returned by ToolSetupUtil::GetSelectionMaterial has a checkerboard pattern on back faces which makes it hard to use
-	if (UMaterialInterface* Material = LoadObject<UMaterial>(nullptr, TEXT("/MeshModelingToolsetExp/Materials/SculptMaterial")))
-	{
-		if (UMaterialInstanceDynamic* MatInstance = UMaterialInstanceDynamic::Create(Material, GetToolManager()))
-		{
-			MatInstance->SetVectorParameterValue(TEXT("Color"), FaceSelectedOrange);
-			PreviewMesh->SetSecondaryRenderMaterial(MatInstance);
-		}
-	}
-	// set up vertex selection mechanic
-	PolygonSelectionMechanic = NewObject<UPolygonSelectionMechanic>(this);
-	PolygonSelectionMechanic->bAddSelectionFilterPropertiesToParentTool = false;
-	PolygonSelectionMechanic->Setup(this);
-	PolygonSelectionMechanic->SetIsEnabled(false);
-	PolygonSelectionMechanic->OnSelectionChanged.AddLambda([this](){OnSelectionChanged.Broadcast();} );
-	// adjust selection rendering for this context
-	PolygonSelectionMechanic->HilightRenderer.PointColor = FLinearColor::Blue;
-	PolygonSelectionMechanic->HilightRenderer.PointSize = 10.0f;
-	// vertex highlighting once selected
-	PolygonSelectionMechanic->SelectionRenderer.LineThickness = 1.0f;
-	PolygonSelectionMechanic->SelectionRenderer.PointColor = VertexSelectedYellow;
-	PolygonSelectionMechanic->SelectionRenderer.PointSize = 5.0f;
-	PolygonSelectionMechanic->SelectionRenderer.DepthBias = 2.0f;
-	// despite the name, this renders the vertices
-	PolygonSelectionMechanic->PolyEdgesRenderer.PointColor = VertexSelectedPurple;
-	PolygonSelectionMechanic->PolyEdgesRenderer.PointSize = 5.0f;
-	PolygonSelectionMechanic->PolyEdgesRenderer.DepthBias = 2.0f;
-	PolygonSelectionMechanic->PolyEdgesRenderer.LineThickness = 1.0f;
-	// restore saved mode
-	SetComponentSelectionMode(WeightToolProperties->ComponentSelectionMode);
-	// secondary triangle buffer used to render face selection
-	PreviewMesh->EnableSecondaryTriangleBuffers([this](const FDynamicMesh3* Mesh, int32 TriangleID)
-	{
-		return PolygonSelectionMechanic->GetActiveSelection().IsSelectedTriangle(Mesh, SelectionTopology.Get(), TriangleID);
-	});
-	// notify preview mesh when triangle selection has been updated
-	PolygonSelectionMechanic->OnSelectionChanged.AddWeakLambda(this, [this]()
-	{
-		PreviewMesh->FastNotifySecondaryTrianglesChanged();
-	});
-	PolygonSelectionMechanic->OnFaceSelectionPreviewChanged.AddWeakLambda(this, [this]()
-	{
-		PreviewMesh->FastNotifySecondaryTrianglesChanged();
-	});
+	// create the transfer manager
+	TransferManager = NewObject<UWeightToolTransferManager>();
+	TransferManager->InitialSetup(this, GetViewportClient());
+
+	// setup selection for the main mesh
+	MeshSelector = NewObject<UWeightToolMeshSelector>(this);
+	auto OnSelectionChangedLambda = [this](){OnSelectionChanged.Broadcast();};
+	MeshSelector->InitialSetup(TargetWorld.Get(), this, GetViewportClient(), OnSelectionChangedLambda);
+	UpdateComponentSelectionMode();
 
 	// run all initialization for mesh/weights
 	PostEditMeshInitialization(Component, *PreviewMesh->GetMesh(), *EditedMesh);
@@ -1445,6 +2206,10 @@ void USkinWeightsPaintTool::Setup()
 	FPreviewProfileController PreviewProfileController;
 	PreviewProfileToRestore = PreviewProfileController.GetActiveProfile();
 	PreviewProfileController.SetActiveProfile(UDefaultEditorProfiles::EditingProfileName.ToString());
+	if (FEditorViewportClient* ViewportClient = GetViewportClient())
+	{
+		ViewportClient->SetViewMode(VMI_Lit_Wireframe);
+	}
 	// turn on bone colors
 	bBoneColorsToRestore = GetDefault<UPersonaOptions>()->bShowBoneColors;
 	GetMutableDefault<UPersonaOptions>()->bShowBoneColors = true;
@@ -1462,21 +2227,33 @@ void USkinWeightsPaintTool::DrawHUD(FCanvas* Canvas, IToolsContextRenderAPI* Ren
 {
 	Super::DrawHUD(Canvas, RenderAPI);
 
-	if (PolygonSelectionMechanic)
+	if (MeshSelector)
 	{
-		PolygonSelectionMechanic->DrawHUD(Canvas, RenderAPI);
+		MeshSelector->DrawHUD(Canvas, RenderAPI);
 	}
+			
+	GetWeightTransferManager()->DrawHUD(Canvas, RenderAPI);
 }
 
 void USkinWeightsPaintTool::Render(IToolsContextRenderAPI* RenderAPI)
 {
-	if (WeightToolProperties->EditingMode == EWeightEditMode::Brush)
-	{
-		Super::Render(RenderAPI);	
-	}
-	else if (PolygonSelectionMechanic && WeightToolProperties->EditingMode == EWeightEditMode::Mesh)
+	switch (WeightToolProperties->EditingMode)
 	{
-		PolygonSelectionMechanic->Render(RenderAPI);
+	case EWeightEditMode::Brush:
+		Super::Render(RenderAPI);
+	case EWeightEditMode::Mesh:
+		{
+			if (MeshSelector)
+			{
+				MeshSelector->Render(RenderAPI);	
+			}
+
+			GetWeightTransferManager()->Render(RenderAPI);
+			
+			break;
+		}
+	default:
+		return;
 	}
 }
 
@@ -1499,7 +2276,8 @@ FBox USkinWeightsPaintTool::GetWorldSpaceFocusBox()
 	case EWeightEditMode::Mesh:
 		{
 			FAxisAlignedBox3d Bounds = FAxisAlignedBox3d::Empty();
-			const TArray<int32>& SelectedVertices = GetSelectedVertices();
+			static const TArray<int32> Dummy;
+			const TArray<int32>& SelectedVertices = MeshSelector ? MeshSelector->GetSelectedVertices() : Dummy;
 			if (!SelectedVertices.IsEmpty())
 			{
 				const FDynamicMesh3* Mesh = PreviewMesh->GetMesh();
@@ -1562,8 +2340,8 @@ FBox USkinWeightsPaintTool::GetWorldSpaceFocusBox()
 	// 3. Finally, fallback on component bounds if nothing else is selected
 	static constexpr bool bNonColliding = true;
 	FBox PreviewBox = PreviewMesh->GetActor()->GetComponentsBoundingBox(bNonColliding);
-	
-	if (SourcePreviewMesh)
+	// expand the bounds by the source transfer mesh (if there is one and it's visible)
+	if (UPreviewMesh* SourcePreviewMesh = GetWeightTransferManager()->GetPreviewMesh())
 	{
 		if (AActor* SourceActor = SourcePreviewMesh->GetActor())
 		{
@@ -1676,8 +2454,8 @@ void USkinWeightsPaintTool::PostEditMeshInitialization(
 	// update vertices & triangle octrees (this must be done after PreviewMesh has been updated)
 	InitializeOctrees();
 
-	// update the polygon selection mechanic (this must be done after PreviewMesh has been updated)
-	InitializeSelectionMechanic();
+	// update the mesh selection mechanic (this must be done after PreviewMesh has been updated)
+	MeshSelector->SetMesh(PreviewMesh, &InMeshDescription, InComponent->GetComponentTransform());
 
 	// update weights
 	Weights = FSkinToolWeights();
@@ -1717,7 +2495,7 @@ void USkinWeightsPaintToolProperties::SetComponentMode(EComponentSelectionMode I
 {
 	ComponentSelectionMode = InComponentMode;
 	
-	WeightTool->SetComponentSelectionMode(ComponentSelectionMode);
+	WeightTool->UpdateComponentSelectionMode();
 	WeightTool->SetFocusInViewport();
 }
 
@@ -1751,6 +2529,12 @@ void USkinWeightsPaintToolProperties::SetBrushMode(EWeightEditOperation InBrushM
 
 bool USkinWeightsPaintTool::HitTest(const FRay& Ray, FHitResult& OutHit)
 {
+	const bool bIsInBrushMode = WeightToolProperties->EditingMode == EWeightEditMode::Brush;
+	if (!bIsInBrushMode)
+	{
+		return false;
+	}
+	
 	// do not query the triangle octree until all async ops are finished
 	TriangleOctreeFuture.Wait();
 	
@@ -2315,33 +3099,6 @@ void USkinWeightsPaintTool::InitializeOctrees()
 	}
 }
 
-void USkinWeightsPaintTool::InitializeSelectionMechanic()
-{
-	const USkeletalMeshComponent* Component = GetSkeletalMeshComponent(Target);
-	if (!ensure(Component))
-	{
-		return;
-	}
-	
-	if (!ensure(PreviewMesh && PreviewMesh->GetMesh()))
-	{
-		return;
-	}
-
-	static constexpr bool bAutoBuild = true;
-	const FDynamicMesh3* DynamicMesh = PreviewMesh->GetMesh();
-	SelectionTopology = MakeUnique<UE::Geometry::FTriangleGroupTopology>(DynamicMesh, bAutoBuild);
-	MeshSpatial = MakeUnique<FDynamicMeshAABBTree3>(DynamicMesh, bAutoBuild);
-	PolygonSelectionMechanic->Initialize(
-		DynamicMesh,
-		FTransform::Identity,
-		Component->GetWorld(),
-		SelectionTopology.Get(),
-		[this]() { return MeshSpatial.Get(); }
-	);
-	PolygonSelectionMechanic->ClearSelection();
-}
-
 void USkinWeightsPaintTool::InitializeSmoothWeightsOperator()
 {
 	if (!ensure(PreviewMesh && PreviewMesh->GetMesh()))
@@ -2411,10 +3168,9 @@ void USkinWeightsPaintTool::OnShutdown(EToolShutdownType ShutdownType)
 	WeightToolProperties->SaveProperties(this);
 
 	// shutdown polygon selection mechanic
-	if (PolygonSelectionMechanic)
+	if (MeshSelector)
 	{
-		PolygonSelectionMechanic->Shutdown();
-		PolygonSelectionMechanic = nullptr;
+		MeshSelector->Shutdown();
 	}
 
 	// apply changes to asset
@@ -2438,11 +3194,6 @@ void USkinWeightsPaintTool::OnShutdown(EToolShutdownType ShutdownType)
 	FPreviewProfileController PreviewProfileController;
 	PreviewProfileController.SetActiveProfile(PreviewProfileToRestore);
 	GetMutableDefault<UPersonaOptions>()->bShowBoneColors = bBoneColorsToRestore;
-	// mesh selection mode takes control of "Show Edges" render flag
-	if (IConsoleVariable* CVar = IConsoleManager::Get().FindConsoleVariable(TEXT("ShowFlag.MeshEdges")))
-	{
-		CVar->Unset(ECVF_SetByCode);
-	}
 
 	if (EditorContext.IsValid())
 	{
@@ -2454,7 +3205,26 @@ void USkinWeightsPaintTool::OnShutdown(EToolShutdownType ShutdownType)
 		PersonaModeManagerContext->GetPersonaEditorModeManager()->DeactivateMode(FPersonaEditModes::SkeletonSelection);
 	}
 
-	ResetSourceForTransfer();
+	GetWeightTransferManager()->Shutdown();
+}
+
+FEditorViewportClient* USkinWeightsPaintTool::GetViewportClient() const
+{
+	FEditorViewportClient* ViewportClient = nullptr;
+	if (const UPersonaEditorModeManagerContext* ModeManagerContext = PersonaModeManagerContext.Get())
+	{
+		if (const IPersonaEditorModeManager* PersonaEditorModeManager = ModeManagerContext->GetPersonaEditorModeManager())
+		{
+			ViewportClient = PersonaEditorModeManager->GetHoveredViewportClient();
+			if (!ViewportClient)
+			{
+				ViewportClient = PersonaEditorModeManager->GetFocusedViewportClient();
+			}
+		}
+	}
+
+	ensure(ViewportClient);
+	return ViewportClient;
 }
 
 void USkinWeightsPaintTool::BeginChange()
@@ -2697,7 +3467,7 @@ void USkinWeightsPaintTool::MirrorWeights(EAxis::Type Axis, EMirrorDirection Dir
 	const TMap<VertexIndex, VertexIndex>& VertexMirrorMap = MirrorData.GetVertexMap(); // <Target, Source>
 
 	// get the selected vertices
-	const TArray<int32>& SelectedVertices = GetSelectedVertices();
+	const TArray<int32>& SelectedVertices = MeshSelector->GetSelectedVertices();
 	// we need to convert selection to the equivalent target vertex indices (on the target side of the mirror plane)
 	// if a vertex is already on the target side, great
 	// if the user selected vertices on the source side, we convert them to the mirrored equivalent on the target side
@@ -2797,7 +3567,7 @@ void USkinWeightsPaintTool::EditWeightsOnVertices(
 	if (EditOperation == EWeightEditOperation::Relax)
 	{
 		CreateWeightEditsToRelaxVertices(
-			GetSelectedVertices(),
+			MeshSelector->GetSelectedVertices(),
 			VertexFalloffs,
 			Value,
 			Iterations,
@@ -2830,7 +3600,7 @@ void USkinWeightsPaintTool::PruneWeights(float Threshold, const TArray<BoneIndex
 {
 	// set weights below the given threshold to zero AND remove them as a recorded influence on that vertex
 	FMultiBoneWeightEdits WeightEditsFromPrune;
-	const TArray<VertexIndex>& VerticesToPrune = GetSelectedVertices();
+	const TArray<VertexIndex>& VerticesToPrune = MeshSelector->GetSelectedVertices();
 	for (const VertexIndex VertexID : VerticesToPrune)
 	{
 		TArray<BoneIndex> InfluencesToPrune;
@@ -2901,7 +3671,7 @@ void USkinWeightsPaintTool::AverageWeights(const float Strength)
 	}
 	
 	// get vertices to edit weights on
-	const TArray<VertexIndex>& VerticesToAverage = GetSelectedVertices();
+	const TArray<VertexIndex>& VerticesToAverage = MeshSelector->GetSelectedVertices();
 	TMap<BoneIndex, float> AveragedWeights;
 	AccumulateWeights(Weights.PreChangeWeights, VerticesToAverage, AveragedWeights);
 	TruncateWeightMap(AveragedWeights);
@@ -2989,7 +3759,7 @@ void USkinWeightsPaintTool::NormalizeWeights()
 {
 	// re-set a weight on each vertex to force normalization
 	FMultiBoneWeightEdits WeightEditsFromNormalization;
-	const TArray<VertexIndex> VerticesToNormalize = GetSelectedVertices();
+	const TArray<VertexIndex> VerticesToNormalize = MeshSelector->GetSelectedVertices();
 	for (const VertexIndex VertexID : VerticesToNormalize)
 	{
 		const VertexWeights& VertexWeights = Weights.CurrentWeights[VertexID];
@@ -3014,277 +3784,77 @@ void USkinWeightsPaintTool::NormalizeWeights()
 }
 
 void USkinWeightsPaintTool::HammerWeights()
-{
-	// get selected vertices
-	const TArray<VertexIndex> SelectedVerts = GetSelectedVertices();
-	if (SelectedVerts.IsEmpty())
-	{
-		return;
-	}
-
-	// reset mesh to ref pose so that Dijkstra path lengths are not deformed
-	Weights.Deformer.SetToRefPose(this);
-	
-	// find 1-ring neighbors of the current selection, lets call these "Surrounding" vertices
-	const FDynamicMesh3* Mesh = PreviewMesh->GetMesh();
-	TSet<int32> SurroundingVertices;
-	for (const int32 SelectedVertex : SelectedVerts)
-	{
-		for (const int32 NeighborIndex : Mesh->VtxVerticesItr(SelectedVertex))
-		{
-			if (!SelectedVerts.Contains(NeighborIndex))
-			{
-				SurroundingVertices.Add(NeighborIndex);
-			}
-		}
-	}
-
-	// seed a Dijkstra path finder with the surrounding vertices
-	UE::Geometry::TMeshDijkstra<FDynamicMesh3> PathFinder(Mesh);
-	TArray<UE::Geometry::TMeshDijkstra<FDynamicMesh3>::FSeedPoint> SeedPoints;
-	for (const int32 SurroundingVertex : SurroundingVertices)
-	{
-		SeedPoints.Add({ SurroundingVertex, SurroundingVertex, 0 });
-	}
-	PathFinder.ComputeToMaxDistance(SeedPoints, TNumericLimits<double>::Max());
-
-	// create set of weight edits that hammer the weights
-	FMultiBoneWeightEdits HammerWeightEdits;
-	
-	// for each selected vertex, find the nearest surrounding vertex and copy it's weights
-	TArray<int32> VertexPath;
-	for (const int32 SelectedVertex : SelectedVerts)
-	{
-		// find the closest surrounding vertex to this selected vertex
-		if (!PathFinder.FindPathToNearestSeed(SelectedVertex, VertexPath))
-		{
-			continue;
-		}
-		const int32 ClosestVertex = VertexPath.Last();
-
-		// remove all current weights (but don't prune)
-		for (const FVertexBoneWeight& BoneWeight : Weights.PreChangeWeights[SelectedVertex])
-		{
-			constexpr float NewWeight = 0.f;
-			constexpr bool bPruneInfluence = false;
-			HammerWeightEdits.MergeSingleEdit(BoneWeight.BoneID, SelectedVertex, NewWeight, bPruneInfluence, Weights.PreChangeWeights);
-		}
-
-		// add weights from closest vertex
-		for (const FVertexBoneWeight& BoneWeight : Weights.PreChangeWeights[ClosestVertex])
-		{
-			const float NewWeight = BoneWeight.Weight;
-			constexpr bool bPruneInfluence = false;
-			HammerWeightEdits.MergeSingleEdit(BoneWeight.BoneID, SelectedVertex, NewWeight, bPruneInfluence, Weights.PreChangeWeights);
-		}
-	}
-	
-	// apply the changes
-	const FText TransactionLabel = LOCTEXT("HammerWeightsChange", "Hammer skin weights.");
-	ApplyWeightEditsAsTransaction(HammerWeightEdits, TransactionLabel);
-
-	// put the mesh back in it's current pose
-	Weights.Deformer.SetAllVerticesToBeUpdated();
-}
-
-void USkinWeightsPaintTool::TransferWeights()
-{
-	using namespace UE::Geometry;
-	using namespace UE::AnimationCore;
-
-	const EMeshLODIdentifier TargetLODId = GetLODId(WeightToolProperties->ActiveLOD);
-    const FGetMeshParameters TargetParams(true, TargetLODId);
-	FDynamicMesh3 TargetMesh = UE::ToolTarget::GetDynamicMeshCopy(Target, TargetParams);
-
-	const EMeshLODIdentifier SourceLODId = GetLODId(WeightToolProperties->SourceLOD);
-
-	const FName TargetProfile = WeightToolProperties->GetActiveSkinWeightProfile();
-	FDynamicMeshVertexSkinWeightsAttribute* TransferredSkinWeights = nullptr;
-
-	TArray<int32> VertexSubset;
-	bool bCreateTransaction = false;
-	
-	// If a source mesh was given, or if the LODs differ, then do a weight transfer using Rinat's algo.
-	// Otherwise, do an exact copy from a different profile.
-	if (SourceTarget || TargetLODId != SourceLODId)
-	{
-		const FGetMeshParameters SourceParams(true, SourceLODId);
-		const FDynamicMesh3 SourceMesh = UE::ToolTarget::GetDynamicMeshCopy(SourceTarget ? SourceTarget : Target, SourceParams);
-		
-		if (!SourceMesh.HasAttributes() || !SourceMesh.Attributes()->HasBones())
-		{
-			const FText NotificationText = LOCTEXT("NoWeightsFoundInTransfer", "No skin weights were found in the source skeletal mesh. No weights were transferred.");
-			ShowEditorMessage(ELogVerbosity::Error, NotificationText);
-			return;
-		}
-		
-		if (SourceMesh.Attributes()->GetNumBones() == 0)
-		{
-			const FText NotificationText = LOCTEXT("NoBonesFoundInTransfer", "No bones were found in the source skeletal mesh. No weights were transferred.");
-			ShowEditorMessage(ELogVerbosity::Error, NotificationText);
-			return;
-		}
-		
-		FTransferBoneWeights TransferBoneWeights(&SourceMesh, WeightToolProperties->SourceSkinWeightProfile);
-		TransferBoneWeights.TransferMethod = FTransferBoneWeights::ETransferBoneWeightsMethod::InpaintWeights;
-		
-		if (!TargetMesh.HasAttributes())
-		{
-			TargetMesh.EnableAttributes();
-		}
-
-		FDynamicMeshAttributeSet* TargetAttributes = TargetMesh.Attributes();
-		if (!TargetAttributes->HasBones())
-		{
-			TargetAttributes->CopyBoneAttributes(*SourceMesh.Attributes());
-		}
-		else
-		{
-			const USkeletalMeshComponent* Component = GetSkeletalMeshComponent(Target);
-			if (ensure(Component))
-			{
-				const FReferenceSkeleton& RefSkeleton = Component->GetSkeletalMeshAsset()->GetRefSkeleton();
-				ensure(TargetAttributes->GetNumBones() == RefSkeleton.GetRawBoneNum());
-			}
-		}
-
-		// NOTE should we expose all the options?
-		// 	TransferBoneWeights.NormalThreshold;
-		// 	TransferBoneWeights.SearchRadius
-		// 	TransferBoneWeights.NumSmoothingIterations;
-		// 	TransferBoneWeights.SmoothingStrength;
-		// 	TransferBoneWeights.LayeredMeshSupport;
-		// 	TransferBoneWeights.ForceInpaintWeightMapName;
-		
-		if (TransferBoneWeights.Validate() != EOperationValidationResult::Ok)
-		{
-			return;
-		}
-
-		if (WeightToolProperties->EditingMode == EWeightEditMode::Mesh)
-		{
-			TransferBoneWeights.TargetVerticesSubset = GetSelectedVertices();
-		}
-
-		if (!TransferBoneWeights.TransferWeightsToMesh(TargetMesh, TargetProfile))
-		{
-			// notify user
-			const FText NotificationText = LOCTEXT("TransferWeightsUnknownIssue", "Transfer weights operation encountered an unknown issue. No weights were transferred.");
-			ShowEditorMessage(ELogVerbosity::Error, NotificationText);
-			return;
-		}
-		
-		TransferredSkinWeights = TargetAttributes->GetSkinWeightsAttribute(TargetProfile);
-		
-		VertexSubset = MoveTemp(TransferBoneWeights.TargetVerticesSubset);
-		bCreateTransaction = true;
-	}
-	else
-	{
-		FDynamicMeshAttributeSet* TargetAttributes = TargetMesh.Attributes();
-		TransferredSkinWeights = TargetAttributes->GetSkinWeightsAttribute(TargetProfile);
-		if (!TransferredSkinWeights)
-		{
-			TransferredSkinWeights = new FDynamicMeshVertexSkinWeightsAttribute(&TargetMesh);
-			TargetMesh.Attributes()->AttachSkinWeightsAttribute(TargetProfile, TransferredSkinWeights);
-		}
-		
-		const FDynamicMeshVertexSkinWeightsAttribute* SourceAttributes = TargetAttributes->GetSkinWeightsAttribute(WeightToolProperties->SourceSkinWeightProfile);
-		check(SourceAttributes != nullptr);
-
-		if (SourceAttributes != TransferredSkinWeights)
-		{
-			if (WeightToolProperties->EditingMode == EWeightEditMode::Mesh)
-			{
-				VertexSubset = GetSelectedVertices();
-			}
-
-			if (VertexSubset.IsEmpty())
-			{
-				TransferredSkinWeights->Copy(*SourceAttributes);
-			}
-			else
-			{
-				for (int32 VertexID: VertexSubset)
-				{
-					FBoneWeights BoneWeights;
-					SourceAttributes->GetValue(VertexID, BoneWeights);
-					TransferredSkinWeights->SetValue(VertexID, BoneWeights);
-				}
-			}
-			bCreateTransaction = true;
-		}
+{
+	// get selected vertices
+	const TArray<VertexIndex> SelectedVerts = MeshSelector->GetSelectedVertices();
+	if (SelectedVerts.IsEmpty())
+	{
+		return;
 	}
 
-	if (bCreateTransaction)
+	// reset mesh to ref pose so that Dijkstra path lengths are not deformed
+	Weights.Deformer.SetToRefPose(this);
+	
+	// find 1-ring neighbors of the current selection, lets call these "Surrounding" vertices
+	const FDynamicMesh3* Mesh = PreviewMesh->GetMesh();
+	TSet<int32> SurroundingVertices;
+	for (const int32 SelectedVertex : SelectedVerts)
 	{
-		// weight edits for transaction
-		FMultiBoneWeightEdits WeightEdits;
-		
-		// spin through all the transferred skin weights and record a weight edit to apply as a transaction
-		check(TransferredSkinWeights);
-		static constexpr float ZeroWeight = 0.f;
-		const bool bUseSubset = !VertexSubset.IsEmpty();
-		const int32 NumVertices = bUseSubset ? VertexSubset.Num() : TargetMesh.VertexCount();
-		const FNonManifoldMappingSupport NonManifoldMappingSupport(TargetMesh);
-		for (int32 VertexIndex = 0; VertexIndex < NumVertices; ++VertexIndex)
+		for (const int32 NeighborIndex : Mesh->VtxVerticesItr(SelectedVertex))
 		{
-			const int32 VertexID = bUseSubset ? VertexSubset[VertexIndex] : VertexIndex;
-			const int32 SrcVertexID = NonManifoldMappingSupport.GetOriginalNonManifoldVertexID(VertexID);
-			
-			// remove all weight on vertex
-			const VertexWeights& VertexBoneWeights = Weights.PreChangeWeights[SrcVertexID];
-			if (!VertexBoneWeights.IsEmpty())
-			{
-				for (const FVertexBoneWeight& BoneWeight : VertexBoneWeights)
-				{
-					// when transferring weights we do prune influences because we would prefer the results to be identical to the source
-					constexpr bool bPruneInfluence = true;
-					WeightEdits.MergeSingleEdit(BoneWeight.BoneID, SrcVertexID, ZeroWeight, bPruneInfluence, Weights.PreChangeWeights);
-				}
-			}
-			else
-			{
-				// in the unlikely event that the target vertex has no weight, "fake" remove it from the root so that undo will put it back there
-				constexpr bool bPruneInfluence = false;
-				BoneIndex RootBoneIndex = 0;
-				WeightEdits.MergeSingleEdit(RootBoneIndex, SrcVertexID, ZeroWeight, bPruneInfluence, Weights.PreChangeWeights);
-			}
-
-			// update with new weight
-			FBoneWeights TransferedBoneWeights;
-			TransferredSkinWeights->GetValue(VertexID, TransferedBoneWeights);
-			for (const FBoneWeight& BoneWeight: TransferedBoneWeights)
+			if (!SelectedVerts.Contains(NeighborIndex))
 			{
-				const int32 BoneIndex = BoneWeight.GetBoneIndex();					
-				const float NewWeight = BoneWeight.GetWeight();
-				constexpr bool bPruneInfluence = false;
-				WeightEdits.MergeSingleEdit(BoneIndex, SrcVertexID, NewWeight, bPruneInfluence, Weights.PreChangeWeights);
+				SurroundingVertices.Add(NeighborIndex);
 			}
 		}
-
-		// apply the changes as a transaction
-		const FText TransactionLabel = LOCTEXT("TransferWeightsChange", "Transfer skin weights.");
-		ApplyWeightEditsAsTransaction(WeightEdits, TransactionLabel);
-		
-		// put the mesh back in it's current pose
-		Weights.Deformer.SetAllVerticesToBeUpdated();
 	}
 
-	// notify user that weights were transferred.
-	const FText NotificationText = LOCTEXT("WeightsTransferred", "Skin weights transferred.");
-	ShowEditorMessage(ELogVerbosity::Log, NotificationText);
-}
+	// seed a Dijkstra path finder with the surrounding vertices
+	UE::Geometry::TMeshDijkstra<FDynamicMesh3> PathFinder(Mesh);
+	TArray<UE::Geometry::TMeshDijkstra<FDynamicMesh3>::FSeedPoint> SeedPoints;
+	for (const int32 SurroundingVertex : SurroundingVertices)
+	{
+		SeedPoints.Add({ SurroundingVertex, SurroundingVertex, 0 });
+	}
+	PathFinder.ComputeToMaxDistance(SeedPoints, TNumericLimits<double>::Max());
 
-bool USkinWeightsPaintTool::CanTransferWeights() const
-{
-	if (GetSourceTarget() != nullptr)
+	// create set of weight edits that hammer the weights
+	FMultiBoneWeightEdits HammerWeightEdits;
+	
+	// for each selected vertex, find the nearest surrounding vertex and copy it's weights
+	TArray<int32> VertexPath;
+	for (const int32 SelectedVertex : SelectedVerts)
 	{
-		return true; 
+		// find the closest surrounding vertex to this selected vertex
+		if (!PathFinder.FindPathToNearestSeed(SelectedVertex, VertexPath))
+		{
+			continue;
+		}
+		const int32 ClosestVertex = VertexPath.Last();
+
+		// remove all current weights (but don't prune)
+		for (const FVertexBoneWeight& BoneWeight : Weights.PreChangeWeights[SelectedVertex])
+		{
+			constexpr float NewWeight = 0.f;
+			constexpr bool bPruneInfluence = false;
+			HammerWeightEdits.MergeSingleEdit(BoneWeight.BoneID, SelectedVertex, NewWeight, bPruneInfluence, Weights.PreChangeWeights);
+		}
+
+		// add weights from closest vertex
+		for (const FVertexBoneWeight& BoneWeight : Weights.PreChangeWeights[ClosestVertex])
+		{
+			const float NewWeight = BoneWeight.Weight;
+			constexpr bool bPruneInfluence = false;
+			HammerWeightEdits.MergeSingleEdit(BoneWeight.BoneID, SelectedVertex, NewWeight, bPruneInfluence, Weights.PreChangeWeights);
+		}
 	}
+	
+	// apply the changes
+	const FText TransactionLabel = LOCTEXT("HammerWeightsChange", "Hammer skin weights.");
+	ApplyWeightEditsAsTransaction(HammerWeightEdits, TransactionLabel);
 
-	return GetLODId(WeightToolProperties->ActiveLOD) != GetLODId(WeightToolProperties->SourceLOD) ||
-		   WeightToolProperties->GetActiveSkinWeightProfile() != WeightToolProperties->SourceSkinWeightProfile; 		
+	// put the mesh back in it's current pose
+	Weights.Deformer.SetAllVerticesToBeUpdated();
 }
 
 const FString USkinWeightsPaintTool::CopyPasteWeightsIdentifier = TEXT("UNREAL_VERTEX_WEIGHTS:");
@@ -3292,7 +3862,7 @@ const FString USkinWeightsPaintTool::CopyPasteWeightsIdentifier = TEXT("UNREAL_V
 void USkinWeightsPaintTool::CopyWeights()
 {
 	// get all the selected vertices
-	const TArray<VertexIndex>& VerticesToCopy = GetSelectedVertices();
+	const TArray<VertexIndex>& VerticesToCopy = MeshSelector->GetSelectedVertices();
 	if (VerticesToCopy.IsEmpty())
 	{
 		const FText NotificationText = FText::FromString(TEXT("No vertices were selected. No weights were copied to the clipboard."));
@@ -3338,7 +3908,7 @@ void USkinWeightsPaintTool::CopyWeights()
 
 void USkinWeightsPaintTool::PasteWeights()
 {
-	const TArray<VertexIndex>& VerticesToPasteOn = GetSelectedVertices();
+	const TArray<VertexIndex>& VerticesToPasteOn = MeshSelector->GetSelectedVertices();
 	if (VerticesToPasteOn.IsEmpty())
 	{
 		const FText NotificationText = FText::FromString("No vertices were selected. No weights were pasted.");
@@ -3654,7 +4224,12 @@ void USkinWeightsPaintTool::ToggleEditingMode()
 	SetBrushEnabled(WeightToolProperties->EditingMode == EWeightEditMode::Brush);
 
 	// toggle mesh mode
-	PolygonSelectionMechanic->SetIsEnabled(WeightToolProperties->EditingMode == EWeightEditMode::Mesh);
+	const bool bIsMeshEditing = WeightToolProperties->EditingMode == EWeightEditMode::Mesh;
+	MeshSelector->SetIsEnabled(bIsMeshEditing);
+	if (UWeightToolMeshSelector* SourceMeshSelector = TransferManager->GetMeshSelector())
+	{
+		SourceMeshSelector->SetIsEnabled(bIsMeshEditing);
+	}
 
 	// toggle bone select mode
 	// this mode is set to be compatible with the 
@@ -3673,42 +4248,51 @@ void USkinWeightsPaintTool::ToggleEditingMode()
 	SetFocusInViewport();
 }
 
-void USkinWeightsPaintTool::SetComponentSelectionMode(EComponentSelectionMode InMode)
+void USkinWeightsPaintTool::UpdateComponentSelectionMode() const
 {
-	PolygonSelectionMechanic->Properties->bSelectVertices = InMode == EComponentSelectionMode::Vertices;
-	PolygonSelectionMechanic->Properties->bSelectEdges = InMode == EComponentSelectionMode::Edges;
-	PolygonSelectionMechanic->Properties->bSelectFaces = InMode == EComponentSelectionMode::Faces;
-	
-	PolygonSelectionMechanic->SetShowSelectableCorners(InMode == EComponentSelectionMode::Vertices);
-	PolygonSelectionMechanic->SetShowEdges(InMode == EComponentSelectionMode::Edges);
+	if (MeshSelector)
+	{
+		MeshSelector->SetComponentSelectionMode(WeightToolProperties->ComponentSelectionMode);
+	}
 
-	IConsoleVariable* CVar = IConsoleManager::Get().FindConsoleVariable(TEXT("ShowFlag.MeshEdges"));
-	if (CVar)
+	if (UWeightToolMeshSelector* SourceMeshSelector = TransferManager->GetMeshSelector())
 	{
-		const float Value = (InMode == EComponentSelectionMode::Edges) ? 0.0f : 1.0f;
-		CVar->Set(Value, ECVF_SetByCode);
+		SourceMeshSelector->SetComponentSelectionMode(WeightToolProperties->ComponentSelectionMode);
 	}
 }
 
-void USkinWeightsPaintTool::GrowSelection() const
+UWeightToolMeshSelector* USkinWeightsPaintTool::GetMainMeshSelector()
 {
-	PolygonSelectionMechanic->GrowSelection();
+	return MeshSelector;
 }
 
-void USkinWeightsPaintTool::ShrinkSelection() const
+UWeightToolMeshSelector* USkinWeightsPaintTool::GetActiveMeshSelector()
 {
-	PolygonSelectionMechanic->ShrinkSelection();
+	if (WeightToolProperties->MeshSelectMode == EMeshTransferOption::Source)
+	{
+		return GetWeightTransferManager()->GetMeshSelector();
+	}
+
+	return MeshSelector;
 }
 
-void USkinWeightsPaintTool::FloodSelection() const
+bool USkinWeightsPaintTool::HasActiveSelectionOnMainMesh()
 {
-	PolygonSelectionMechanic->FloodSelection();
+	const bool bSelectingTargetMesh = WeightToolProperties->MeshSelectMode == EMeshTransferOption::Target;
+	const bool bTargetMeshHasSelection = GetMainMeshSelector()->IsAnyComponentSelected();
+	return bSelectingTargetMesh && bTargetMeshHasSelection;
 }
 
 void USkinWeightsPaintTool::SelectAffected() const
 {
+	UPolygonSelectionMechanic* SelectionMechanic = MeshSelector->GetSelectionMechanic();
+	if (!ensure(SelectionMechanic))
+	{
+		return;
+	}
+	
 	GetToolManager()->BeginUndoTransaction(LOCTEXT("AffectedSelectionChange", "Select Affected"));
-	PolygonSelectionMechanic->BeginChange();
+	SelectionMechanic->BeginChange();
 	
 	// get all vertices affected by all selected bones
 	TSet<int32> AffectedVertices;
@@ -3721,7 +4305,7 @@ void USkinWeightsPaintTool::SelectAffected() const
 	FGroupTopologySelection Selection;
 
 	// optionally add/remove/replace selection based on modifier key state
-	const FGroupTopologySelection& CurrentSelection = PolygonSelectionMechanic->GetActiveSelection();
+	const FGroupTopologySelection& CurrentSelection = SelectionMechanic->GetActiveSelection();
 	if (bShiftToggle)
 	{
 		// ADD to selection
@@ -3741,55 +4325,11 @@ void USkinWeightsPaintTool::SelectAffected() const
 	
 	// select vertices
 	constexpr bool bBroadcast = true;
-	PolygonSelectionMechanic->SetSelection(Selection, bBroadcast);
-	PolygonSelectionMechanic->EndChangeAndEmitIfModified();
-	GetToolManager()->EndUndoTransaction();
-}
-
-void USkinWeightsPaintTool::SelectBorder() const
-{
-	const FDynamicMesh3* Mesh = PreviewMesh->GetMesh();
-
-	const FGroupTopologySelection& CurrentSelection = PolygonSelectionMechanic->GetActiveSelection();
-	
-	// find all border vertices
-	// a "border" vertex is one that has a 1-ring neighbor that is not in the selection set
-	TSet<int32> BorderVertices;
-	for (const int32 SelectedVertex : CurrentSelection.SelectedCornerIDs)
-	{
-		for (const int32 NeighborIndex : Mesh->VtxVerticesItr(SelectedVertex))
-		{
-			if (!CurrentSelection.SelectedCornerIDs.Contains(NeighborIndex))
-			{
-				BorderVertices.Add(SelectedVertex);
-			}
-		}
-	}
-
-	GetToolManager()->BeginUndoTransaction(LOCTEXT("BorderSelectionChange", "Select Border"));
-	PolygonSelectionMechanic->BeginChange();
-
-	// create selection set
-	FGroupTopologySelection Selection;
-	Selection.SelectedCornerIDs.Append(BorderVertices);
-	
-	// select vertices
-	constexpr bool bBroadcast = true;
-	PolygonSelectionMechanic->SetSelection(Selection, bBroadcast);
-	PolygonSelectionMechanic->EndChangeAndEmitIfModified();
+	SelectionMechanic->SetSelection(Selection, bBroadcast);
+	SelectionMechanic->EndChangeAndEmitIfModified();
 	GetToolManager()->EndUndoTransaction();
 }
 
-bool USkinWeightsPaintTool::IsAnyComponentSelected() const
-{
-	if (!PolygonSelectionMechanic)
-	{
-		return false;
-	}
-
-	return PolygonSelectionMechanic->HasSelection();
-}
-
 bool USkinWeightsPaintTool::IsSelectionIsolated() const
 {
 	return PartialMeshDescription.IsValid();
@@ -3810,7 +4350,8 @@ void USkinWeightsPaintTool::SetIsolateSelected(const bool bIsolateSelection)
 		return;
 	}
 	
-	if (!ensure(PolygonSelectionMechanic))
+	UPolygonSelectionMechanic* SelectionMechanic = MeshSelector->GetSelectionMechanic();
+	if (!ensure(SelectionMechanic))
 	{
 		return;
 	}
@@ -3839,9 +4380,9 @@ void USkinWeightsPaintTool::SetIsolateSelected(const bool bIsolateSelection)
 	IsolatedSelectionToRestoreVertices.ElementType = UE::Geometry::EGeometryElementType::Vertex;
 	IsolatedSelectionToRestoreEdges.ElementType = UE::Geometry::EGeometryElementType::Edge;
 	IsolatedSelectionToRestoreFaces.ElementType = UE::Geometry::EGeometryElementType::Face;
-	PolygonSelectionMechanic->GetSelection_AsTriangleTopology(IsolatedSelectionToRestoreVertices);
-	PolygonSelectionMechanic->GetSelection_AsTriangleTopology(IsolatedSelectionToRestoreEdges);
-	PolygonSelectionMechanic->GetSelection_AsTriangleTopology(IsolatedSelectionToRestoreFaces);
+	SelectionMechanic->GetSelection_AsTriangleTopology(IsolatedSelectionToRestoreVertices);
+	SelectionMechanic->GetSelection_AsTriangleTopology(IsolatedSelectionToRestoreEdges);
+	SelectionMechanic->GetSelection_AsTriangleTopology(IsolatedSelectionToRestoreFaces);
 
 	// store copy of original FDynamicMesh to restore
 	FDynamicMesh3 DynamicMesh;
@@ -3851,7 +4392,7 @@ void USkinWeightsPaintTool::SetIsolateSelected(const bool bIsolateSelection)
 
 	// create a submesh from the selected triangles (or triangles connected to selected vertices/edges)
 	TArray<int32> TrianglesToIsolate;
-	GetSelectedTriangles(TrianglesToIsolate);
+	MeshSelector->GetSelectedTriangles(TrianglesToIsolate);
 	if (TrianglesToIsolate.IsEmpty())
 	{
 		return;
@@ -3923,9 +4464,12 @@ void USkinWeightsPaintTool::FinishIsolatedSelection()
 	Weights.ApplyCurrentWeightsToMeshDescription(EditedMesh);
 
 	// restore selection (allows for easily adjusting crop)
-	PolygonSelectionMechanic->SetSelection_AsTriangleTopology(IsolatedSelectionToRestoreVertices);
-	PolygonSelectionMechanic->SetSelection_AsTriangleTopology(IsolatedSelectionToRestoreEdges);
-	PolygonSelectionMechanic->SetSelection_AsTriangleTopology(IsolatedSelectionToRestoreFaces);
+	if (UPolygonSelectionMechanic* SelectionMechanic = MeshSelector->GetSelectionMechanic())
+	{
+		SelectionMechanic->SetSelection_AsTriangleTopology(IsolatedSelectionToRestoreVertices);
+		SelectionMechanic->SetSelection_AsTriangleTopology(IsolatedSelectionToRestoreEdges);
+		SelectionMechanic->SetSelection_AsTriangleTopology(IsolatedSelectionToRestoreFaces);
+	}
 
 	PartialMeshDescription = nullptr;
 }
@@ -3950,63 +4494,6 @@ int32 USkinWeightsPaintTool::FullToPartialMeshVertexIndex(int32 FullMeshVertexIn
 	return PartialSubMesh.MapVertexToSubmesh(FullMeshVertexIndex);
 }
 
-const TArray<int32>& USkinWeightsPaintTool::GetSelectedVertices()
-{
-	SelectedVerticesInternal.Empty();
-	if (!PolygonSelectionMechanic)
-	{
-		return SelectedVerticesInternal;
-	}
-	
-	const FGroupTopologySelection& Selection = PolygonSelectionMechanic->GetActiveSelection();
-	const FDynamicMesh3* DynamicMesh = PreviewMesh->GetMesh();
-	const FVertexArray& AllVertices = EditedMesh->Vertices();
-
-	// validate and add vertices to the output array
-	auto AddVertices = [this, &AllVertices](const TSet<int32>& VerticesToAdd)
-	{
-		// we have to make sure that the vertex ids are safe to use as PolygonSelectionMechanic does not act on the
-		// mesh description but on the dynamic mesh that can duplicate vertices when dealing with degenerate triangles.
-		// cf. FMeshDescriptionToDynamicMesh::Convert for more details.
-		Algo::CopyIf(VerticesToAdd, SelectedVerticesInternal, [&](int32 VertexID)
-		{
-			return AllVertices.IsValid(VertexID);	
-		});
-	};
-
-	// add selected vertices
-	AddVertices(Selection.SelectedCornerIDs);
-
-	// add vertices on selected edges
-	{
-		TSet<int32> VerticesInSelectedEdges;
-		for (const int32 SelectedEdgeIndex : Selection.SelectedEdgeIDs)
-		{
-			FDynamicMesh3::FEdge CurrentEdge = DynamicMesh->GetEdge(SelectedEdgeIndex);
-			VerticesInSelectedEdges.Add(CurrentEdge.Vert.A);
-			VerticesInSelectedEdges.Add(CurrentEdge.Vert.B);
-		}
-		
-		AddVertices(VerticesInSelectedEdges);
-	}
-
-	// add vertices in selected faces
-	{
-		TSet<int32> VerticesInSelectedFaces;
-		for (const int32 SelectedFaceIndex : Selection.SelectedGroupIDs)
-		{
-			UE::Geometry::FIndex3i TriangleVertices = DynamicMesh->GetTriangleRef(SelectedFaceIndex);
-			VerticesInSelectedFaces.Add(TriangleVertices[0]);
-			VerticesInSelectedFaces.Add(TriangleVertices[1]);
-			VerticesInSelectedFaces.Add(TriangleVertices[2]);
-		}
-		
-		AddVertices(VerticesInSelectedFaces);
-	}
-	
-	return SelectedVerticesInternal;
-}
-
 void USkinWeightsPaintTool::GetVerticesAffectedByBone(BoneIndex IndexOfBone, TSet<int32>& OutVertexIndices) const
 {
 	VertexIndex VertexID = 0;
@@ -4031,42 +4518,6 @@ void USkinWeightsPaintTool::GetVerticesAffectedByBone(BoneIndex IndexOfBone, TSe
 	}
 }
 
-void USkinWeightsPaintTool::GetSelectedTriangles(TArray<int32>& OutTriangleIndices) const
-{
-	OutTriangleIndices.Empty();
-	if (!PolygonSelectionMechanic)
-	{
-		return;
-	}
-	
-	const FGroupTopologySelection& Selection = PolygonSelectionMechanic->GetActiveSelection();
-	const FDynamicMesh3* DynamicMesh = PreviewMesh->GetMesh();
-	TSet<int32> TriangleSet;
-
-	// add triangles connected to selected vertices
-	for (const int32 VertexIndex : Selection.SelectedCornerIDs)
-	{
-		DynamicMesh->EnumerateVertexTriangles(VertexIndex, [&TriangleSet](int32 TriangleIndex)
-		{
-			TriangleSet.Add(TriangleIndex);
-		});
-	}
-
-	// add triangles connected to selected edges
-	for (const int32 EdgeIndex : Selection.SelectedEdgeIDs)
-	{
-		DynamicMesh->EnumerateEdgeTriangles(EdgeIndex, [&TriangleSet](int32 TriangleIndex)
-		{
-			TriangleSet.Add(TriangleIndex);
-		});
-	}
-
-	// add selected triangles
-	TriangleSet.Append(Selection.SelectedGroupIDs);	
-
-	OutTriangleIndices = TriangleSet.Array();
-}
-
 void USkinWeightsPaintTool::GetInfluences(const TArray<int32>& VertexIndices, TArray<BoneIndex>& OutBoneIndices)
 {
 	for (const int32 SelectedVertex : VertexIndices)
@@ -4088,14 +4539,14 @@ float USkinWeightsPaintTool::GetAverageWeightOnBone(
 	float TotalWeight = 0.f;
 	float NumVerticesInfluencedByBone = 0.f;
 	
-	for (const int32 SelectedVertex : VertexIndices)
+	for (const VertexIndex VertexID : VertexIndices)
 	{
-		if (!Weights.CurrentWeights.IsValidIndex(SelectedVertex))
+		if (!Weights.CurrentWeights.IsValidIndex(VertexID))
 		{
 			continue;
 		}
 		
-		for (const FVertexBoneWeight& VertexBoneData : Weights.CurrentWeights[SelectedVertex])
+		for (const FVertexBoneWeight& VertexBoneData : Weights.CurrentWeights[VertexID])
 		{
 			if (VertexBoneData.BoneID == InBoneIndex)
 			{
@@ -4177,75 +4628,10 @@ void USkinWeightsPaintTool::OnPropertyModified(UObject* ModifiedObject, FPropert
 		}
 	}
 
-	if (ModifiedProperty->GetName() == GET_MEMBER_NAME_STRING_CHECKED(USkinWeightsPaintToolProperties, SourceSkeletalMesh))
-	{
-		ResetSourceForTransfer(WeightToolProperties->SourceSkeletalMesh.Get());
-	}
-
-	if (ModifiedProperty->GetName() == GET_MEMBER_NAME_STRING_CHECKED(USkinWeightsPaintToolProperties, SourceLOD))
-	{
-		if (SourcePreviewMesh)
-		{
-			const EMeshLODIdentifier SourceLODId = GetLODId(WeightToolProperties->SourceLOD);
-			const FGetMeshParameters SourceParams(true, SourceLODId);
-			SourcePreviewMesh->ReplaceMesh(UE::ToolTarget::GetDynamicMeshCopy(SourceTarget, SourceParams));
-		}
-	}
-
-	if (ModifiedProperty->GetName() == GET_MEMBER_NAME_STRING_CHECKED(USkinWeightsPaintToolProperties, SourcePreviewOffset))
-	{
-		if (SourcePreviewMesh)
-		{
-			SourcePreviewMesh->SetTransform(WeightToolProperties->SourcePreviewOffset);
-		}
-	}
+	// let the mesh transfer system react to properties being set
+	TransferManager->OnPropertyModified(WeightToolProperties, ModifiedProperty);
 	
 	SetFocusInViewport();
 }
 
-void USkinWeightsPaintTool::ResetSourceForTransfer(USkeletalMesh* InSkeletalMesh)
-{
-	if (SourcePreviewMesh)
-	{
-		SourcePreviewMesh->SetVisible(false);
-		SourcePreviewMesh->Disconnect();
-		SourcePreviewMesh = nullptr;
-	}
-
-	if (SourceTarget)
-	{
-		SourceTarget = nullptr;
-	}
-
-	if (InSkeletalMesh)
-	{
-		SourceTarget = TargetManager->BuildTarget(InSkeletalMesh, FToolTargetTypeRequirements());
-
-		SourcePreviewMesh = NewObject<UPreviewMesh>(this);
-		SourcePreviewMesh->CreateInWorld(TargetWorld.Get(), FTransform::Identity);
-
-		if (USkeletalMeshComponent* SkeletalMeshComponent = GetSkeletalMeshComponent(Target))
-		{
-			const FBoxSphereBounds TargetBounds = SkeletalMeshComponent->GetSkeletalMeshAsset()->GetBounds();
-			const FBoxSphereBounds SourceBounds = InSkeletalMesh->GetBounds();
-
-			FTransform Transform = UE::ToolTarget::GetLocalToWorldTransform(Target);
-			FVector Location = Transform.GetLocation();
-			Location.X += TargetBounds.GetBoxExtrema(1).X;
-			Location.X += 1.1 * SourceBounds.GetBoxExtrema(1).X;
-			Transform.SetLocation(Location);
-			WeightToolProperties->SourcePreviewOffset = Transform;
-		}
-		
-		SourcePreviewMesh->SetTransform(WeightToolProperties->SourcePreviewOffset);
-
-		ToolSetupUtil::ApplyRenderingConfigurationToPreview(SourcePreviewMesh, SourceTarget);
-		SourcePreviewMesh->SetTangentsMode(EDynamicMeshComponentTangentsMode::AutoCalculated);
-		SourcePreviewMesh->ReplaceMesh(UE::ToolTarget::GetDynamicMeshCopy(SourceTarget));
-
-		const FComponentMaterialSet MaterialSet = UE::ToolTarget::GetMaterialSet(SourceTarget);
-		SourcePreviewMesh->SetMaterials(MaterialSet.Materials);
-	}
-}
-
 #undef LOCTEXT_NAMESPACE
\ No newline at end of file
diff --git a/Engine/Plugins/Experimental/MeshModelingToolsetExp/Source/MeshModelingToolsEditorOnlyExp/Public/SkeletalMesh/SkinWeightsPaintTool.h b/Engine/Plugins/Experimental/MeshModelingToolsetExp/Source/MeshModelingToolsEditorOnlyExp/Public/SkeletalMesh/SkinWeightsPaintTool.h
--- a/Engine/Plugins/Experimental/MeshModelingToolsetExp/Source/MeshModelingToolsEditorOnlyExp/Public/SkeletalMesh/SkinWeightsPaintTool.h
+++ b/Engine/Plugins/Experimental/MeshModelingToolsetExp/Source/MeshModelingToolsEditorOnlyExp/Public/SkeletalMesh/SkinWeightsPaintTool.h
@@ -26,6 +26,7 @@
 class USkinWeightsPaintTool;
 class UPolygonSelectionMechanic;
 class UPersonaEditorModeManagerContext;
+class FEditorViewportClient;
 
 namespace UE::Geometry 
 {
@@ -55,6 +56,14 @@
 	Faces
 };
 
+// weight transfers happen between a source and target
+UENUM()
+enum class EMeshTransferOption : uint8
+{
+	Source,
+	Target,
+};
+
 // weight color mode
 UENUM()
 enum class EWeightColorMode : uint8
@@ -307,6 +316,8 @@
 		TMap<int32, int32> BoneMap;
 		TMap<int32, int32> VertexMap; // <Target, Source>
 	};
+
+	
 }
 
 UCLASS()
@@ -413,9 +424,9 @@
 	FSkinWeightBrushConfig BrushConfigRelax;
 
 	// skin weight layer properties
-	UPROPERTY(EditAnywhere, Category = SkinWeightLayer, meta = (DisplayName = "Active LOD", GetOptions = GetLODsFunc))
+	UPROPERTY(EditAnywhere, Category = SkinWeightLayer, meta = (DisplayName = "Active LOD", GetOptions = GetTargetLODsFunc))
 	FName ActiveLOD = "LOD0";
-	UPROPERTY(EditAnywhere, Category = SkinWeightLayer, meta = (DisplayName = "Active Profile", GetOptions = GetSkinWeightProfilesFunc))
+	UPROPERTY(EditAnywhere, Category = SkinWeightLayer, meta = (DisplayName = "Active Profile", GetOptions = GetTargetSkinWeightProfilesFunc))
 	FName ActiveSkinWeightProfile = FSkeletalMeshAttributesShared::DefaultSkinWeightProfileName;
 	
 	// new profile properties
@@ -436,31 +447,167 @@
 	void SetBrushMode(EWeightEditOperation InBrushMode);
 
 	// transfer
-	UPROPERTY(EditAnywhere, Transient, Category = WeightTransfer)
+	UPROPERTY(EditAnywhere, Transient, Category = "WeightTransfer")
 	TWeakObjectPtr<USkeletalMesh> SourceSkeletalMesh;
-	
-	UPROPERTY(EditAnywhere, Category = "WeightTransfer|SkinWeightLayer", meta = (GetOptions = GetSourceLODsFunc))
+	UPROPERTY(EditAnywhere, Transient, Category = "WeightTransfer")
+	EMeshTransferOption MeshSelectMode = EMeshTransferOption::Target;
+	UPROPERTY(EditAnywhere, Category = "WeightTransfer", meta = (GetOptions = GetSourceLODsFunc))
 	FName SourceLOD = "LOD0";
-	
-	UPROPERTY(EditAnywhere, Category = "WeightTransfer|SkinWeightLayer", meta = (DisplayName = "Source Profile", GetOptions = GetSourceSkinWeightProfilesFunc))
+	UPROPERTY(EditAnywhere, Category = "WeightTransfer", meta = (DisplayName = "Source Profile", GetOptions = GetSourceSkinWeightProfilesFunc))
 	FName SourceSkinWeightProfile = FSkeletalMeshAttributesShared::DefaultSkinWeightProfileName;
-	
-	UPROPERTY(EditAnywhere, Transient, Category = "WeightTransfer|Preview")
+	UPROPERTY(EditAnywhere, Transient, Category = "WeightTransfer")
 	FTransform SourcePreviewOffset = FTransform::Identity;
 	
 private:
-	
 	UFUNCTION()
-	TArray<FName> GetLODsFunc() const;
+	TArray<FName> GetTargetLODsFunc() const;
 	UFUNCTION()
-	TArray<FName> GetSkinWeightProfilesFunc() const;
-
+	TArray<FName> GetTargetSkinWeightProfilesFunc() const;
 	UFUNCTION()
 	TArray<FName> GetSourceLODsFunc() const;
 	UFUNCTION()
 	TArray<FName> GetSourceSkinWeightProfilesFunc() const;
 };
 
+// this class wraps the all the components to enable selection on a single mesh in the skin weights tool
+// this allows us to make selections on multiple different meshes
+// NOTE: at some point we may want to do component selections on multiple meshes in any/all viewports
+// at which time this class should be centralized and renamed to UMeshSelector or something like that.
+// But there will need to be some sort of centralized facility to manage that and make sure it interacts nicely with other tools.
+UCLASS()
+class MESHMODELINGTOOLSEDITORONLYEXP_API UWeightToolMeshSelector : public UObject
+{
+	GENERATED_BODY()
+
+public:
+
+	// must be called during the Setup of the parent tool
+	void InitialSetup(
+		UWorld* InWorld,
+		UInteractiveTool* InParentTool,
+		FEditorViewportClient* InViewportClient,
+		TFunction<void()> OnSelectionChangedFunc);
+
+	// must be called AFTER InitialSetup, and any time the mesh is changed
+	// passing in a null preview mesh will disable the selector
+	void SetMesh(
+		UPreviewMesh* InMesh,
+		const FMeshDescription* InMeshDescription,
+		const FTransform3d& InMeshTransform);
+	
+	void UpdateAfterMeshDeformation();
+	
+	void Shutdown();
+
+	void SetIsEnabled(bool bIsEnabled);
+	void SetComponentSelectionMode(EComponentSelectionMode InMode);
+	void SetTransform(const FTransform3d& InTargetTransform);
+
+	// viewport 
+	void DrawHUD(FCanvas* Canvas, IToolsContextRenderAPI* RenderAPI);
+	void Render(IToolsContextRenderAPI* RenderAPI);
+
+	// get a list of currently selected vertices (converting edges and faces to vertices)
+	const TArray<int32>& GetSelectedVertices();
+	bool IsAnyComponentSelected() const;
+	void GetSelectedTriangles(TArray<int32>& OutTriangleIndices) const;
+
+	// edit selection
+	void GrowSelection() const;
+	void ShrinkSelection() const;
+	void FloodSelection() const;
+	void SelectBorder() const;
+
+	// get access to the selection mechanic
+	UPolygonSelectionMechanic* GetSelectionMechanic() { return PolygonSelectionMechanic; };
+
+private:
+
+	UPROPERTY()
+	TObjectPtr<UInteractiveTool> ParentTool;
+	UPROPERTY()
+	TObjectPtr<UWorld> World;
+	FEditorViewportClient* ViewportClient;
+	UPROPERTY()
+	TObjectPtr<UPreviewMesh> PreviewMesh;
+	FMeshDescription MeshDescription;
+	UPROPERTY()
+	TObjectPtr<UPolygonSelectionMechanic> PolygonSelectionMechanic;
+	
+	TUniquePtr<UE::Geometry::FDynamicMeshAABBTree3> MeshSpatial = nullptr;
+	TUniquePtr<UE::Geometry::FTriangleGroupTopology> SelectionTopology = nullptr;
+	
+	TArray<VertexIndex> SelectedVerticesInternal;
+};
+
+// this class wraps a source skeletal mesh used to transfer skin weights to the tool target mesh
+UCLASS()
+class MESHMODELINGTOOLSEDITORONLYEXP_API UWeightToolTransferManager : public UObject
+{
+	GENERATED_BODY()
+
+public:
+
+	// this must be called from within the parent tool's Setup() so that the selection mechanics are registered for capturing input
+	void InitialSetup(USkinWeightsPaintTool* InWeightTool, FEditorViewportClient* InViewportClient);
+	
+	// called when the tool is shutdown
+	void Shutdown();
+
+	// render the selection mechanism
+	void Render(IToolsContextRenderAPI* RenderAPI);
+	void DrawHUD(FCanvas* Canvas, IToolsContextRenderAPI* RenderAPI);
+
+	// update the mesh we are transferring from
+	void SetSourceMesh(USkeletalMesh* InSkeletalMesh = nullptr);
+
+	// change which mesh we are selecting (either source or target) and update visibilty
+	void UpdateSelectionAndVisibility() const;
+
+	// run the weight transfer
+	void TransferWeights();
+
+	// returns true if everything is setup and ready to transfer
+	bool CanTransferWeights() const;
+	
+	// gets the tool target for the source mesh
+	UToolTarget* GetTarget() const { return SourceTarget; }
+
+	// get the preview mesh for the source mesh
+	UPreviewMesh* GetPreviewMesh() const { return SourcePreviewMesh; };
+	
+	// get the mesh selector for the source mesh
+	UWeightToolMeshSelector* GetMeshSelector() const { return MeshSelector; };
+	
+	// called when tool settings are modified
+	void OnPropertyModified(const USkinWeightsPaintToolProperties* WeightToolProperties, const FProperty* ModifiedProperty);
+
+private:
+	
+	// actually run the weight transfer to copy weights from the source to the target
+	void TransferWeightsFromOtherMesh();
+
+	// actually run the weight transfer to copy weights from the source to the target
+	void TransferWeightsFromSameMeshAndLOD();
+
+	void ApplyTranferredWeightsAsTransaction(
+		const UE::Geometry::FDynamicMeshVertexSkinWeightsAttribute* TransferredSkinWeights,
+		const TArray<int32>& VertexSubset,
+		FDynamicMesh3& TargetMesh);
+	
+	USkinWeightsPaintToolProperties* GetToolProperties() const;
+	
+	UPROPERTY()
+	TObjectPtr<UPreviewMesh> SourcePreviewMesh = nullptr;
+	UPROPERTY()
+	TObjectPtr<USkeletalMesh> SourceSkeletalMesh = nullptr;
+	UPROPERTY()
+	TObjectPtr<UToolTarget> SourceTarget = nullptr;
+	UPROPERTY()
+	TObjectPtr<UWeightToolMeshSelector> MeshSelector;
+	TWeakObjectPtr<USkinWeightsPaintTool> WeightTool;
+};
+
 // An interactive tool for painting and editing skin weights.
 UCLASS()
 class MESHMODELINGTOOLSEDITORONLYEXP_API USkinWeightsPaintTool : public UDynamicMeshBrushTool, public ISkeletalMeshEditingInterface
@@ -507,8 +654,6 @@
 	void AverageWeights(const float Strength);
 	void NormalizeWeights();
 	void HammerWeights();
-	void TransferWeights();
-	bool CanTransferWeights() const;
 
 	// copy paste
 	void CopyWeights();
@@ -526,25 +671,21 @@
 
 	// toggle brush / selection mode
 	void ToggleEditingMode();
+	void UpdateComponentSelectionMode() const;
 
-	// edit selection
-	void SetComponentSelectionMode(EComponentSelectionMode InMode);
-	void GrowSelection() const;
-	void ShrinkSelection() const;
-	void FloodSelection() const;
+	// get access to the mesh selector for the main mesh
+	UWeightToolMeshSelector* GetMainMeshSelector();
+	// get access to the currently active mesh selector (may be on the transfer source mesh)
+	UWeightToolMeshSelector* GetActiveMeshSelector();
+	// does the main mesh have an active selection ("active" meaning the selection is currently being rendered in the view and is editable)
+	bool HasActiveSelectionOnMainMesh();
+	// select all vertices affected by the currently selected bone(s)
 	void SelectAffected() const;
-	void SelectBorder() const;
+	
 	// isolate selection
-	bool IsAnyComponentSelected() const;
 	bool IsSelectionIsolated() const;
 	void SetIsolateSelected(const bool bIsolateSelection);
 
-	// get a list of currently selected vertices (converting edges and faces to vertices)
-	const TArray<int32>& GetSelectedVertices();
-	bool HasSelectedVertices() const { return !SelectedVerticesInternal.IsEmpty(); };
-	void GetVerticesAffectedByBone(BoneIndex IndexOfBone, TSet<int32>& OutVertexIndices) const;
-	void GetSelectedTriangles(TArray<int32>& OutTriangleIndices) const;
-
 	// get the average weight value of each influence on the given vertices
 	void GetInfluences(const TArray<int32>& VertexIndices, TArray<BoneIndex>& OutBoneIndices);
 	// get the average weight value of a single bone on the given vertices
@@ -553,19 +694,33 @@
 	FName GetBoneNameFromIndex(BoneIndex InIndex) const;
 	// get the currently selected bone
 	BoneIndex GetCurrentBoneIndex() const;
+	// get a list of vertices affected by the given bone
+	void GetVerticesAffectedByBone(BoneIndex IndexOfBone, TSet<int32>& OutVertexIndices) const;
 
 	// toggle the display of weights on the preview mesh (if false, uses the normal skeletal mesh material)
 	void SetDisplayVertexColors(bool bShowVertexColors=true);
 	// set focus back to viewport so that hotkeys are immediately detected while hovering
 	void SetFocusInViewport() const;
 
+	// get the target manager (cached from Setup)
+	UToolTargetManager* GetTargetManager() const { return TargetManager.Get(); };
+
+	// allows outside systems to access the weight data
+	SkinPaintTool::FSkinToolWeights& GetWeights() { return Weights; };
+
+	// get access to the weight tranfer system
+	UWeightToolTransferManager* GetWeightTransferManager() const { return TransferManager; };
+
+	// get the viewport this tool is operating in
+	FEditorViewportClient* GetViewportClient() const;
+
 	// HOW TO EDIT WEIGHTS WITH UNDO/REDO:
 	//
 	// "Interactive" Edits:
 	// For multiple weight editing operations that need to be grouped into a single transaction, like dragging a slider or
 	// dragging a brush, you must call:
 	//  1. BeginChange()
-	//  2. ApplyWeightEditsToMeshMidChange() (this may be called multiple times)
+	//  2. ApplyWeightEditsWithoutTransaction() (this may be called multiple times)
 	//  2. EndChange()
 	// All the edits are stored into the "ActiveChange" and applied as a single transaction in EndChange().
 	// Deformations and vertex colors will be updated throughout the duration of the change.
@@ -585,9 +740,6 @@
 	DECLARE_MULTICAST_DELEGATE(FOnWeightsChanged);
 	FOnWeightsChanged OnWeightsChanged;
 
-	// gets the current source target
-	UToolTarget* GetSourceTarget() const { return SourceTarget; }
-
 protected:
 
 	virtual void ApplyStamp(const FBrushStampData& Stamp);
@@ -660,7 +812,7 @@
 	// isolate selection sub-meshes
 	UE::Geometry::FDynamicSubmesh3 PartialSubMesh;
 	UE::Geometry::FDynamicMesh3 FullDynamicMesh;
-	//
+	// convert to/from partial-isolated and full mesh vertex indices
 	int32 PartialToFullMeshVertexIndex(int32 PartialMeshVertexIndex) const;
 	int32 FullToPartialMeshVertexIndex(int32 FullMeshVertexIndex) const;
 
@@ -709,13 +861,9 @@
 	// ISkeletalMeshEditionInterface
 	virtual void HandleSkeletalMeshModified(const TArray<FName>& InBoneNames, const ESkeletalMeshNotifyType InNotifyType) override;
 
-	// polygon selection mechanic
+	// the selection system for the main mesh
 	UPROPERTY()
-	TObjectPtr<UPolygonSelectionMechanic> PolygonSelectionMechanic;
-	TUniquePtr<UE::Geometry::FDynamicMeshAABBTree3> MeshSpatial = nullptr;
-	TUniquePtr<UE::Geometry::FTriangleGroupTopology> SelectionTopology = nullptr;
-	void InitializeSelectionMechanic();
-	TArray<VertexIndex> SelectedVerticesInternal;
+	TObjectPtr<UWeightToolMeshSelector> MeshSelector;
 
 	// skin weight layer
 	void OnActiveLODChanged();
@@ -731,13 +879,9 @@
 	UPROPERTY()
 	TWeakObjectPtr<UToolTargetManager> TargetManager = nullptr;
 
-	// skin weights transfer properties
+	// manages transferring skin weights from a separate mesh
 	UPROPERTY()
-	TObjectPtr<UPreviewMesh> SourcePreviewMesh = nullptr;
-	UPROPERTY()
-	TObjectPtr<UToolTarget> SourceTarget = nullptr;
-
-	void ResetSourceForTransfer(USkeletalMesh* InSkeletalMesh = nullptr);
+	TObjectPtr<UWeightToolTransferManager> TransferManager = nullptr;
 	
 	// editor state to restore when exiting the paint tool
 	FString PreviewProfileToRestore;
diff --git a/Engine/Plugins/Runtime/MeshModelingToolset/Source/ModelingComponents/Private/Selection/MeshTopologySelectionMechanic.cpp b/Engine/Plugins/Runtime/MeshModelingToolset/Source/ModelingComponents/Private/Selection/MeshTopologySelectionMechanic.cpp
--- a/Engine/Plugins/Runtime/MeshModelingToolset/Source/ModelingComponents/Private/Selection/MeshTopologySelectionMechanic.cpp
+++ b/Engine/Plugins/Runtime/MeshModelingToolset/Source/ModelingComponents/Private/Selection/MeshTopologySelectionMechanic.cpp
@@ -130,6 +130,11 @@ void UMeshTopologySelectionMechanic::SetIsEnabled(bool bOn)
 	UpdateMarqueeEnabled();
 }
 
+void UMeshTopologySelectionMechanic::SetTransform(const FTransform3d& InTargetTransform)
+{
+	TargetTransform = InTargetTransform;
+}
+
 void UMeshTopologySelectionMechanic::SetMarqueeSelectionUpdateType(EMarqueeSelectionUpdateType InType)
 {
 	MarqueeSelectionUpdateType = InType;
@@ -162,6 +167,11 @@ TPair<FInputCapturePriority, FInputCapturePriority> UMeshTopologySelectionMechan
 
 void UMeshTopologySelectionMechanic::Render(IToolsContextRenderAPI* RenderAPI)
 {
+	if (!bIsEnabled)
+	{
+		return;
+	}
+	
 	MarqueeMechanic->Render(RenderAPI);
 
 	// Cache the view camera state so we can use for snapping/etc.
@@ -215,6 +225,11 @@ void UMeshTopologySelectionMechanic::Render(IToolsContextRenderAPI* RenderAPI)
 
 void UMeshTopologySelectionMechanic::DrawHUD(FCanvas* Canvas, IToolsContextRenderAPI* RenderAPI)
 {
+	if (!bIsEnabled)
+	{
+		return;
+	}
+	
 	MarqueeMechanic->DrawHUD(Canvas, RenderAPI);
 }
 
@@ -260,6 +275,11 @@ bool UMeshTopologySelectionMechanic::TopologyHitTest(const FRay& WorldRay, FHitR
 
 bool UMeshTopologySelectionMechanic::TopologyHitTest(const FRay& WorldRay, FHitResult& OutHit, FGroupTopologySelection& OutSelection, bool bUseOrthoSettings)
 {
+	if (!bIsEnabled)
+	{
+		return false;
+	}
+	
 	FRay3d LocalRay(TargetTransform.InverseTransformPosition((FVector3d)WorldRay.Origin),
 		TargetTransform.InverseTransformVector((FVector3d)WorldRay.Direction));
 	UE::Geometry::Normalize(LocalRay.Direction);
@@ -312,6 +332,11 @@ bool UMeshTopologySelectionMechanic::TopologyHitTest(const FRay& WorldRay, FHitR
 
 void UMeshTopologySelectionMechanic::HandleRectangleChanged(const FCameraRectangle& InRectangle)
 {
+	if (!bIsEnabled)
+	{
+		return;
+	}
+	
 	FGroupTopologySelection RectangleSelection;
 
 	TopoSelector->FindSelectedElement(PreDragTopoSelectorSettings, InRectangle, TargetTransform,
@@ -763,6 +788,36 @@ void UMeshTopologySelectionMechanic::FloodSelection()
 	ParentTool->GetToolManager()->EndUndoTransaction();
 }
 
+void UMeshTopologySelectionMechanic::ConvertSelectionToBorder()
+{
+	// find all border vertices
+	// a "border" vertex is one that has a 1-ring neighbor that is not in the selection set
+	const FGroupTopologySelection& CurrentSelection = GetActiveSelection();
+	FGroupTopologySelection Selection;
+	for (const int32 SelectedVertex : CurrentSelection.SelectedCornerIDs)
+	{
+		for (const int32 NeighborIndex : Mesh->VtxVerticesItr(SelectedVertex))
+		{
+			if (!CurrentSelection.SelectedCornerIDs.Contains(NeighborIndex))
+			{
+				Selection.SelectedCornerIDs.Add(SelectedVertex);
+			}
+		}
+	}
+
+	ParentTool->GetToolManager()->BeginUndoTransaction(LOCTEXT("BorderSelectionChange", "Select Border"));
+	BeginChange();
+	
+	// select vertices
+	PersistentSelection = Selection;
+
+	// broadcast and end transaction
+	SelectionTimestamp++;
+	OnSelectionChanged.Broadcast();
+	EndChangeAndEmitIfModified();
+	ParentTool->GetToolManager()->EndUndoTransaction();
+}
+
 FInputRayHit UMeshTopologySelectionMechanic::IsHitByClick(const FInputDeviceRay& ClickPos)
 {
 	if (!bIsEnabled)
diff --git a/Engine/Plugins/Runtime/MeshModelingToolset/Source/ModelingComponents/Public/Selection/MeshTopologySelectionMechanic.h b/Engine/Plugins/Runtime/MeshModelingToolset/Source/ModelingComponents/Public/Selection/MeshTopologySelectionMechanic.h
--- a/Engine/Plugins/Runtime/MeshModelingToolset/Source/ModelingComponents/Public/Selection/MeshTopologySelectionMechanic.h
+++ b/Engine/Plugins/Runtime/MeshModelingToolset/Source/ModelingComponents/Public/Selection/MeshTopologySelectionMechanic.h
@@ -163,6 +163,11 @@ public:
 	 */
 	void SetIsEnabled(bool bOn);
 
+	/**
+	 * Update the location of the mesh after Initialize()
+	 */
+	void SetTransform(const FTransform3d& InTargetTransform);
+
 	/**
 	 * Sets how/when the selection updates are handled.
 	 */
@@ -277,6 +282,7 @@ public:
 	void GrowSelection();
 	void ShrinkSelection();
 	void FloodSelection();
+	void ConvertSelectionToBorder();
 
 	/** 
 	 * @return true if the current selection is non-empty 
