LOCTEXT fix in SkinWeightsPaintTool.cpp

diff --git a/Engine/Plugins/Experimental/MeshModelingToolsetExp/Source/MeshModelingToolsEditorOnlyExp/Private/SkeletalMesh/SkinWeightsPaintTool.cpp b/Engine/Plugins/Experimental/MeshModelingToolsetExp/Source/MeshModelingToolsEditorOnlyExp/Private/SkeletalMesh/SkinWeightsPaintTool.cpp
--- a/Engine/Plugins/Experimental/MeshModelingToolsetExp/Source/MeshModelingToolsEditorOnlyExp/Private/SkeletalMesh/SkinWeightsPaintTool.cpp
+++ b/Engine/Plugins/Experimental/MeshModelingToolsetExp/Source/MeshModelingToolsEditorOnlyExp/Private/SkeletalMesh/SkinWeightsPaintTool.cpp
@@ -875,7 +875,7 @@ void UWeightToolTransferManager::TransferWeights()
 	// cannot transfer between same mesh/LOD/profile (identical weights)
 	if (bSameMesh && bSameLOD && bSameProfile)
 	{
-		const FText NotificationText = LOCTEXT("NoSourceMesh", "Cannot copy weights from the same mesh, LOD and profile. No weights were transferred.");
+		const FText NotificationText = LOCTEXT("IdenticalSourceAndTarget", "Cannot copy weights from the same mesh, LOD and profile. No weights were transferred.");
 		ShowEditorMessage(ELogVerbosity::Error, NotificationText);
 		return;
 	}
