Skin weight tool: fixed crash when undo'ing mix of isolated selections and vertex selections.
Added transaction support for isolated selection.

diff --git a/Engine/Plugins/Experimental/Animation/SkeletalMeshModelingTools/Source/SkeletalMeshModelingTools/Private/DetailCustomization/SkinWeightDetailCustomization.cpp b/Engine/Plugins/Experimental/Animation/SkeletalMeshModelingTools/Source/SkeletalMeshModelingTools/Private/DetailCustomization/SkinWeightDetailCustomization.cpp
--- a/Engine/Plugins/Experimental/Animation/SkeletalMeshModelingTools/Source/SkeletalMeshModelingTools/Private/DetailCustomization/SkinWeightDetailCustomization.cpp
+++ b/Engine/Plugins/Experimental/Animation/SkeletalMeshModelingTools/Source/SkeletalMeshModelingTools/Private/DetailCustomization/SkinWeightDetailCustomization.cpp
@@ -8,7 +8,6 @@
 #include "ModelingToolsEditorModeStyle.h"
 #include "Widgets/Input/SSegmentedControl.h"
 #include "SkeletalMesh/SkinWeightsPaintTool.h"
-#include "Selection/PolygonSelectionMechanic.h"
 #include "Widgets/Input/SSpinBox.h"
 #include "Widgets/Input/SButton.h"
 #include "Widgets/Input/SNumericEntryBox.h"
@@ -422,30 +421,30 @@ void FSkinWeightDetailCustomization::AddSelectionUI(IDetailLayoutBuilder& Detail
 					{
 						// isolated selection only available on main mesh (for now)
 						const bool bHasSelection = Tool->GetMainMeshSelector()->IsAnyComponentSelected();
-						const bool bAlreadyIsolatingSelection = Tool->IsSelectionIsolated();
+						const bool bAlreadyIsolatingSelection = Tool->GetSelectionIsolator()->IsSelectionIsolated();
 						return bHasSelection ||  bAlreadyIsolatingSelection;
 					})
 					.IsChecked_Lambda([this]()
 					{
-						const bool bAlreadyIsolatingSelection = Tool->IsSelectionIsolated();
+						const bool bAlreadyIsolatingSelection = Tool->GetSelectionIsolator()->IsSelectionIsolated();
 						return bAlreadyIsolatingSelection ? ECheckBoxState::Checked : ECheckBoxState::Unchecked;
 					})
 					.OnCheckStateChanged_Lambda([this](ECheckBoxState InCheckBoxState)
 					{
 						if (InCheckBoxState == ECheckBoxState::Checked)
 						{
-							Tool->SetIsolateSelected(true);	
+							Tool->GetSelectionIsolator()->IsolateSelectionAsTransaction();
 						}
 						else
 						{
-							Tool->SetIsolateSelected(false);	
+							Tool->GetSelectionIsolator()->UnIsolateSelectionAsTransaction();
 						}
 					})
 					[
 						SNew(STextBlock)
 						.Text_Lambda([this]()
 						{
-							if (Tool->IsSelectionIsolated())
+							if (Tool->GetSelectionIsolator()->IsSelectionIsolated())
 							{
 								return LOCTEXT("ShowAllButtonLabel", "Show Full Mesh");
 							}
diff --git a/Engine/Plugins/Experimental/MeshModelingToolsetExp/Source/MeshModelingToolsEditorOnlyExp/Private/SkeletalMesh/SkinWeightsPaintTool.cpp b/Engine/Plugins/Experimental/MeshModelingToolsetExp/Source/MeshModelingToolsEditorOnlyExp/Private/SkeletalMesh/SkinWeightsPaintTool.cpp
--- a/Engine/Plugins/Experimental/MeshModelingToolsetExp/Source/MeshModelingToolsEditorOnlyExp/Private/SkeletalMesh/SkinWeightsPaintTool.cpp
+++ b/Engine/Plugins/Experimental/MeshModelingToolsetExp/Source/MeshModelingToolsEditorOnlyExp/Private/SkeletalMesh/SkinWeightsPaintTool.cpp
@@ -711,7 +711,7 @@ void UWeightToolTransferManager::InitialSetup(USkinWeightsPaintTool* InWeightToo
 	WeightTool = InWeightTool;
 
 	// always reset back to target selection
-	GetToolProperties()->MeshSelectMode = EMeshTransferOption::Target;
+	WeightTool->GetWeightToolProperties()->MeshSelectMode = EMeshTransferOption::Target;
 
 	// create the mesh selector and run initial setup
 	// NOTE: currently this must happen inside the Setup of a UInteractiveTool so that input is routed to the selection mechanic
@@ -724,7 +724,7 @@ void UWeightToolTransferManager::SetSourceMesh(USkeletalMesh* InSkeletalMesh)
 {
 	SourceSkeletalMesh = InSkeletalMesh;
 
-	USkinWeightsPaintToolProperties* ToolProperties = GetToolProperties();
+	USkinWeightsPaintToolProperties* ToolProperties = WeightTool->GetWeightToolProperties();
 	
 	// reset to prepare for new mesh (or possibly no mesh)
 	{
@@ -783,7 +783,7 @@ void UWeightToolTransferManager::SetSourceMesh(USkeletalMesh* InSkeletalMesh)
 		Location.X -= TargetBounds.GetBoxExtrema(1).X;
 		Location.X -= 1.1 * SourceBounds.GetBoxExtrema(1).X;
 		Transform.SetLocation(Location);
-		GetToolProperties()->SourcePreviewOffset = Transform;
+		WeightTool->GetWeightToolProperties()->SourcePreviewOffset = Transform;
 		SourcePreviewMesh->SetTransform(Transform);
 	}
 
@@ -809,7 +809,7 @@ void UWeightToolTransferManager::SetSourceMesh(USkeletalMesh* InSkeletalMesh)
 
 void UWeightToolTransferManager::UpdateSelectionAndVisibility() const
 {
-	const USkinWeightsPaintToolProperties* ToolProperties = GetToolProperties();
+	const USkinWeightsPaintToolProperties* ToolProperties = WeightTool->GetWeightToolProperties();
 	const bool bHasSourceMesh = GetPreviewMesh() != nullptr;
 	const bool bWasSetToSelectSource = ToolProperties->MeshSelectMode == EMeshTransferOption::Source;
 	const bool bEnableSourceMeshSelector = bWasSetToSelectSource && bHasSourceMesh;
@@ -865,7 +865,7 @@ void UWeightToolTransferManager::TransferWeights()
 		return;
 	}
 
-	USkinWeightsPaintToolProperties* ToolProperties = GetToolProperties();
+	USkinWeightsPaintToolProperties* ToolProperties = WeightTool->GetWeightToolProperties();
 
 	const USkeletalMesh* TargetSkeletalMesh = GetSkeletalMeshComponent(WeightTool->GetTarget())->GetSkeletalMeshAsset();
 	const bool bSameMesh = SourceSkeletalMesh == TargetSkeletalMesh;
@@ -902,7 +902,7 @@ void UWeightToolTransferManager::TransferWeightsFromOtherMesh()
 		return;
 	}
 
-	USkinWeightsPaintToolProperties* ToolProperties = GetToolProperties();
+	USkinWeightsPaintToolProperties* ToolProperties = WeightTool->GetWeightToolProperties();
 
 	// get LOD IDs
 	const EMeshLODIdentifier TargetLODId = GetLODId(ToolProperties->ActiveLOD);
@@ -1029,7 +1029,7 @@ void UWeightToolTransferManager::TransferWeightsFromSameMeshAndLOD()
 	using UE::Geometry::FDynamicMeshVertexSkinWeightsAttribute;
 	using UE::AnimationCore::FBoneWeights;
 	
-	const USkinWeightsPaintToolProperties* ToolProperties = GetToolProperties();
+	const USkinWeightsPaintToolProperties* ToolProperties = WeightTool->GetWeightToolProperties();
 	
 	// get target dynamic mesh
 	const EMeshLODIdentifier TargetLODId = GetLODId(ToolProperties->ActiveLOD);
@@ -1178,21 +1178,6 @@ void UWeightToolTransferManager::ApplyTranferredWeightsAsTransaction(
 	ShowEditorMessage(ELogVerbosity::Log, NotificationText);
 }
 
-USkinWeightsPaintToolProperties* UWeightToolTransferManager::GetToolProperties() const
-{
-	TArray<UObject*> AllToolProperties = WeightTool->GetToolProperties(false/*bEnabledOnly*/);
-
-	for (UObject* PropertyObject : AllToolProperties)
-	{
-		if (PropertyObject->IsA(USkinWeightsPaintToolProperties::StaticClass()))
-		{
-			return CastChecked<USkinWeightsPaintToolProperties>(PropertyObject);
-		}
-	}
-
-	return nullptr;
-}
-
 FName USkinWeightsPaintToolProperties::GetActiveSkinWeightProfile() const
 {
 	return bShowNewProfileName ? NewSkinWeightProfile : ActiveSkinWeightProfile;
@@ -2184,6 +2169,10 @@ void USkinWeightsPaintTool::Setup()
 	TransferManager = NewObject<UWeightToolTransferManager>();
 	TransferManager->InitialSetup(this, GetViewportClient());
 
+	// create the isolated selection manager
+	SelectionIsolator = NewObject<UWeightToolSelectionIsolator>();
+	SelectionIsolator->InitialSetup(this);
+
 	// setup selection for the main mesh
 	MeshSelector = NewObject<UWeightToolMeshSelector>(this);
 	auto OnSelectionChangedLambda = [this](){OnSelectionChanged.Broadcast();};
@@ -2191,7 +2180,7 @@ void USkinWeightsPaintTool::Setup()
 	UpdateComponentSelectionMode();
 
 	// run all initialization for mesh/weights
-	PostEditMeshInitialization(Component, *PreviewMesh->GetMesh(), *EditedMesh);
+	UpdateCurrentlyEditedMesh(Component, *PreviewMesh->GetMesh(), *EditedMesh);
 
 	// bind the skeletal mesh editor context
 	if (EditorContext.IsValid())
@@ -2401,10 +2390,9 @@ FInputRayHit USkinWeightsPaintTool::CanBeginClickDragSequence(const FInputDevice
 
 void USkinWeightsPaintTool::OnTick(float DeltaTime)
 {
-	if (bPendingUpdateFromPartialMesh)
+	if (SelectionIsolator)
 	{
-		FinishIsolatedSelection();
-		bPendingUpdateFromPartialMesh = false;
+		SelectionIsolator->OnTick(DeltaTime);
 	}
 	
 	if (bStampPending)
@@ -2435,7 +2423,7 @@ void USkinWeightsPaintTool::OnTick(float DeltaTime)
 	Weights.Deformer.UpdateVertexDeformation(this, Weights.Deformer.Component->GetComponentSpaceTransforms());
 }
 
-void USkinWeightsPaintTool::PostEditMeshInitialization(
+void USkinWeightsPaintTool::UpdateCurrentlyEditedMesh(
 	const USkeletalMeshComponent* InComponent,
 	const FDynamicMesh3& InDynamicMesh,
 	const FMeshDescription& InMeshDescription)
@@ -2527,6 +2515,29 @@ void USkinWeightsPaintToolProperties::SetBrushMode(EWeightEditOperation InBrushM
 	WeightTool->SetFocusInViewport();
 }
 
+void FIsolateSelectionChange::Apply(UObject* Object)
+{
+	USkinWeightsPaintTool* Tool = Cast<USkinWeightsPaintTool>(Object);
+	if (Tool)
+	{
+		Tool->GetSelectionIsolator()->SetIsolatedTriangles(IsolatedTrianglesAfter);
+	}	
+}
+
+void FIsolateSelectionChange::Revert(UObject* Object)
+{
+	USkinWeightsPaintTool* Tool = Cast<USkinWeightsPaintTool>(Object);
+	if (Tool)
+	{
+		Tool->GetSelectionIsolator()->SetIsolatedTriangles(IsolatedTrianglesBefore);
+	}
+}
+
+FString FIsolateSelectionChange::ToString() const
+{
+	return FToolCommandChange::ToString();
+}
+
 bool USkinWeightsPaintTool::HitTest(const FRay& Ray, FHitResult& OutHit)
 {
 	const bool bIsInBrushMode = WeightToolProperties->EditingMode == EWeightEditMode::Brush;
@@ -3120,7 +3131,7 @@ void USkinWeightsPaintTool::ApplyWeightEditsWithoutTransaction(const FMultiBoneW
 	// update list of vertices needing updated vertex colors
 	WeightEdits.AddEditedVerticesToSet(VerticesToUpdateColor);
 	// store changes in the transaction buffer, but keep it open
-	auto PartialToFullVertexConverter = [this](int32 InVertexIndex){ return PartialToFullMeshVertexIndex(InVertexIndex); };
+	auto PartialToFullVertexConverter = [this](int32 InVertexIndex){ return GetSelectionIsolator()->PartialToFullMeshVertexIndex(InVertexIndex); };
 	ActiveChange->StoreMultipleWeightEdits(WeightEdits, PartialToFullVertexConverter);
 }
 
@@ -3162,7 +3173,10 @@ void USkinWeightsPaintTool::SetFocusInViewport() const
 void USkinWeightsPaintTool::OnShutdown(EToolShutdownType ShutdownType)
 {
 	// shutdown must be performed on full mesh, so end isolated selection
-	FinishIsolatedSelection();
+	if (SelectionIsolator)
+	{
+		SelectionIsolator->RestoreFullMesh();
+	}
 	
 	// save tool properties
 	WeightToolProperties->SaveProperties(this);
@@ -3227,6 +3241,11 @@ FEditorViewportClient* USkinWeightsPaintTool::GetViewportClient() const
 	return ViewportClient;
 }
 
+USkinWeightsPaintToolProperties* USkinWeightsPaintTool::GetWeightToolProperties() const
+{
+	return WeightToolProperties;
+}
+
 void USkinWeightsPaintTool::BeginChange()
 {
 	const EMeshLODIdentifier LOD = GetLODId(WeightToolProperties->ActiveLOD);
@@ -3254,7 +3273,7 @@ void USkinWeightsPaintTool::ExternalUpdateWeights(const int32 BoneIndex, const T
 	{
 		// weights are always stored in transactions as full mesh indices
 		// if we are in an isolated selection, we must convert them to the partial mesh for them to be applied
-		const int32 VertexID = FullToPartialMeshVertexIndex(Pair.Key);
+		const int32 VertexID = SelectionIsolator->FullToPartialMeshVertexIndex(Pair.Key);
 		const float Weight = Pair.Value;
 		Weights.SetWeightOfBoneOnVertex(BoneIndex, VertexID, Weight, Weights.CurrentWeights);
 		Weights.SetWeightOfBoneOnVertex(BoneIndex, VertexID, Weight, Weights.PreChangeWeights);
@@ -4118,9 +4137,9 @@ void USkinWeightsPaintTool::OnActiveLODChanged()
 		return;
 	}
 
-	if (IsSelectionIsolated())
+	if (SelectionIsolator->IsSelectionIsolated())
 	{
-		FinishIsolatedSelection();
+		SelectionIsolator->RestoreFullMesh();
 	}
 
 	// apply previous changes
@@ -4138,7 +4157,7 @@ void USkinWeightsPaintTool::OnActiveLODChanged()
 
 	// reinitialize all mesh data structures
 	const FDynamicMesh3 DynamicMesh = UE::ToolTarget::GetDynamicMeshCopy(Target, Params);
-	PostEditMeshInitialization(Component, DynamicMesh, *EditedMesh);
+	UpdateCurrentlyEditedMesh(Component, DynamicMesh, *EditedMesh);
 }
 
 void USkinWeightsPaintTool::OnActiveSkinWeightProfileChanged()
@@ -4151,9 +4170,9 @@ void USkinWeightsPaintTool::OnActiveSkinWeightProfileChanged()
 
 	WeightToolProperties->bShowNewProfileName = WeightToolProperties->ActiveSkinWeightProfile == CreateNewName();
 
-	if (IsSelectionIsolated())
+	if (SelectionIsolator->IsSelectionIsolated())
 	{
-		FinishIsolatedSelection();
+		SelectionIsolator->RestoreFullMesh();
 	}
 
 	if (WeightToolProperties->bShowNewProfileName)
@@ -4276,102 +4295,105 @@ UWeightToolMeshSelector* USkinWeightsPaintTool::GetActiveMeshSelector()
 	return MeshSelector;
 }
 
-bool USkinWeightsPaintTool::HasActiveSelectionOnMainMesh()
+void UWeightToolSelectionIsolator::InitialSetup(USkinWeightsPaintTool* InTool)
 {
-	const bool bSelectingTargetMesh = WeightToolProperties->MeshSelectMode == EMeshTransferOption::Target;
-	const bool bTargetMeshHasSelection = GetMainMeshSelector()->IsAnyComponentSelected();
-	return bSelectingTargetMesh && bTargetMeshHasSelection;
+	WeightTool = InTool;
 }
 
-void USkinWeightsPaintTool::SelectAffected() const
+void UWeightToolSelectionIsolator::OnTick(float DeltaTime)
 {
-	UPolygonSelectionMechanic* SelectionMechanic = MeshSelector->GetSelectionMechanic();
-	if (!ensure(SelectionMechanic))
+	if (bRestoreFullMeshOnNextTick)
 	{
-		return;
+		// this is queued to run on Tick() because modifying the mesh from other threads can cause the tool's Render() to be out of sync
+		RestoreFullMesh();
+		bRestoreFullMeshOnNextTick = false;
 	}
-	
-	GetToolManager()->BeginUndoTransaction(LOCTEXT("AffectedSelectionChange", "Select Affected"));
-	SelectionMechanic->BeginChange();
-	
-	// get all vertices affected by all selected bones
-	TSet<int32> AffectedVertices;
-	for (const BoneIndex SelectedBone : SelectedBoneIndices)
-	{
-		GetVerticesAffectedByBone(SelectedBone, AffectedVertices);
-	}
-	
-	// create selection set
-	FGroupTopologySelection Selection;
+}
 
-	// optionally add/remove/replace selection based on modifier key state
-	const FGroupTopologySelection& CurrentSelection = SelectionMechanic->GetActiveSelection();
-	if (bShiftToggle)
-	{
-		// ADD to selection
-		Selection.SelectedCornerIDs.Append(CurrentSelection.SelectedCornerIDs);
-		Selection.SelectedCornerIDs.Append(AffectedVertices);
-	}
-	else if (bCtrlToggle)
-	{
-		// REMOVE from selection
-		Selection.SelectedCornerIDs = CurrentSelection.SelectedCornerIDs.Difference(AffectedVertices);
-	}
-	else
-	{
-		// REPLACE selection
-		Selection.SelectedCornerIDs.Append(AffectedVertices);
-	}
+bool UWeightToolSelectionIsolator::IsSelectionIsolated() const
+{
+	return PartialMeshDescription.IsValid();
+}
+
+void UWeightToolSelectionIsolator::IsolateSelectionAsTransaction()
+{
+	const FText& TransactionLabel = LOCTEXT("IsolateSelectTransaction", "Isolate Selection");
+
+	TUniquePtr<FIsolateSelectionChange> ActiveChange = MakeUnique<FIsolateSelectionChange>();
+	ActiveChange->IsolatedTrianglesBefore = GetIsolatedTriangles();
+	WeightTool->GetMainMeshSelector()->GetSelectedTriangles(ActiveChange->IsolatedTrianglesAfter);
+
+	SetIsolatedTriangles(ActiveChange->IsolatedTrianglesAfter);
 	
-	// select vertices
-	constexpr bool bBroadcast = true;
-	SelectionMechanic->SetSelection(Selection, bBroadcast);
-	SelectionMechanic->EndChangeAndEmitIfModified();
-	GetToolManager()->EndUndoTransaction();
+	UInteractiveToolManager* ToolManager = WeightTool->GetToolManager();
+	ToolManager->BeginUndoTransaction(TransactionLabel);
+	ToolManager->EmitObjectChange(WeightTool, MoveTemp(ActiveChange), TransactionLabel);
+	ToolManager->EndUndoTransaction();
 }
 
-bool USkinWeightsPaintTool::IsSelectionIsolated() const
+void UWeightToolSelectionIsolator::UnIsolateSelectionAsTransaction()
 {
-	return PartialMeshDescription.IsValid();
+	const FText& TransactionLabel =  LOCTEXT("ShowAllTransaction", "Show All");
+
+	TUniquePtr<FIsolateSelectionChange> ActiveChange = MakeUnique<FIsolateSelectionChange>();
+	ActiveChange->IsolatedTrianglesBefore = GetIsolatedTriangles();
+	ActiveChange->IsolatedTrianglesAfter = {};
+
+	SetIsolatedTriangles(ActiveChange->IsolatedTrianglesAfter);
+	
+	UInteractiveToolManager* ToolManager = WeightTool->GetToolManager();
+	ToolManager->BeginUndoTransaction(TransactionLabel);
+	ToolManager->EmitObjectChange(WeightTool, MoveTemp(ActiveChange), TransactionLabel);
+	ToolManager->EndUndoTransaction();
 }
 
-void USkinWeightsPaintTool::SetIsolateSelected(const bool bIsolateSelection)
+void UWeightToolSelectionIsolator::SetIsolatedTriangles(const TArray<int32>& TrianglesToIsolate)
 {
 	// if we are turning off an isolated selection, we must queue the Tick() to update the full mesh
-	if (!bIsolateSelection && PartialMeshDescription.IsValid())
+	if (TrianglesToIsolate.IsEmpty() && PartialMeshDescription.IsValid())
 	{
-		bPendingUpdateFromPartialMesh = true;
+		bRestoreFullMeshOnNextTick = true;
 		return;
 	}
 	
 	if (PartialMeshDescription.IsValid())
 	{
-		ensure(false); // should be reset to null
+		// should be reset to null
+		// NOTE: we do not support isolating the selection twice in a row
+		// ie. you must isolate the selection, then un-isolate it before isolating it again
+		ensure(false); 
 		return;
 	}
 	
-	UPolygonSelectionMechanic* SelectionMechanic = MeshSelector->GetSelectionMechanic();
+	UPolygonSelectionMechanic* SelectionMechanic = WeightTool->GetMainMeshSelector()->GetSelectionMechanic();
 	if (!ensure(SelectionMechanic))
 	{
 		return;
 	}
 
-	const USkeletalMeshComponent* SkeletalMeshComponent = GetSkeletalMeshComponent(Target);
+	const USkeletalMeshComponent* SkeletalMeshComponent = GetSkeletalMeshComponent(WeightTool->GetTarget());
 	if (!ensure(SkeletalMeshComponent))
 	{
 		return;
 	}
 
+	const FMeshDescription* EditedMesh = WeightTool->GetCurrentlyEditedMeshDescription();
 	if (!ensure(EditedMesh))
 	{
 		return;
 	}
 
+	// record the triangles we are isolating
+	CurrentlyIsolatedTriangles = TrianglesToIsolate;
+
+	// get the weights
+	FSkinToolWeights& Weights = WeightTool->GetWeights();
+
 	// apply previous changes
-	Weights.ApplyCurrentWeightsToMeshDescription(EditedMesh);
+	Weights.ApplyCurrentWeightsToMeshDescription(WeightTool->GetCurrentlyEditedMeshDescription());
 	
 	// put into ref pose, BEFORE copying the mesh, so that submesh deformer initializes with vertices in ref pose
-	Weights.Deformer.SetToRefPose(this);
+	Weights.Deformer.SetToRefPose(WeightTool);
 
 	// store selection to be restored
 	IsolatedSelectionToRestoreVertices.Reset();
@@ -4390,13 +4412,7 @@ void USkinWeightsPaintTool::SetIsolateSelected(const bool bIsolateSelection)
 	Converter.Convert(EditedMesh, DynamicMesh);
 	FullDynamicMesh = MoveTemp(DynamicMesh);
 
-	// create a submesh from the selected triangles (or triangles connected to selected vertices/edges)
-	TArray<int32> TrianglesToIsolate;
-	MeshSelector->GetSelectedTriangles(TrianglesToIsolate);
-	if (TrianglesToIsolate.IsEmpty())
-	{
-		return;
-	}
+	// create a submesh from the given triangles
 	PartialSubMesh = UE::Geometry::FDynamicSubmesh3(&FullDynamicMesh, TrianglesToIsolate);
 
 	// create mesh description for sub-mesh
@@ -4411,12 +4427,12 @@ void USkinWeightsPaintTool::SetIsolateSelected(const bool bIsolateSelection)
 	DnyToDescConverter.Convert(&PartialSubMesh.GetSubmesh(), *PartialMeshDescription, bCopyTangents);
 	
 	// reinitialize all mesh data structures
-	PostEditMeshInitialization(SkeletalMeshComponent, PartialSubMesh.GetSubmesh(), *PartialMeshDescription.Get());
+	WeightTool->UpdateCurrentlyEditedMesh(SkeletalMeshComponent, PartialSubMesh.GetSubmesh(), *PartialMeshDescription.Get());
 }
 
-void USkinWeightsPaintTool::FinishIsolatedSelection()
+void UWeightToolSelectionIsolator::RestoreFullMesh()
 {
-	const USkeletalMeshComponent* SkeletalMeshComponent = GetSkeletalMeshComponent(Target);
+	const USkeletalMeshComponent* SkeletalMeshComponent = GetSkeletalMeshComponent(WeightTool->GetTarget());
 	if (!SkeletalMeshComponent)
 	{
 		// this can happen at shutdown
@@ -4429,15 +4445,19 @@ void USkinWeightsPaintTool::FinishIsolatedSelection()
 		return;
 	}
 
+	FSkinToolWeights& Weights = WeightTool->GetWeights();
+
 	// apply partial mesh weights to partial mesh description
 	Weights.ApplyCurrentWeightsToMeshDescription(PartialMeshDescription.Get());
 
 	// reinitialize with full mesh
-	PostEditMeshInitialization(SkeletalMeshComponent, FullDynamicMesh, *EditedMesh);
+	FMeshDescription* CurrentlyEditedMesh = WeightTool->GetCurrentlyEditedMeshDescription();
+	WeightTool->UpdateCurrentlyEditedMesh(SkeletalMeshComponent, FullDynamicMesh, *CurrentlyEditedMesh);
 
 	// copy the remapped weights back to the full mesh
 	const FSkeletalMeshConstAttributes MeshAttribs(*PartialMeshDescription.Get());
-	const FSkinWeightsVertexAttributesConstRef AllVertexWeights = MeshAttribs.GetVertexSkinWeights(WeightToolProperties->GetActiveSkinWeightProfile());
+	const FName ActiveProfile = WeightTool->GetWeightToolProperties()->GetActiveSkinWeightProfile();
+	const FSkinWeightsVertexAttributesConstRef AllVertexWeights = MeshAttribs.GetVertexSkinWeights(ActiveProfile);
 	const int32 NumVerticesInPartialMesh = PartialMeshDescription.Get()->Vertices().Num();
 	for (int32 VertexIndexPartial = 0; VertexIndexPartial < NumVerticesInPartialMesh; VertexIndexPartial++)
 	{
@@ -4461,10 +4481,10 @@ void USkinWeightsPaintTool::FinishIsolatedSelection()
 	// sync both weight buffers
 	Weights.PreChangeWeights = Weights.CurrentWeights;
 	// apply full mesh weights to full mesh description
-	Weights.ApplyCurrentWeightsToMeshDescription(EditedMesh);
+	Weights.ApplyCurrentWeightsToMeshDescription(CurrentlyEditedMesh);
 
-	// restore selection (allows for easily adjusting crop)
-	if (UPolygonSelectionMechanic* SelectionMechanic = MeshSelector->GetSelectionMechanic())
+	// restore selection (a convenience for the user that allows for easily adjusting the isolation by going back/forth as needed)
+	if (UPolygonSelectionMechanic* SelectionMechanic = WeightTool->GetMainMeshSelector()->GetSelectionMechanic())
 	{
 		SelectionMechanic->SetSelection_AsTriangleTopology(IsolatedSelectionToRestoreVertices);
 		SelectionMechanic->SetSelection_AsTriangleTopology(IsolatedSelectionToRestoreEdges);
@@ -4472,9 +4492,10 @@ void USkinWeightsPaintTool::FinishIsolatedSelection()
 	}
 
 	PartialMeshDescription = nullptr;
+	CurrentlyIsolatedTriangles.Reset();
 }
 
-int32 USkinWeightsPaintTool::PartialToFullMeshVertexIndex(int32 PartialMeshVertexIndex) const
+int32 UWeightToolSelectionIsolator::PartialToFullMeshVertexIndex(int32 PartialMeshVertexIndex) const
 {
 	if (!PartialMeshDescription)
 	{
@@ -4484,7 +4505,7 @@ int32 USkinWeightsPaintTool::PartialToFullMeshVertexIndex(int32 PartialMeshVerte
 	return PartialSubMesh.MapVertexToBaseMesh(PartialMeshVertexIndex);
 }
 
-int32 USkinWeightsPaintTool::FullToPartialMeshVertexIndex(int32 FullMeshVertexIndex) const
+int32 UWeightToolSelectionIsolator::FullToPartialMeshVertexIndex(int32 FullMeshVertexIndex) const
 {
 	if (!PartialMeshDescription)
 	{
@@ -4494,6 +4515,60 @@ int32 USkinWeightsPaintTool::FullToPartialMeshVertexIndex(int32 FullMeshVertexIn
 	return PartialSubMesh.MapVertexToSubmesh(FullMeshVertexIndex);
 }
 
+bool USkinWeightsPaintTool::HasActiveSelectionOnMainMesh()
+{
+	const bool bSelectingTargetMesh = WeightToolProperties->MeshSelectMode == EMeshTransferOption::Target;
+	const bool bTargetMeshHasSelection = GetMainMeshSelector()->IsAnyComponentSelected();
+	return bSelectingTargetMesh && bTargetMeshHasSelection;
+}
+
+void USkinWeightsPaintTool::SelectAffected() const
+{
+	UPolygonSelectionMechanic* SelectionMechanic = MeshSelector->GetSelectionMechanic();
+	if (!ensure(SelectionMechanic))
+	{
+		return;
+	}
+	
+	GetToolManager()->BeginUndoTransaction(LOCTEXT("AffectedSelectionChange", "Select Affected"));
+	SelectionMechanic->BeginChange();
+	
+	// get all vertices affected by all selected bones
+	TSet<int32> AffectedVertices;
+	for (const BoneIndex SelectedBone : SelectedBoneIndices)
+	{
+		GetVerticesAffectedByBone(SelectedBone, AffectedVertices);
+	}
+	
+	// create selection set
+	FGroupTopologySelection Selection;
+
+	// optionally add/remove/replace selection based on modifier key state
+	const FGroupTopologySelection& CurrentSelection = SelectionMechanic->GetActiveSelection();
+	if (bShiftToggle)
+	{
+		// ADD to selection
+		Selection.SelectedCornerIDs.Append(CurrentSelection.SelectedCornerIDs);
+		Selection.SelectedCornerIDs.Append(AffectedVertices);
+	}
+	else if (bCtrlToggle)
+	{
+		// REMOVE from selection
+		Selection.SelectedCornerIDs = CurrentSelection.SelectedCornerIDs.Difference(AffectedVertices);
+	}
+	else
+	{
+		// REPLACE selection
+		Selection.SelectedCornerIDs = MoveTemp(AffectedVertices);
+	}
+	
+	// select vertices
+	constexpr bool bBroadcast = true;
+	SelectionMechanic->SetSelection(Selection, bBroadcast);
+	SelectionMechanic->EndChangeAndEmitIfModified();
+	GetToolManager()->EndUndoTransaction();
+}
+
 void USkinWeightsPaintTool::GetVerticesAffectedByBone(BoneIndex IndexOfBone, TSet<int32>& OutVertexIndices) const
 {
 	VertexIndex VertexID = 0;
diff --git a/Engine/Plugins/Experimental/MeshModelingToolsetExp/Source/MeshModelingToolsEditorOnlyExp/Public/SkeletalMesh/SkinWeightsPaintTool.h b/Engine/Plugins/Experimental/MeshModelingToolsetExp/Source/MeshModelingToolsEditorOnlyExp/Public/SkeletalMesh/SkinWeightsPaintTool.h
--- a/Engine/Plugins/Experimental/MeshModelingToolsetExp/Source/MeshModelingToolsEditorOnlyExp/Public/SkeletalMesh/SkinWeightsPaintTool.h
+++ b/Engine/Plugins/Experimental/MeshModelingToolsetExp/Source/MeshModelingToolsEditorOnlyExp/Public/SkeletalMesh/SkinWeightsPaintTool.h
@@ -595,8 +595,6 @@ private:
 		const TArray<int32>& VertexSubset,
 		FDynamicMesh3& TargetMesh);
 	
-	USkinWeightsPaintToolProperties* GetToolProperties() const;
-	
 	UPROPERTY()
 	TObjectPtr<UPreviewMesh> SourcePreviewMesh = nullptr;
 	UPROPERTY()
@@ -608,6 +606,71 @@ private:
 	TWeakObjectPtr<USkinWeightsPaintTool> WeightTool;
 };
 
+// this class wraps all the data needed to isolate a selection of a mesh while editing skin weights
+UCLASS()
+class MESHMODELINGTOOLSEDITORONLYEXP_API UWeightToolSelectionIsolator : public UObject
+{
+	GENERATED_BODY()
+
+public:
+
+	// call during tool Setup()
+	void InitialSetup(USkinWeightsPaintTool* InTool);
+
+	// call every tick to apply deferred changes to mesh
+	void OnTick(float DeltaTime);
+
+	// returns true if any triangles are currently isolated
+	bool IsSelectionIsolated() const;
+
+	// isolate the current selection
+	void IsolateSelectionAsTransaction();
+
+	// unisolate the current selection
+	void UnIsolateSelectionAsTransaction();
+
+	// isolate the array of triangles
+	void SetIsolatedTriangles(const TArray<int32>& TrianglesToIsolate);
+	
+	// restores the whole mesh
+	void RestoreFullMesh();
+	
+	// get the current triangles that are isolated
+	const TArray<int32>& GetIsolatedTriangles() { return CurrentlyIsolatedTriangles; };
+
+	// convert to/from partial-isolated and full mesh vertex indices
+	int32 PartialToFullMeshVertexIndex(int32 PartialMeshVertexIndex) const;
+	int32 FullToPartialMeshVertexIndex(int32 FullMeshVertexIndex) const;
+
+private:
+	
+	UPROPERTY()
+	TObjectPtr<USkinWeightsPaintTool> WeightTool;
+		
+	// when selection is isolated, we hide the full mesh and show a submesh
+	// when islated selection is unhidden, we remap all changes from the submesh back to the full mesh
+	TSharedPtr<FMeshDescription> PartialMeshDescription = nullptr; // only non-null during isolated selection
+	TArray<int32> CurrentlyIsolatedTriangles;
+	UE::Geometry::FGeometrySelection IsolatedSelectionToRestoreVertices;
+	UE::Geometry::FGeometrySelection IsolatedSelectionToRestoreEdges;
+	UE::Geometry::FGeometrySelection IsolatedSelectionToRestoreFaces;
+	bool bRestoreFullMeshOnNextTick = false;
+	// isolate selection sub-meshes
+	UE::Geometry::FDynamicSubmesh3 PartialSubMesh;
+	UE::Geometry::FDynamicMesh3 FullDynamicMesh;
+};
+
+class MESHMODELINGTOOLSEDITORONLYEXP_API FIsolateSelectionChange : public FToolCommandChange
+{
+public:
+	TArray<int32> IsolatedTrianglesBefore;
+	TArray<int32> IsolatedTrianglesAfter;
+
+	virtual void Apply(UObject* Object) override;
+	virtual void Revert(UObject* Object) override;
+	virtual FString ToString() const override;
+};
+
 // An interactive tool for painting and editing skin weights.
 UCLASS()
 class MESHMODELINGTOOLSEDITORONLYEXP_API USkinWeightsPaintTool : public UDynamicMeshBrushTool, public ISkeletalMeshEditingInterface
@@ -681,10 +744,6 @@ public:
 	bool HasActiveSelectionOnMainMesh();
 	// select all vertices affected by the currently selected bone(s)
 	void SelectAffected() const;
-	
-	// isolate selection
-	bool IsSelectionIsolated() const;
-	void SetIsolateSelected(const bool bIsolateSelection);
 
 	// get the average weight value of each influence on the given vertices
 	void GetInfluences(const TArray<int32>& VertexIndices, TArray<BoneIndex>& OutBoneIndices);
@@ -713,6 +772,15 @@ public:
 
 	// get the viewport this tool is operating in
 	FEditorViewportClient* GetViewportClient() const;
+	
+	// get access to the selection isolation system
+	UWeightToolSelectionIsolator* GetSelectionIsolator() const { return SelectionIsolator; };
+
+	// get the tool properties
+	USkinWeightsPaintToolProperties* GetWeightToolProperties() const;
+
+	// get access to the mesh description for the mesh being edited
+	FMeshDescription* GetCurrentlyEditedMeshDescription() const { return EditedMesh; };
 
 	// HOW TO EDIT WEIGHTS WITH UNDO/REDO:
 	//
@@ -732,6 +800,12 @@ public:
 	// It will Begin/End the change and create a transaction for it.
 	void ApplyWeightEditsAsTransaction(const SkinPaintTool::FMultiBoneWeightEdits& WeightEdits, const FText& TransactionLabel);
 
+	// call this whenever the target mesh is modified
+	void UpdateCurrentlyEditedMesh(
+		const USkeletalMeshComponent* Component,
+		const FDynamicMesh3& InDynamicMesh,
+		const FMeshDescription& InMeshDescription);
+
 	// called whenever the selection is modified
 	DECLARE_MULTICAST_DELEGATE(FOnSelectionChanged);
 	FOnSelectionChanged OnSelectionChanged;
@@ -746,11 +820,6 @@ protected:
 	void OnShutdown(EToolShutdownType ShutdownType) override;
 	void OnTick(float DeltaTime) override;
 
-	void PostEditMeshInitialization(
-		const USkeletalMeshComponent* Component,
-		const FDynamicMesh3& InDynamicMesh,
-		const FMeshDescription& InMeshDescription);
-
 	void CleanMesh() const;
 
 	// stamp
@@ -801,20 +870,6 @@ protected:
 	// the currently edited mesh descriptions
 	mutable TMap<EMeshLODIdentifier, FMeshDescription> EditedMeshes;
 	FMeshDescription* EditedMesh = nullptr;
-	// when selection is isolated, we hide the full mesh and show a submesh
-	// when islated selection is unhidden, we remap all changes from the submesh back to the full mesh
-	TSharedPtr<FMeshDescription> PartialMeshDescription = nullptr; // only non-null during isolated selection
-	UE::Geometry::FGeometrySelection IsolatedSelectionToRestoreVertices;
-	UE::Geometry::FGeometrySelection IsolatedSelectionToRestoreEdges;
-	UE::Geometry::FGeometrySelection IsolatedSelectionToRestoreFaces;
-	bool bPendingUpdateFromPartialMesh = false;
-	void FinishIsolatedSelection();
-	// isolate selection sub-meshes
-	UE::Geometry::FDynamicSubmesh3 PartialSubMesh;
-	UE::Geometry::FDynamicMesh3 FullDynamicMesh;
-	// convert to/from partial-isolated and full mesh vertex indices
-	int32 PartialToFullMeshVertexIndex(int32 PartialMeshVertexIndex) const;
-	int32 FullToPartialMeshVertexIndex(int32 FullMeshVertexIndex) const;
 
 	// storage of vertex weights per bone 
 	SkinPaintTool::FSkinToolWeights Weights;
@@ -882,6 +937,10 @@ protected:
 	// manages transferring skin weights from a separate mesh
 	UPROPERTY()
 	TObjectPtr<UWeightToolTransferManager> TransferManager = nullptr;
+
+	// manages isolating a selection of the mesh
+	UPROPERTY()
+	TObjectPtr<UWeightToolSelectionIsolator> SelectionIsolator = nullptr;
 	
 	// editor state to restore when exiting the paint tool
 	FString PreviewProfileToRestore;
