Fix holes in coarse LODs:
- Stop using the non-coarse LODs as a starting point for coarse LOD generation
- Make the simplifier respect boundary edges
- Allow low-tri-count non-coarse LODs propagate over coarse LODs when needed

diff --git a/Engine/Plugins/Experimental/MeshModelingToolsetExp/Source/GeometryProcessingAdapters/Private/CombineMeshInstancesImpl.cpp b/Engine/Plugins/Experimental/MeshModelingToolsetExp/Source/GeometryProcessingAdapters/Private/CombineMeshInstancesImpl.cpp
--- a/Engine/Plugins/Experimental/MeshModelingToolsetExp/Source/GeometryProcessingAdapters/Private/CombineMeshInstancesImpl.cpp
+++ b/Engine/Plugins/Experimental/MeshModelingToolsetExp/Source/GeometryProcessingAdapters/Private/CombineMeshInstancesImpl.cpp
@@ -2749,8 +2749,9 @@ void DoSimplifyMesh(
 		}
 	}
 
-	// do these flags matter here since we are not flipping??
-	EEdgeRefineFlags MeshBoundaryConstraints = EEdgeRefineFlags::NoFlip;
+	// If we allow boundary collapse, this can introduce visible holes in the simplified result
+	constexpr bool bAllowBoundaryCollapse = false;
+	EEdgeRefineFlags MeshBoundaryConstraints = bAllowBoundaryCollapse ? EEdgeRefineFlags::NoFlip : EEdgeRefineFlags::SplitsOnly;
 	EEdgeRefineFlags GroupBorderConstraints = EEdgeRefineFlags::NoConstraint;
 	EEdgeRefineFlags MaterialBorderConstraints = EEdgeRefineFlags::NoConstraint;
 
@@ -3302,14 +3303,12 @@ void ProcessCombinedLODChain(
 		int32 MaxTriCount = GetCoarseLODMaxTriCount(FirstVoxWrappedIndex);
 		double SimplifyTolerance = CoarseLODBaseTolerance;
 
-		// for very simple parts it can be the case that the last approximate LOD is
-		// lower tri-count than the first coarse approximation. In that case just use it.
+		// Note for very simple parts it can be the case that the last approximate LOD is
+		// lower tri-count than the first coarse approximation. To handle such cases,
+		// we rely on the BuildOutputSubAssembly to propagate simpler LODs down the chain.
+		// (We don't use the non-coarse LODs as a starting point because they tend not
+		// to simplify down as well for subsequent coarse LODs)
 		UE::Tasks::Wait(PendingRemoveHiddenTasks);
-		int32 PrevLODTriCount = MeshLODs[FirstVoxWrappedIndex - 1].Mesh.TriangleCount();
-		if (PrevLODTriCount < InitialCoarseApproximation.TriangleCount() && PrevLODTriCount < MaxTriCount)
-		{
-			InitialCoarseApproximation = MeshLODs[FirstVoxWrappedIndex-1].Mesh;
-		}
 
 		// Current state of InitialCoarseApproximation is our initial voxel LOD. To ensure
 		// that voxel LODs have compatible UVs (to allow baking), we compute UVs on
@@ -3465,16 +3464,13 @@ static void BuildOutputSubAssembly(
 	OutputSubAssembly.SubAssemblyID = SubAssemblyID;
 
 	// collect output meshes
-	// can't replace voxel LODs if we are generating UVs for them!
-	// (no way to communicate this upwards...)
-	//int MaxReplaceLOD = MeshLODs.Num();
-	int MaxReplaceLOD = FirstVoxWrappedIndex;
+	int MaxReplaceLOD = MeshLODs.Num();
 	for (int32 LODLevel = 0; LODLevel < MeshLODs.Num(); ++LODLevel)
 	{
 		FDynamicMesh3 LODMesh = MoveTemp(MeshLODs[LODLevel].Mesh);
 
 		// If we ended up larger than the mesh in the previous LOD, we should use that instead!
-		// This can happen particular with VoxWrap LODs [TODO: should we therefore not be excluding voxwrap lods w/ MaxReplaceLOD, above?]
+		// This can happen particular with VoxWrap LODs
 		if (LODLevel > 0 && LODLevel < MaxReplaceLOD)
 		{
 			if (LODMesh.TriangleCount() > OutputSubAssembly.MeshLODs.Last().TriangleCount())
