MeshModelingTools: Fix the gizmo numerical panel getting stuck dragging if right clicked.

diff --git a/Engine/Plugins/Experimental/MeshModelingToolsetExp/Source/ModelingUI/Private/ModelingWidgets/SDraggableBox.cpp b/Engine/Plugins/Experimental/MeshModelingToolsetExp/Source/ModelingUI/Private/ModelingWidgets/SDraggableBox.cpp
--- a/Engine/Plugins/Experimental/MeshModelingToolsetExp/Source/ModelingUI/Private/ModelingWidgets/SDraggableBox.cpp
+++ b/Engine/Plugins/Experimental/MeshModelingToolsetExp/Source/ModelingUI/Private/ModelingWidgets/SDraggableBox.cpp
@@ -99,6 +99,11 @@ void SDraggableBox::Construct(const FArguments& InArgs)
 
 FReply SDraggableBox::OnMouseButtonDown(const FGeometry& MyGeometry, const FPointerEvent& MouseEvent)
 {
+	if (MouseEvent.GetEffectingButton() != EKeys::LeftMouseButton)
+	{
+		return SCompoundWidget::OnMouseButtonDown(MyGeometry, MouseEvent);
+	}
+
 	// Need to remember where within the box we grabbed. We do this here instead of OnDragDetected because 
 	// our mouse can potentially travel some distance before OnDragDetected fires.
 	ScreenSpaceOffsetOfGrab = MouseEvent.GetScreenSpacePosition() - MyGeometry.GetAbsolutePosition();
