Skin weight tool: Fixes bug, mesh selection renders incorrectly when switching out of mesh mode.

diff --git a/Engine/Plugins/Experimental/Animation/SkeletalMeshModelingTools/Source/SkeletalMeshModelingTools/Private/DetailCustomization/SkinWeightDetailCustomization.cpp b/Engine/Plugins/Experimental/Animation/SkeletalMeshModelingTools/Source/SkeletalMeshModelingTools/Private/DetailCustomization/SkinWeightDetailCustomization.cpp
--- a/Engine/Plugins/Experimental/Animation/SkeletalMeshModelingTools/Source/SkeletalMeshModelingTools/Private/DetailCustomization/SkinWeightDetailCustomization.cpp
+++ b/Engine/Plugins/Experimental/Animation/SkeletalMeshModelingTools/Source/SkeletalMeshModelingTools/Private/DetailCustomization/SkinWeightDetailCustomization.cpp
@@ -1254,7 +1254,7 @@ void FSkinWeightDetailCustomization::AddTransferUI(IDetailLayoutBuilder& DetailB
 			.OnValueChanged_Lambda([this](EMeshTransferOption Mode)
 			{
 				ToolSettings->MeshSelectMode = Mode;
-				ToolSettings->WeightTool->GetWeightTransferManager()->UpdateSelectionAndVisibility();
+				ToolSettings->WeightTool->UpdateSelectorState();
 			})
 			+ SSegmentedControl<EMeshTransferOption>::Slot(EMeshTransferOption::Source)
 			.Text(LOCTEXT("SourceMode", "Source"))
diff --git a/Engine/Plugins/Experimental/MeshModelingToolsetExp/Source/MeshModelingToolsEditorOnlyExp/Private/SkeletalMesh/SkinWeightsPaintTool.cpp b/Engine/Plugins/Experimental/MeshModelingToolsetExp/Source/MeshModelingToolsEditorOnlyExp/Private/SkeletalMesh/SkinWeightsPaintTool.cpp
--- a/Engine/Plugins/Experimental/MeshModelingToolsetExp/Source/MeshModelingToolsEditorOnlyExp/Private/SkeletalMesh/SkinWeightsPaintTool.cpp
+++ b/Engine/Plugins/Experimental/MeshModelingToolsetExp/Source/MeshModelingToolsEditorOnlyExp/Private/SkeletalMesh/SkinWeightsPaintTool.cpp
@@ -750,7 +750,7 @@ void UWeightToolTransferManager::SetSourceMesh(USkeletalMesh* InSkeletalMesh)
 	{
 		ToolProperties->SourceSkinWeightProfile = FSkeletalMeshAttributesShared::DefaultSkinWeightProfileName;
 		ToolProperties->SourceLOD = "LOD0";
-		UpdateSelectionAndVisibility();
+		WeightTool->UpdateSelectorState();
 		return;
 	}
 	
@@ -803,24 +803,7 @@ void UWeightToolTransferManager::SetSourceMesh(USkeletalMesh* InSkeletalMesh)
 	MeshSelector->SetIsEnabled(ToolProperties->EditingMode == EWeightEditMode::Mesh);
 	MeshSelector->SetComponentSelectionMode(ToolProperties->ComponentSelectionMode);
 
-	UpdateSelectionAndVisibility();
-}
-
-void UWeightToolTransferManager::UpdateSelectionAndVisibility() const
-{
-	const USkinWeightsPaintToolProperties* ToolProperties = WeightTool->GetWeightToolProperties();
-	const bool bHasSourceMesh = GetPreviewMesh() != nullptr;
-	const bool bWasSetToSelectSource = ToolProperties->MeshSelectMode == EMeshTransferOption::Source;
-	const bool bEnableSourceMeshSelector = bWasSetToSelectSource && bHasSourceMesh;
-	
-	if (MeshSelector)
-	{
-		MeshSelector->SetIsEnabled(bEnableSourceMeshSelector);
-		MeshSelector->SetComponentSelectionMode(ToolProperties->ComponentSelectionMode);
-	}
-
-	WeightTool->GetMainMeshSelector()->SetIsEnabled(!bEnableSourceMeshSelector);
-	WeightTool->GetMainMeshSelector()->SetComponentSelectionMode(ToolProperties->ComponentSelectionMode);
+	WeightTool->UpdateSelectorState();
 }
 
 void UWeightToolTransferManager::Shutdown()
@@ -2176,7 +2159,6 @@ void USkinWeightsPaintTool::Setup()
 	MeshSelector = NewObject<UWeightToolMeshSelector>(this);
 	auto OnSelectionChangedLambda = [this](){OnSelectionChanged.Broadcast();};
 	MeshSelector->InitialSetup(TargetWorld.Get(), this, GetViewportClient(), OnSelectionChangedLambda);
-	UpdateComponentSelectionMode();
 
 	// run all initialization for mesh/weights
 	UpdateCurrentlyEditedMesh(Component, *PreviewMesh->GetMesh(), *EditedMesh);
@@ -2482,7 +2464,7 @@ void USkinWeightsPaintToolProperties::SetComponentMode(EComponentSelectionMode I
 {
 	ComponentSelectionMode = InComponentMode;
 	
-	WeightTool->UpdateComponentSelectionMode();
+	WeightTool->UpdateSelectorState();
 	WeightTool->SetFocusInViewport();
 }
 
@@ -4241,13 +4223,8 @@ void USkinWeightsPaintTool::ToggleEditingMode()
 	// toggle brush mode
 	SetBrushEnabled(WeightToolProperties->EditingMode == EWeightEditMode::Brush);
 
-	// toggle mesh mode
-	const bool bIsMeshEditing = WeightToolProperties->EditingMode == EWeightEditMode::Mesh;
-	MeshSelector->SetIsEnabled(bIsMeshEditing);
-	if (UWeightToolMeshSelector* SourceMeshSelector = TransferManager->GetMeshSelector())
-	{
-		SourceMeshSelector->SetIsEnabled(bIsMeshEditing);
-	}
+	// toggle which mesh we're selecting and what components (vert/edge/face)
+	UpdateSelectorState();
 
 	// toggle bone select mode
 	// this mode is set to be compatible with the 
@@ -4266,17 +4243,26 @@ void USkinWeightsPaintTool::ToggleEditingMode()
 	SetFocusInViewport();
 }
 
-void USkinWeightsPaintTool::UpdateComponentSelectionMode() const
+void USkinWeightsPaintTool::UpdateSelectorState() const
 {
-	if (MeshSelector)
-	{
-		MeshSelector->SetComponentSelectionMode(WeightToolProperties->ComponentSelectionMode);
-	}
+	const USkinWeightsPaintToolProperties* ToolProperties = GetWeightToolProperties();
+	
+	const bool bIsMeshEditing = ToolProperties->EditingMode == EWeightEditMode::Mesh;
+	const bool bHasSourceMesh = TransferManager->GetPreviewMesh() != nullptr;
+	const bool bWasSetToSelectSource = ToolProperties->MeshSelectMode == EMeshTransferOption::Source;
+	const bool bEnableSourceMeshSelector = bIsMeshEditing && bWasSetToSelectSource && bHasSourceMesh;
+	const bool bEnableTargetMeshSelector = bIsMeshEditing && !bWasSetToSelectSource;
+	
+	// source mesh selector
+	UWeightToolMeshSelector* SourceMeshSelector = TransferManager->GetMeshSelector();
+	SourceMeshSelector->SetIsEnabled(bEnableSourceMeshSelector);
+	
+	// main mesh selector
+	MeshSelector->SetIsEnabled(bEnableTargetMeshSelector);
 
-	if (UWeightToolMeshSelector* SourceMeshSelector = TransferManager->GetMeshSelector())
-	{
-		SourceMeshSelector->SetComponentSelectionMode(WeightToolProperties->ComponentSelectionMode);
-	}
+	// update component mode
+	MeshSelector->SetComponentSelectionMode(WeightToolProperties->ComponentSelectionMode);
+	SourceMeshSelector->SetComponentSelectionMode(WeightToolProperties->ComponentSelectionMode);
 }
 
 UWeightToolMeshSelector* USkinWeightsPaintTool::GetMainMeshSelector()
diff --git a/Engine/Plugins/Experimental/MeshModelingToolsetExp/Source/MeshModelingToolsEditorOnlyExp/Public/SkeletalMesh/SkinWeightsPaintTool.h b/Engine/Plugins/Experimental/MeshModelingToolsetExp/Source/MeshModelingToolsEditorOnlyExp/Public/SkeletalMesh/SkinWeightsPaintTool.h
--- a/Engine/Plugins/Experimental/MeshModelingToolsetExp/Source/MeshModelingToolsEditorOnlyExp/Public/SkeletalMesh/SkinWeightsPaintTool.h
+++ b/Engine/Plugins/Experimental/MeshModelingToolsetExp/Source/MeshModelingToolsEditorOnlyExp/Public/SkeletalMesh/SkinWeightsPaintTool.h
@@ -561,9 +561,6 @@ public:
 	// update the mesh we are transferring from
 	void SetSourceMesh(USkeletalMesh* InSkeletalMesh = nullptr);
 
-	// change which mesh we are selecting (either source or target) and update visibilty
-	void UpdateSelectionAndVisibility() const;
-
 	// run the weight transfer
 	void TransferWeights();
 
@@ -734,7 +731,8 @@ public:
 
 	// toggle brush / selection mode
 	void ToggleEditingMode();
-	void UpdateComponentSelectionMode() const;
+	// update the state of the mesh selectors
+	void UpdateSelectorState() const;
 
 	// get access to the mesh selector for the main mesh
 	UWeightToolMeshSelector* GetMainMeshSelector();
