Skin weight tool: Fixes crash stemming from CleanMesh operation due to invalid non-manifold vertex mapping
Also fixes issue where material assignment can be scrambled after exiting the tool.

#swarm https://p4-swarm.epicgames.net/reviews/36490505

diff --git a/Engine/Plugins/Experimental/MeshModelingToolsetExp/Source/MeshModelingToolsEditorOnlyExp/Private/SkeletalMesh/SkinWeightsPaintTool.cpp b/Engine/Plugins/Experimental/MeshModelingToolsetExp/Source/MeshModelingToolsEditorOnlyExp/Private/SkeletalMesh/SkinWeightsPaintTool.cpp
--- a/Engine/Plugins/Experimental/MeshModelingToolsetExp/Source/MeshModelingToolsEditorOnlyExp/Private/SkeletalMesh/SkinWeightsPaintTool.cpp
+++ b/Engine/Plugins/Experimental/MeshModelingToolsetExp/Source/MeshModelingToolsEditorOnlyExp/Private/SkeletalMesh/SkinWeightsPaintTool.cpp
@@ -1674,13 +1674,13 @@ void FSkinToolWeights::CreateWeightEditForVertex(
 void FSkinToolWeights::ApplyCurrentWeightsToMeshDescription(FMeshDescription* MeshDescription)
 {
 	FSkeletalMeshAttributes MeshAttribs(*MeshDescription);
-	FSkinWeightsVertexAttributesRef VertexSkinWeights = MeshAttribs.GetVertexSkinWeights(Profile);
+	FSkinWeightsVertexAttributesRef VertexWeightAttrs = MeshAttribs.GetVertexSkinWeights(Profile);
 	
 	UE::AnimationCore::FBoneWeightsSettings Settings;
 	Settings.SetNormalizeType(UE::AnimationCore::EBoneWeightNormalizeType::None);
 
-	TArray<UE::AnimationCore::FBoneWeight> SourceBoneWeights;
-	SourceBoneWeights.Reserve(UE::AnimationCore::MaxInlineBoneWeightCount);
+	TArray<UE::AnimationCore::FBoneWeight> BoneWeightsToApply;
+	BoneWeightsToApply.Reserve(UE::AnimationCore::MaxInlineBoneWeightCount);
 
 	const int32 NumVertices = MeshDescription->Vertices().Num();
 	if (!ensure(CurrentWeights.Num() == NumVertices))
@@ -1691,7 +1691,7 @@ void FSkinToolWeights::ApplyCurrentWeightsToMeshDescription(FMeshDescription* Me
 	
 	for (int32 VertexIndex = 0; VertexIndex < NumVertices; VertexIndex++)
 	{
-		SourceBoneWeights.Reset();
+		BoneWeightsToApply.Reset();
 
 		const VertexWeights& VertexWeights = CurrentWeights[VertexIndex];
 		for (const FVertexBoneWeight& SingleBoneWeight : VertexWeights)
@@ -1700,10 +1700,10 @@ void FSkinToolWeights::ApplyCurrentWeightsToMeshDescription(FMeshDescription* Me
 			{
 				continue;
 			}
-			SourceBoneWeights.Add(UE::AnimationCore::FBoneWeight(SingleBoneWeight.BoneID, SingleBoneWeight.Weight));
+			BoneWeightsToApply.Add(UE::AnimationCore::FBoneWeight(SingleBoneWeight.BoneID, SingleBoneWeight.Weight));
 		}
 
-		VertexSkinWeights.Set(FVertexID(VertexIndex), UE::AnimationCore::FBoneWeights::Create(SourceBoneWeights, Settings));
+		VertexWeightAttrs.Set(FVertexID(VertexIndex), UE::AnimationCore::FBoneWeights::Create(BoneWeightsToApply, Settings));
 	}
 }
 
@@ -2098,21 +2098,19 @@ void USkinWeightsPaintTool::Setup()
 	
 	UDynamicMeshBrushTool::Setup();
 
-	const IPrimitiveComponentBackedTarget* TargetComponent = Cast<IPrimitiveComponentBackedTarget>(Target);
-	check(TargetComponent);
-	const USkeletalMeshComponent* Component = Cast<USkeletalMeshComponent>(TargetComponent->GetOwnerComponent());
+	const USkeletalMeshComponent* Component = GetSkeletalMeshComponent(Target);
 	check(Component && Component->GetSkeletalMeshAsset())
 
-	// prepare mesh for skin editing
-	CleanMesh();
-
-	// create a mesh description for editing (this must be done before calling UpdateBonePositionInfos) 
+	// create a mesh description for editing (this must be done before calling UpdateBonePositionInfos)
+	
 	bool bSupportsLODs = false;
 	const EMeshLODIdentifier DefaultLOD = UE::ToolTarget::GetTargetMeshDescriptionLOD(Target, bSupportsLODs);
 	ensure(bSupportsLODs);
-
-	EditedMesh = &EditedMeshes.Emplace(DefaultLOD);
-	*EditedMesh = *UE::ToolTarget::GetMeshDescription(Target);
+	const FGetMeshParameters Params(true, DefaultLOD);
+	// copy the mesh description from the tool target to edit it
+	const FMeshDescription* ToolTargetMeshDescription = UE::ToolTarget::GetMeshDescription(Target, Params);
+	EditedMeshes.Emplace(DefaultLOD, FCleanedEditMesh(*PreviewMesh->GetMesh(), *ToolTargetMeshDescription));
+	CurrentlyEditedLOD = DefaultLOD;
 
 	// create a custom set of properties inheriting from the base tool properties
 	WeightToolProperties = NewObject<USkinWeightsPaintToolProperties>(this);
@@ -2161,7 +2159,7 @@ void USkinWeightsPaintTool::Setup()
 	MeshSelector->InitialSetup(TargetWorld.Get(), this, GetViewportClient(), OnSelectionChangedLambda);
 
 	// run all initialization for mesh/weights
-	UpdateCurrentlyEditedMesh(Component, *PreviewMesh->GetMesh(), *EditedMesh);
+	UpdateCurrentlyEditedMesh(Component, *GetCurrentlyEditedDynamicMesh(), *GetCurrentlyEditedMeshDescription());
 
 	// bind the skeletal mesh editor context
 	if (EditorContext.IsValid())
@@ -2409,7 +2407,7 @@ void USkinWeightsPaintTool::UpdateCurrentlyEditedMesh(
 	const FDynamicMesh3& InDynamicMesh,
 	const FMeshDescription& InMeshDescription)
 {
-	// update the preview mesh
+	// update the preview mesh in the viewport
 	PreviewMesh->ReplaceMesh(InDynamicMesh);
 	PreviewMesh->EditMesh([](FDynamicMesh3& Mesh)
 	{
@@ -2444,22 +2442,6 @@ void USkinWeightsPaintTool::UpdateCurrentlyEditedMesh(
 	MirrorData.SetNeedsReinitialized();
 }
 
-void USkinWeightsPaintTool::CleanMesh() const
-{
-	if (PreviewMesh->GetMesh()->HasUnusedVertices())
-	{
-		// orphaned vertices wreak havoc on our selection tools
-		PreviewMesh->EditMesh([](FDynamicMesh3& Mesh)
-		{
-			Mesh.RemoveUnusedVertices();
-			Mesh.CompactInPlace();
-		});
-	
-		IDynamicMeshCommitter* DynamicMeshCommitter = Cast<IDynamicMeshCommitter>(Target);
-		DynamicMeshCommitter->CommitDynamicMesh(*PreviewMesh->GetMesh());
-	}
-}
-
 void USkinWeightsPaintToolProperties::SetComponentMode(EComponentSelectionMode InComponentMode)
 {
 	ComponentSelectionMode = InComponentMode;
@@ -3171,16 +3153,20 @@ void USkinWeightsPaintTool::OnShutdown(EToolShutdownType ShutdownType)
 	// apply changes to asset
 	if (ShutdownType == EToolShutdownType::Accept)
 	{
-		// apply the weights to the mesh description
-		Weights.ApplyCurrentWeightsToMeshDescription(EditedMesh);
+		// profile to edit
+		const FName ActiveProfile = WeightToolProperties->GetActiveSkinWeightProfile();
+	
+		// apply the currently edited weights to the mesh description
+		Weights.ApplyCurrentWeightsToMeshDescription(GetCurrentlyEditedMeshDescription());
 
 		// this block bakes the modified DynamicMeshComponent back into the StaticMeshComponent inside an undo transaction
 		GetToolManager()->BeginUndoTransaction(LOCTEXT("SkinWeightsPaintTool", "Paint Skin Weights"));
 		static constexpr bool bHaveTargetLOD = true;
-		for (auto& [LOD, MeshDescription] : EditedMeshes)
+		for (auto& [LOD, CleanedMesh] : EditedMeshes)
 		{
+			CleanedMesh.CopyWeightsToOriginalMesh(ActiveProfile);
 			const FCommitMeshParameters CommitParams(bHaveTargetLOD, LOD);
-			UE::ToolTarget::CommitMeshDescriptionUpdate(Target, &MeshDescription, nullptr, CommitParams);	
+			UE::ToolTarget::CommitMeshDescriptionUpdate(Target, &CleanedMesh.GetOriginalMeshDescription(), nullptr, CommitParams);	
 		}
 		GetToolManager()->EndUndoTransaction();
 	}
@@ -3227,6 +3213,36 @@ USkinWeightsPaintToolProperties* USkinWeightsPaintTool::GetWeightToolProperties(
 	return WeightToolProperties;
 }
 
+FCleanedEditMesh* USkinWeightsPaintTool::GetCurrentCleanedEditMesh() const
+{
+	if (EditedMeshes.Contains(CurrentlyEditedLOD))
+	{
+		return &EditedMeshes[CurrentlyEditedLOD];
+	}
+
+	return nullptr;
+}
+
+FMeshDescription* USkinWeightsPaintTool::GetCurrentlyEditedMeshDescription() const
+{
+	if (FCleanedEditMesh* CleanedEditMesh = GetCurrentCleanedEditMesh())
+	{
+		return &CleanedEditMesh->GetEditableMeshDescription();
+	}
+	
+	return nullptr;
+}
+
+FDynamicMesh3* USkinWeightsPaintTool::GetCurrentlyEditedDynamicMesh() const
+{
+	if (FCleanedEditMesh* CleanedEditMesh = GetCurrentCleanedEditMesh())
+	{
+		return &CleanedEditMesh->GetEditableMesh();
+	}
+	
+	return nullptr;
+}
+
 void USkinWeightsPaintTool::BeginChange()
 {
 	const EMeshLODIdentifier LOD = GetLODId(WeightToolProperties->ActiveLOD);
@@ -4124,21 +4140,21 @@ void USkinWeightsPaintTool::OnActiveLODChanged()
 	}
 
 	// apply previous changes
-	Weights.ApplyCurrentWeightsToMeshDescription(EditedMesh);
+	Weights.ApplyCurrentWeightsToMeshDescription(GetCurrentlyEditedMeshDescription());
 
-	// update EditedMesh using the new LOD
+	// update current mesh using the new LOD
 	const EMeshLODIdentifier LODId = GetLODId(WeightToolProperties->ActiveLOD);
 	const FGetMeshParameters Params(true, LODId);
-	EditedMesh = EditedMeshes.Find(LODId);
-	if (!EditedMesh)
+	FCleanedEditMesh* CleanedEditMesh = EditedMeshes.Find(LODId);
+	if (!CleanedEditMesh)
 	{
-		EditedMesh = &EditedMeshes.Emplace(LODId);
-		*EditedMesh = *UE::ToolTarget::GetMeshDescription(Target, Params);
+		const FDynamicMesh3 DynamicMesh = UE::ToolTarget::GetDynamicMeshCopy(Target, Params);
+		const FMeshDescription* MeshDescription = UE::ToolTarget::GetMeshDescription(Target, Params);
+		CleanedEditMesh = &EditedMeshes.Add(LODId, FCleanedEditMesh(DynamicMesh, *MeshDescription));
 	}
 
 	// reinitialize all mesh data structures
-	const FDynamicMesh3 DynamicMesh = UE::ToolTarget::GetDynamicMeshCopy(Target, Params);
-	UpdateCurrentlyEditedMesh(Component, DynamicMesh, *EditedMesh);
+	UpdateCurrentlyEditedMesh(Component, CleanedEditMesh->GetEditableMesh(), CleanedEditMesh->GetEditableMeshDescription());
 }
 
 void USkinWeightsPaintTool::OnActiveSkinWeightProfileChanged()
@@ -4160,7 +4176,7 @@ void USkinWeightsPaintTool::OnActiveSkinWeightProfileChanged()
 	{
 		if (!IsProfileValid(WeightToolProperties->NewSkinWeightProfile))
 		{
-			GetOrCreateSkinWeightsAttribute(*EditedMesh, WeightToolProperties->NewSkinWeightProfile);
+			GetOrCreateSkinWeightsAttribute(*GetCurrentlyEditedMeshDescription(), WeightToolProperties->NewSkinWeightProfile);
 		} 
 	}
 	
@@ -4176,12 +4192,12 @@ void USkinWeightsPaintTool::OnActiveSkinWeightProfileChanged()
 	}
 
 	// apply previous changes
-	Weights.ApplyCurrentWeightsToMeshDescription(EditedMesh);
+	Weights.ApplyCurrentWeightsToMeshDescription(GetCurrentlyEditedMeshDescription());
 
 	// re-init Weights with new skin profile
 	Weights = FSkinToolWeights();
 	Weights.Profile = WeightToolProperties->GetActiveSkinWeightProfile();
-	Weights.InitializeSkinWeights(SkeletalMeshComponent, EditedMesh);
+	Weights.InitializeSkinWeights(SkeletalMeshComponent, GetCurrentlyEditedMeshDescription());
 	bVertexColorsNeedUpdated = true;
 }
 
@@ -4189,7 +4205,7 @@ void USkinWeightsPaintTool::OnNewSkinWeightProfileChanged()
 {
 	if (WeightToolProperties->bShowNewProfileName && WeightToolProperties->NewSkinWeightProfile != Weights.Profile)
 	{
-		const bool bRenamed = RenameSkinWeightsAttribute(*EditedMesh, Weights.Profile, WeightToolProperties->NewSkinWeightProfile);
+		const bool bRenamed = RenameSkinWeightsAttribute(*GetCurrentlyEditedMeshDescription(), Weights.Profile, WeightToolProperties->NewSkinWeightProfile);
 		if (ensure(bRenamed))
 		{
 			Weights.Profile = WeightToolProperties->NewSkinWeightProfile;
@@ -4206,7 +4222,7 @@ bool USkinWeightsPaintTool::IsProfileValid(const FName InProfileName) const
 	}
 
 	// check current MeshDescription
-	const FSkeletalMeshConstAttributes MeshAttribs(*EditedMesh);
+	const FSkeletalMeshConstAttributes MeshAttribs(*GetCurrentlyEditedMeshDescription());
 	const TArray<FName> MeshDescProfiles = MeshAttribs.GetSkinWeightProfileNames();
 	const bool bHasProfile = MeshDescProfiles.ContainsByPredicate([InProfileName](const FName Name)
 	{
@@ -4500,6 +4516,120 @@ int32 UWeightToolSelectionIsolator::FullToPartialMeshVertexIndex(int32 FullMeshV
 	return PartialSubMesh.MapVertexToSubmesh(FullMeshVertexIndex);
 }
 
+FCleanedEditMesh::FCleanedEditMesh(
+	const FDynamicMesh3& InDynamicMesh,
+	const FMeshDescription& InMeshDescription) : OriginalDynamicMesh(InDynamicMesh), OriginalMeshDescription(InMeshDescription)
+{
+	// store copy of the original mesh data
+	OriginalDynamicMesh = InDynamicMesh;
+	OriginalMeshDescription = InMeshDescription;
+	
+	// if the mesh doesn't have any unused vertices, then we can skip making a cleaned duplicate and just use it directly
+	// (the getters will return the original mesh)
+	if (!OriginalDynamicMesh.HasUnusedVertices())
+	{
+		return;
+	}
+	
+	// create a submesh from all the triangles in the mesh
+	// NOTE: this is where the cleaning happens. this will leave out unused vertices because they don't belong to any triangle
+	TArray<int32> AllTris;
+	AllTris.Reserve(OriginalDynamicMesh.TriangleCount());
+	for (const int32 TriID : OriginalDynamicMesh.TriangleIndicesItr())
+	{
+		AllTris.Add(TriID);
+	}
+	CleanedSubMesh = UE::Geometry::FDynamicSubmesh3(&OriginalDynamicMesh, AllTris);
+
+	// remap the non-manifold mapping from full mesh to partial mesh
+	UE::Geometry::FNonManifoldMappingSupport OriginalNonManifoldMapping(OriginalDynamicMesh);
+	TArray<int32> SubMeshVertexToNonManifoldVertexIDMap;
+	const FDynamicMesh3& SubMeshDynamicMesh = CleanedSubMesh.GetSubmesh();
+	for (int32 SubMeshVertexID=0; SubMeshVertexID<SubMeshDynamicMesh.VertexCount(); ++SubMeshVertexID)
+	{
+		int32 BaseMeshVertexID = CleanedSubMesh.MapVertexToBaseMesh(SubMeshVertexID);
+		int32 SrcBaseMeshVertexID = OriginalNonManifoldMapping.GetOriginalNonManifoldVertexID(BaseMeshVertexID);
+		int32 SrcSubMeshVertexID = CleanedSubMesh.MapVertexToSubmesh(SrcBaseMeshVertexID);
+		if (SrcSubMeshVertexID == INDEX_NONE)
+		{
+			SrcSubMeshVertexID = SubMeshVertexID;
+		}
+		SubMeshVertexToNonManifoldVertexIDMap.Add(SrcSubMeshVertexID);
+	}
+
+	// replace the non-manifold vertex map in the cleaned mesh attributes
+	UE::Geometry::FNonManifoldMappingSupport CleanedNonManifoldMapping(CleanedSubMesh.GetSubmesh());
+	CleanedNonManifoldMapping.AttachNonManifoldVertexMappingData(SubMeshVertexToNonManifoldVertexIDMap,CleanedSubMesh.GetSubmesh());
+
+	// create mesh description for sub-mesh
+	CleanedSubMeshDescription = MakeShared<FMeshDescription>();
+	// registering skeletal mesh attributes is required to create room to copy attributes during conversion from dynamic mesh
+	FSkeletalMeshAttributes Attributes(*CleanedSubMeshDescription);
+	Attributes.Register();
+	// convert the partial dynamic mesh to a mesh description
+	// NOTE: this copies vertex weights to partial mesh description (later used to load weights into the tool)
+	FDynamicMeshToMeshDescription DnyToDescConverter;
+	constexpr bool bCopyTangents = true;
+	DnyToDescConverter.Convert(&CleanedSubMesh.GetSubmesh(), *CleanedSubMeshDescription, bCopyTangents);
+}
+
+void FCleanedEditMesh::CopyWeightsToOriginalMesh(FName Profile)
+{
+	// if the mesh was not duplicated and cleaned, we have nothing to copy from
+	if (!CleanedSubMeshDescription)
+	{
+		return;
+	}
+
+	FSkeletalMeshAttributes SubMeshAttrs(*CleanedSubMeshDescription);
+	FSkinWeightsVertexAttributesRef SubMeshWeightAttrs = SubMeshAttrs.GetVertexSkinWeights(Profile);
+
+	FSkeletalMeshAttributes BaseMeshAttrs(OriginalMeshDescription);
+	FSkinWeightsVertexAttributesRef BaseMeshWeightAttrs = BaseMeshAttrs.GetVertexSkinWeights(Profile);
+
+	UE::AnimationCore::FBoneWeightsSettings Settings;
+	Settings.SetNormalizeType(UE::AnimationCore::EBoneWeightNormalizeType::None);
+
+	TArray<UE::AnimationCore::FBoneWeight> BoneWeightsToApply;
+	BoneWeightsToApply.Reserve(UE::AnimationCore::MaxInlineBoneWeightCount);
+	
+	const int32 NumVerticesInSubMesh = CleanedSubMeshDescription.Get()->Vertices().Num();
+	for (int32 SubVertexID=0; SubVertexID < NumVerticesInSubMesh; SubVertexID++)
+	{
+		BoneWeightsToApply.Reset();
+		
+		FVertexBoneWeights SubVertexWeights = SubMeshWeightAttrs.Get(SubVertexID);
+		for (UE::AnimationCore::FBoneWeight SingleBoneWeight : SubVertexWeights)
+		{
+			BoneWeightsToApply.Add(UE::AnimationCore::FBoneWeight(SingleBoneWeight.GetBoneIndex(), SingleBoneWeight.GetWeight()));
+		}
+
+		// remap the submesh vertex ID to the base mesh and then copy the weights over
+		int32 BaseVertexID = CleanedSubMesh.MapVertexToBaseMesh(SubVertexID);
+		BaseMeshWeightAttrs.Set(FVertexID(BaseVertexID), UE::AnimationCore::FBoneWeights::Create(BoneWeightsToApply, Settings));
+	}
+}
+
+FDynamicMesh3& FCleanedEditMesh::GetEditableMesh()
+{
+	return CleanedSubMeshDescription ? CleanedSubMesh.GetSubmesh() : OriginalDynamicMesh;
+}
+
+FMeshDescription& FCleanedEditMesh::GetEditableMeshDescription()
+{
+	return CleanedSubMeshDescription ? *CleanedSubMeshDescription.Get() : OriginalMeshDescription;
+}
+
+FDynamicMesh3& FCleanedEditMesh::GetOriginalMesh()
+{
+	return OriginalDynamicMesh;
+}
+
+FMeshDescription& FCleanedEditMesh::GetOriginalMeshDescription()
+{
+	return OriginalMeshDescription;
+}
+
 bool USkinWeightsPaintTool::HasActiveSelectionOnMainMesh()
 {
 	const bool bSelectingTargetMesh = WeightToolProperties->MeshSelectMode == EMeshTransferOption::Target;
diff --git a/Engine/Plugins/Experimental/MeshModelingToolsetExp/Source/MeshModelingToolsEditorOnlyExp/Public/SkeletalMesh/SkinWeightsPaintTool.h b/Engine/Plugins/Experimental/MeshModelingToolsetExp/Source/MeshModelingToolsEditorOnlyExp/Public/SkeletalMesh/SkinWeightsPaintTool.h
--- a/Engine/Plugins/Experimental/MeshModelingToolsetExp/Source/MeshModelingToolsEditorOnlyExp/Public/SkeletalMesh/SkinWeightsPaintTool.h
+++ b/Engine/Plugins/Experimental/MeshModelingToolsetExp/Source/MeshModelingToolsEditorOnlyExp/Public/SkeletalMesh/SkinWeightsPaintTool.h
@@ -668,6 +668,29 @@ public:
 	virtual FString ToString() const override;
 };
 
+// Why does this exist?
+// It is important for the mesh selection algorithms to operate on a mesh that is free from floating vertices,
+// but changing the topology of the mesh the tool is working on causes issue with the non-manifold mapping, UV attributes and possibly other things
+//
+// So, rather than track these issues down individually, we leave the original mesh topology intact, and operate on a cleaned submesh.
+// When the tool exits, we copy the weight edits from the submesh to the full mesh.
+// This means weights of the floating vertices cannot be edited, however they do not matter since they do not belong to any triangles they cannot be rendered
+struct MESHMODELINGTOOLSEDITORONLYEXP_API FCleanedEditMesh
+{
+	FCleanedEditMesh(const FDynamicMesh3& InDynamicMesh, const FMeshDescription& InMeshDescription);
+	void CopyWeightsToOriginalMesh(FName Profile);
+	FDynamicMesh3& GetEditableMesh();
+	FMeshDescription& GetEditableMeshDescription();
+	FDynamicMesh3& GetOriginalMesh();
+	FMeshDescription& GetOriginalMeshDescription();
+
+private:
+	FDynamicMesh3 OriginalDynamicMesh;
+	FMeshDescription OriginalMeshDescription;
+	UE::Geometry::FDynamicSubmesh3 CleanedSubMesh;
+	TSharedPtr<FMeshDescription> CleanedSubMeshDescription = nullptr;
+};
+
 // An interactive tool for painting and editing skin weights.
 UCLASS()
 class MESHMODELINGTOOLSEDITORONLYEXP_API USkinWeightsPaintTool : public UDynamicMeshBrushTool, public ISkeletalMeshEditingInterface
@@ -733,7 +756,7 @@ public:
 	void ToggleEditingMode();
 	// update the state of the mesh selectors
 	void UpdateSelectorState() const;
-
+	
 	// get access to the mesh selector for the main mesh
 	UWeightToolMeshSelector* GetMainMeshSelector();
 	// get access to the currently active mesh selector (may be on the transfer source mesh)
@@ -778,7 +801,9 @@ public:
 	USkinWeightsPaintToolProperties* GetWeightToolProperties() const;
 
 	// get access to the mesh description for the mesh being edited
-	FMeshDescription* GetCurrentlyEditedMeshDescription() const { return EditedMesh; };
+	FCleanedEditMesh* GetCurrentCleanedEditMesh() const;
+	FMeshDescription* GetCurrentlyEditedMeshDescription() const;
+	FDynamicMesh3* GetCurrentlyEditedDynamicMesh() const;
 
 	// HOW TO EDIT WEIGHTS WITH UNDO/REDO:
 	//
@@ -818,8 +843,6 @@ protected:
 	void OnShutdown(EToolShutdownType ShutdownType) override;
 	void OnTick(float DeltaTime) override;
 
-	void CleanMesh() const;
-
 	// stamp
 	float CalculateBrushFalloff(float Distance) const;
 	void CalculateVertexROI(
@@ -866,8 +889,8 @@ protected:
 	virtual void OnPropertyModified(UObject* ModifiedObject, FProperty* ModifiedProperty) override;
 	
 	// the currently edited mesh descriptions
-	mutable TMap<EMeshLODIdentifier, FMeshDescription> EditedMeshes;
-	FMeshDescription* EditedMesh = nullptr;
+	mutable TMap<EMeshLODIdentifier, FCleanedEditMesh> EditedMeshes;
+	EMeshLODIdentifier CurrentlyEditedLOD;
 
 	// storage of vertex weights per bone 
 	SkinPaintTool::FSkinToolWeights Weights;
diff --git a/Engine/Source/Runtime/GeometryCore/Private/DynamicMesh/NonManifoldMappingSupport.cpp b/Engine/Source/Runtime/GeometryCore/Private/DynamicMesh/NonManifoldMappingSupport.cpp
--- a/Engine/Source/Runtime/GeometryCore/Private/DynamicMesh/NonManifoldMappingSupport.cpp
+++ b/Engine/Source/Runtime/GeometryCore/Private/DynamicMesh/NonManifoldMappingSupport.cpp
@@ -208,7 +208,11 @@ public:
 	{
 		bValid = false;
 	}
-	
+
+	virtual void OnNewVertex(int VertexID, bool bInserted) override
+	{
+		SourceVertexIds.InsertAt(VertexID, VertexID);
+	}
 };
 
 void FNonManifoldSourceVertexIdAttributeChange::SaveInitialVertex(const FDynamicMeshAttributeBase* Attribute, int VertexID)
