Skin Weights Binding Tool: defer selection update so that the mesh operator get a chance to get updated first

#rnx

diff --git a/Engine/Plugins/Experimental/MeshModelingToolsetExp/Source/MeshModelingToolsEditorOnlyExp/Private/SkeletalMesh/SkinWeightsBindingTool.cpp b/Engine/Plugins/Experimental/MeshModelingToolsetExp/Source/MeshModelingToolsEditorOnlyExp/Private/SkeletalMesh/SkinWeightsBindingTool.cpp
--- a/Engine/Plugins/Experimental/MeshModelingToolsetExp/Source/MeshModelingToolsEditorOnlyExp/Private/SkeletalMesh/SkinWeightsBindingTool.cpp
+++ b/Engine/Plugins/Experimental/MeshModelingToolsetExp/Source/MeshModelingToolsEditorOnlyExp/Private/SkeletalMesh/SkinWeightsBindingTool.cpp
@@ -79,6 +79,18 @@ void USkinWeightsBindingTool::Setup()
 	if (ensure(Properties))
 	{
 		Properties->RestoreProperties(this);
+		
+		if (Properties->CurrentBone != NAME_None)
+		{
+			const FName RestoredCurrentBone = Properties->CurrentBone;
+			Properties->CurrentBone = NAME_None;
+			PendingUpdateFunction = [this, RestoredCurrentBone]()
+			{
+				const int32 BoneIndex = ReferenceSkeleton.FindRawBoneIndex(RestoredCurrentBone);
+				Properties->CurrentBone = BoneIndex != INDEX_NONE ? RestoredCurrentBone : NAME_None;
+				UpdateVisualization();
+			};
+		}
 	}
 
 	if (!ensure(Targets.Num() > 0) || !ensure(Targets[0]))
@@ -196,6 +208,12 @@ void USkinWeightsBindingTool::OnShutdown(EToolShutdownType ShutdownType)
 void USkinWeightsBindingTool::OnTick(float DeltaTime)
 {
 	Preview->Tick(DeltaTime);
+
+	if (PendingUpdateFunction)
+	{
+		PendingUpdateFunction();
+		PendingUpdateFunction.Reset();
+	}
 }
 
 static void DrawBox(IToolsContextRenderAPI* RenderAPI, const FTransform& Transform, const FBox &Box, const FLinearColor &Color, float LineThickness)
@@ -243,7 +261,7 @@ static void DrawBox(IToolsContextRenderAPI* RenderAPI, const FTransform& Transfo
 void USkinWeightsBindingTool::Render(IToolsContextRenderAPI* RenderAPI)
 {
 	using namespace UE::Geometry;
-
+	
 	if (Occupancy && Properties->bDebugDraw)
 	{
 		constexpr bool bShowInterior = false;
@@ -423,7 +441,13 @@ void USkinWeightsBindingTool::HandleSkeletalMeshModified(const TArray<FName>& In
 		break;
 	case ESkeletalMeshNotifyType::BonesSelected:
 		{
-			Properties->CurrentBone = InBoneNames.IsEmpty() ? NAME_None : InBoneNames[0];
+			const int32 BoneIndex = InBoneNames.IsEmpty() ? INDEX_NONE : ReferenceSkeleton.FindRawBoneIndex(InBoneNames[0]);
+			const FName NewSelectedBone = BoneIndex != INDEX_NONE ? InBoneNames[0] : NAME_None;
+			PendingUpdateFunction = [this, NewSelectedBone]()
+			{
+				Properties->CurrentBone = NewSelectedBone;
+				UpdateVisualization();
+			};
 		}
 		break;
 	case ESkeletalMeshNotifyType::BonesRenamed:
diff --git a/Engine/Plugins/Experimental/MeshModelingToolsetExp/Source/MeshModelingToolsEditorOnlyExp/Public/SkeletalMesh/SkinWeightsBindingTool.h b/Engine/Plugins/Experimental/MeshModelingToolsetExp/Source/MeshModelingToolsEditorOnlyExp/Public/SkeletalMesh/SkinWeightsBindingTool.h
--- a/Engine/Plugins/Experimental/MeshModelingToolsetExp/Source/MeshModelingToolsEditorOnlyExp/Public/SkeletalMesh/SkinWeightsBindingTool.h
+++ b/Engine/Plugins/Experimental/MeshModelingToolsetExp/Source/MeshModelingToolsEditorOnlyExp/Public/SkeletalMesh/SkinWeightsBindingTool.h
@@ -134,4 +134,7 @@ protected:
 	
 	// ISkeletalMeshEditionInterface
 	virtual void HandleSkeletalMeshModified(const TArray<FName>& InBoneNames, const ESkeletalMeshNotifyType InNotifyType) override;
+
+	// Pending update function */
+	TFunction<void()> PendingUpdateFunction;
 };
