Fixes crash in skin paint tool during rapid undo/redo of isolated selection

diff --git a/Engine/Plugins/Experimental/MeshModelingToolsetExp/Source/MeshModelingToolsEditorOnlyExp/Private/SkeletalMesh/SkinWeightsPaintTool.cpp b/Engine/Plugins/Experimental/MeshModelingToolsetExp/Source/MeshModelingToolsEditorOnlyExp/Private/SkeletalMesh/SkinWeightsPaintTool.cpp
--- a/Engine/Plugins/Experimental/MeshModelingToolsetExp/Source/MeshModelingToolsEditorOnlyExp/Private/SkeletalMesh/SkinWeightsPaintTool.cpp
+++ b/Engine/Plugins/Experimental/MeshModelingToolsetExp/Source/MeshModelingToolsEditorOnlyExp/Private/SkeletalMesh/SkinWeightsPaintTool.cpp
@@ -456,6 +456,7 @@ void UWeightToolMeshSelector::SetMesh(
 
 	// clear the selection (old selection is invalid on new topo)
 	PolygonSelectionMechanic->ClearSelection();
+	PolygonSelectionMechanic->ClearHighlight();
 
 	// selection colors
 	constexpr FLinearColor FaceSelectedOrange = FLinearColor(0.886f, 0.672f, 0.473f);
