Fixes ensure in skin paint tool when closing tool

diff --git a/Engine/Plugins/Experimental/MeshModelingToolsetExp/Source/MeshModelingToolsEditorOnlyExp/Private/SkeletalMesh/SkinWeightsPaintTool.cpp b/Engine/Plugins/Experimental/MeshModelingToolsetExp/Source/MeshModelingToolsEditorOnlyExp/Private/SkeletalMesh/SkinWeightsPaintTool.cpp
--- a/Engine/Plugins/Experimental/MeshModelingToolsetExp/Source/MeshModelingToolsEditorOnlyExp/Private/SkeletalMesh/SkinWeightsPaintTool.cpp
+++ b/Engine/Plugins/Experimental/MeshModelingToolsetExp/Source/MeshModelingToolsEditorOnlyExp/Private/SkeletalMesh/SkinWeightsPaintTool.cpp
@@ -623,7 +623,7 @@ const TArray<int32>& UWeightToolMeshSelector::GetSelectedVertices()
 
 bool UWeightToolMeshSelector::IsAnyComponentSelected() const
 {
-	if (!ensure(PolygonSelectionMechanic))
+	if (!PolygonSelectionMechanic)
 	{
 		return false;
 	}
