Make Edit Normals behavior more consistent in how it updates the target

diff --git a/Engine/Plugins/Experimental/MeshModelingToolsetExp/Source/MeshModelingToolsExp/Private/EditNormalsTool.cpp b/Engine/Plugins/Experimental/MeshModelingToolsetExp/Source/MeshModelingToolsExp/Private/EditNormalsTool.cpp
--- a/Engine/Plugins/Experimental/MeshModelingToolsetExp/Source/MeshModelingToolsExp/Private/EditNormalsTool.cpp
+++ b/Engine/Plugins/Experimental/MeshModelingToolsetExp/Source/MeshModelingToolsExp/Private/EditNormalsTool.cpp
@@ -421,20 +421,12 @@ void UEditNormalsTool::GenerateAsset(const TArray<FDynamicMeshOpResult>& Results
 		const FDynamicMesh3* NewDynamicMesh = Results[ComponentIdx].Mesh.Get();
 		if (NewDynamicMesh)
 		{
-			if (bool bTopologyChanged = BasicProperties->WillTopologyChange())
-			{
-				// Tool may have changed the topology of the normal overlay (according to the specified tool properties), so we can't simply update the target mesh.
-				// Passing in bTopologyChanged = true will trigger the slower Convert function rather than the fast Update function.
-				UE::ToolTarget::CommitMeshDescriptionUpdateViaDynamicMesh(Targets[ComponentIdx], *NewDynamicMesh, bTopologyChanged);
-			}
-			else
-			{
-				// The tool didn't change the overlay topology so there's a chance we can do a fast path Update of the normal attributes.
-				// This function will still check if there is a mismatch between the dynamic mesh and target mesh in terms of triangles/vertices, and
-				// if so it will do the full conversion.
-				constexpr bool bUpdateTangents = false;
-				UE::ToolTarget::CommitDynamicMeshNormalsUpdate(Targets[ComponentIdx], NewDynamicMesh, bUpdateTangents);
-			}
+			// Note we could use BasicProperties->WillTopologyChange(); to be faster in the simple cases topology has not changed,
+			// or use UE::ToolTarget::CommitDynamicMeshNormalsUpdate(Targets[ComponentIdx], NewDynamicMesh, false); in that case,
+			// but doing so avoids recomputing tangents in some cases where it *should* be done, so for more consistent behavior
+			// we always commit the full mesh update here.
+			// See JIRA UE-225381 for more info on the inconsistent behavior case.
+			UE::ToolTarget::CommitMeshDescriptionUpdateViaDynamicMesh(Targets[ComponentIdx], *NewDynamicMesh, true /*bHaveModifiedTopology*/);
 		}
 	}
 
