MeshModelingTools: Fix Transform tool and ISM Editor resetting the "Set Pivot" setting when rebuilding their gizmo.

diff --git a/Engine/Plugins/Experimental/MeshModelingToolsetExp/Source/MeshModelingToolsEditorOnlyExp/Private/ISMEditorTool.cpp b/Engine/Plugins/Experimental/MeshModelingToolsetExp/Source/MeshModelingToolsEditorOnlyExp/Private/ISMEditorTool.cpp
--- a/Engine/Plugins/Experimental/MeshModelingToolsetExp/Source/MeshModelingToolsEditorOnlyExp/Private/ISMEditorTool.cpp
+++ b/Engine/Plugins/Experimental/MeshModelingToolsetExp/Source/MeshModelingToolsEditorOnlyExp/Private/ISMEditorTool.cpp
@@ -573,6 +573,7 @@ void UISMEditorTool::SetActiveGizmos_Single(bool bLocalRotations)
 
 	FISMEditorTarget Transformable;
 	Transformable.TransformProxy = NewObject<UTransformProxy>(this);
+	Transformable.TransformProxy->bSetPivotMode = TransformProps ? TransformProps->bSetPivotMode : false;
 	Transformable.TransformProxy->bRotatePerObject = bLocalRotations;
 	Transformable.TransformProxy->OnBeginTransformEdit.AddLambda([this](UTransformProxy*) { bInActiveDrag = true; });
 	Transformable.TransformProxy->OnEndTransformEdit.AddLambda([this](UTransformProxy*) { OnTransformCompleted(); bInActiveDrag = false; });
@@ -606,6 +607,7 @@ void UISMEditorTool::SetActiveGizmos_PerObject()
 		{
 			FISMEditorTarget Transformable;
 			Transformable.TransformProxy = NewObject<UTransformProxy>(this);
+			Transformable.TransformProxy->bSetPivotMode = TransformProps ? TransformProps->bSetPivotMode : false;
 			Transformable.TransformProxy->OnBeginTransformEdit.AddLambda([this](UTransformProxy*) { bInActiveDrag = true; });
 			Transformable.TransformProxy->OnEndTransformEdit.AddLambda([this](UTransformProxy*) { OnTransformCompleted(); bInActiveDrag = false; });
 			Transformable.TransformProxy->OnTransformChangedUndoRedo.AddLambda([this](UTransformProxy*, FTransform) { OnTransformCompleted(); });
diff --git a/Engine/Plugins/Experimental/MeshModelingToolsetExp/Source/MeshModelingToolsExp/Private/TransformMeshesTool.cpp b/Engine/Plugins/Experimental/MeshModelingToolsetExp/Source/MeshModelingToolsExp/Private/TransformMeshesTool.cpp
--- a/Engine/Plugins/Experimental/MeshModelingToolsetExp/Source/MeshModelingToolsExp/Private/TransformMeshesTool.cpp
+++ b/Engine/Plugins/Experimental/MeshModelingToolsetExp/Source/MeshModelingToolsExp/Private/TransformMeshesTool.cpp
@@ -277,6 +277,7 @@ void UTransformMeshesTool::SetActiveGizmos_Single(bool bLocalRotations)
 
 	FTransformMeshesTarget Transformable;
 	Transformable.TransformProxy = NewObject<UTransformProxy>(this);
+	Transformable.TransformProxy->bSetPivotMode = TransformProps ? TransformProps->bSetPivotMode : false;
 	Transformable.TransformProxy->bRotatePerObject = bLocalRotations;
 
 	TArray<const UPrimitiveComponent*> ComponentsToIgnoreInAlignment;
@@ -333,6 +334,7 @@ void UTransformMeshesTool::SetActiveGizmos_PerObject()
 
 				FTransformMeshesTarget Transformable;
 				Transformable.TransformProxy = NewObject<UTransformProxy>(this);
+				Transformable.TransformProxy->bSetPivotMode = TransformProps ? TransformProps->bSetPivotMode : false;
 				UE::Local::AddInstancedComponentInstance(InstancedComponent, k, Transformable.TransformProxy, true);
 
 				ETransformGizmoSubElements GizmoElements = ETransformGizmoSubElements::FullTranslateRotateScale;
@@ -351,6 +353,7 @@ void UTransformMeshesTool::SetActiveGizmos_PerObject()
 
 			FTransformMeshesTarget Transformable;
 			Transformable.TransformProxy = NewObject<UTransformProxy>(this);
+			Transformable.TransformProxy->bSetPivotMode = TransformProps ? TransformProps->bSetPivotMode : false;
 			Transformable.TransformProxy->AddComponent(Component);
 
 			ETransformGizmoSubElements GizmoElements = ETransformGizmoSubElements::FullTranslateRotateScale;
