MeshModelingTools: Fix DrawSplineTool not hiding objects in outliner and not updating its blueprint in multiple places.

diff --git a/Engine/Plugins/Experimental/MeshModelingToolsetExp/Source/MeshModelingToolsEditorOnlyExp/MeshModelingToolsEditorOnlyExp.Build.cs b/Engine/Plugins/Experimental/MeshModelingToolsetExp/Source/MeshModelingToolsEditorOnlyExp/MeshModelingToolsEditorOnlyExp.Build.cs
--- a/Engine/Plugins/Experimental/MeshModelingToolsetExp/Source/MeshModelingToolsEditorOnlyExp/MeshModelingToolsEditorOnlyExp.Build.cs
+++ b/Engine/Plugins/Experimental/MeshModelingToolsetExp/Source/MeshModelingToolsEditorOnlyExp/MeshModelingToolsEditorOnlyExp.Build.cs
@@ -57,6 +57,7 @@ public class MeshModelingToolsEditorOnlyExp : ModuleRules
 				"UnrealEd",
 				"Persona",
 				"AdvancedPreviewScene",
+				"SceneOutliner",
 				"Slate",
 				"Json",
 				"ApplicationCore"
diff --git a/Engine/Plugins/Experimental/MeshModelingToolsetExp/Source/MeshModelingToolsEditorOnlyExp/Private/DrawSplineTool.cpp b/Engine/Plugins/Experimental/MeshModelingToolsetExp/Source/MeshModelingToolsEditorOnlyExp/Private/DrawSplineTool.cpp
--- a/Engine/Plugins/Experimental/MeshModelingToolsetExp/Source/MeshModelingToolsEditorOnlyExp/Private/DrawSplineTool.cpp
+++ b/Engine/Plugins/Experimental/MeshModelingToolsetExp/Source/MeshModelingToolsEditorOnlyExp/Private/DrawSplineTool.cpp
@@ -13,8 +13,12 @@
 #include "Engine/Blueprint.h"
 #include "Engine/World.h" 
 #include "GameFramework/Actor.h"
+#include "ILevelEditor.h"
 #include "InputState.h" // FInputDeviceRay
 #include "InteractiveToolManager.h"
+#include "ISceneOutliner.h"
+#include "LevelEditor.h"
+#include "Modules/ModuleManager.h"
 #include "Kismet2/ComponentEditorUtils.h" // GenerateValidVariableName
 #include "Mechanics/ConstructionPlaneMechanic.h"
 #include "SceneManagement.h" // FPrimitiveDrawInterface
@@ -28,10 +32,83 @@
 
 using namespace UE::Geometry;
 
+/**
+ * Helper to hide actors from the outliner if it's not an actor that
+ *  we defined to be automatically hidden (e.g. APreviewGeometryActor),
+ *  that works by being friended to FSetActorHiddenInSceneOutliner.
+ * This is a temporary measure until we have a cleaner way to hide
+ *  ourselves from the outliner through TEDS.
+ * Note that just creating this class doesn't actually refresh the outliner
+ *  unless you happen to take an action that does (such as reparenting things),
+ *  so you would need to call RefreshOutliner().
+ */
+class FModelingToolsSetActorHiddenInSceneOutliner
+{
+public:
+	FModelingToolsSetActorHiddenInSceneOutliner(AActor* InActor, bool bHidden)
+	{
+		FSetActorHiddenInSceneOutliner Setter(InActor, bHidden);
+	}
+
+	/**
+	 * Does a full refresh of the outliner. Note that this can be comparatively
+	 *  slow, so it should happen rarely.
+	 */
+	void RefreshOutliner()
+	{
+		FLevelEditorModule* LevelEditorModule = FModuleManager::GetModulePtr<FLevelEditorModule>(TEXT("LevelEditor"));
+		if (!LevelEditorModule)
+		{
+			return;
+		}
+
+		TSharedPtr<ILevelEditor> LevelEditor = LevelEditorModule->GetLevelEditorInstance().Pin();
+		if (!LevelEditor.IsValid())
+		{
+			return;
+		}
+
+		for (TWeakPtr<ISceneOutliner> OutlinerWeak : LevelEditor->GetAllSceneOutliners())
+		{
+			if (TSharedPtr<ISceneOutliner> Outliner = OutlinerWeak.Pin())
+			{
+				Outliner->FullRefresh();
+			}
+		}
+	}
+};
+
 namespace DrawSplineToolLocals
 {
 	FText AddPointTransactionName = LOCTEXT("AddPointTransactionName", "Add Point");
 
+	// Helper for destroying actors in a "best practice" way. This is more valuable when destroying non-preview
+	//  actors to make sure that things are properly deselected, etc., so it's not critical in this particular tool,
+	//  but we'll use it.
+	// TODO: Move this to some helper file for use in multiple tools.
+	void DestroyActor(AActor* Actor)
+	{
+		if (!Actor)
+		{
+			return;
+		}
+		if (UWorld* ActorWorld = Actor->GetWorld())
+		{
+			if (GIsEditor && GUnrealEd)
+			{
+				GUnrealEd->DeleteActors(TArray{ Actor }, ActorWorld, GUnrealEd->GetSelectedActors()->GetElementSelectionSet());
+			}
+			else
+			{
+				ActorWorld->DestroyActor(Actor);
+			}
+		}
+		else
+		{
+			Actor->Destroy();
+		}
+	};
+
 	USplineComponent* CreateNewSplineInActor(AActor* Actor, bool bTransact = false, bool bSetAsRoot = false)
 	{
 		if (!ensure(Actor))
@@ -302,6 +379,7 @@ void UDrawSplineTool::Setup()
 		if (ensure(WorkingSpline.IsValid()))
 		{
 			WorkingSpline->SetClosedLoop(Settings->bLoop);
+			bNeedToRerunConstructionScript = true;
 		}
 	});
 
@@ -345,7 +423,7 @@ void UDrawSplineTool::TransitionOutputMode()
 	// Restore the visibility of the previous target actor and spline, if needed
 	if (PreviousTargetActor)
 	{
-		PreviousTargetActor->SetIsTemporarilyHiddenInEditor(PreviousTargetActorVisibility);
+		PreviousTargetActor->GetRootComponent()->SetVisibility(true, true);
 		PreviousTargetActor = nullptr;
 	}
 	if (HiddenSpline.IsValid())
@@ -354,26 +432,13 @@ void UDrawSplineTool::TransitionOutputMode()
 		HiddenSpline = nullptr;
 	}
 
-	// Keep the previous spline/previews temporarily so we can transfer over spline data
+	// Keep the previous spline/preview temporarily so we can transfer over spline data
 	// when we make new previews
-	AActor* PreviousPreviewRoot = PreviewRootActor;
 	AActor* PreviousPreview = PreviewActor;
 	USplineComponent* PreviousSpline = WorkingSpline.Get();
 
-	PreviewRootActor = nullptr;
 	PreviewActor = nullptr;
 	WorkingSpline = nullptr;
-
-	// Create an entirely new preview root. We could probably keep the same one and disconnect/connect,
-	// but it seems cleaner to build from scratch whenever we have to change output mode.
-	FRotator Rotation(0.0f, 0.0f, 0.0f);
-	FActorSpawnParameters SpawnInfo;
-	SpawnInfo.ObjectFlags = RF_Transient;
-	PreviewRootActor = GetTargetWorld()->SpawnActor<APreviewGeometryActor>(FVector::ZeroVector, Rotation, SpawnInfo);
-	USceneComponent* RootComponent = NewObject<USceneComponent>(PreviewRootActor);
-	PreviewRootActor->AddOwnedComponent(RootComponent);
-	PreviewRootActor->SetRootComponent(RootComponent);
-	RootComponent->RegisterComponent();
 	
 	// Used for visualizing the effect of a spline on some special actor
 	auto CreateDuplicatePreviewActor = [this](AActor* Actor)
@@ -382,37 +447,40 @@ void UDrawSplineTool::TransitionOutputMode()
 		GUnrealEd->DuplicateActors({ Actor }, NewActors, GetWorld()->GetCurrentLevel(), FVector3d::Zero());
 		if (!ensure(NewActors.Num() > 0))
 		{
-			return;
+			return false;
 		}
 
 		PreviewActor = NewActors[0];
 		PreviewActor->ClearFlags(RF_Transactional);
 		PreviewActor->SetFlags(RF_Transient);
 
-		if (!ensure(PreviewActor->GetRootComponent()))
-		{
-			USceneComponent* NewRoot = NewObject<USceneComponent>(PreviewRootActor);
-			PreviewActor->AddOwnedComponent(NewRoot);
-			PreviewActor->SetRootComponent(NewRoot);
-			NewRoot->RegisterComponent();
-		}
+		// Make the original invisible
+		PreviousTargetActor = Actor;
+		PreviousTargetActor->GetRootComponent()->SetVisibility(false, true);
 
-		// Attach the preview actor to the non-outliner-visible preview root. The proper way to do this is 
-		// "GEditor->ParentActors(PreviewActor, PreviewRootActor, NAME_None);", but it is hard to prevent
-		// that call from emitting an undo/redo transaction. It may be possible if the actors AND the root
-		// components are all marked as not transactable, but seems simpler to do this by hand.
-		PreviewActor->GetRootComponent()->AttachToComponent(PreviewRootActor->GetRootComponent(), 
-			FAttachmentTransformRules::KeepWorldTransform, NAME_None);
+		// Note: unfortunately this won't hide the spline itself as long as the object is selected because that
+		//  drawing goes through a different path that doesn't seem to be disableable (unless we deselected
+		//  the actor).
+		
+		// Hide this preview from the outliner
+		FModelingToolsSetActorHiddenInSceneOutliner Hider(PreviewActor, true);
+		Hider.RefreshOutliner();
 
-		// Hide the original
-		PreviousTargetActor = Actor;
-		PreviousTargetActorVisibility = PreviousTargetActor->IsHiddenEd();
-		PreviousTargetActor->SetIsTemporarilyHiddenInEditor(true);
+		return true;
 	};
 
 	auto FallbackSplinePlacement = [this]()
 	{
-		WorkingSpline = CreateNewSplineInActor(PreviewRootActor);
+		FRotator Rotation(0.0f, 0.0f, 0.0f);
+		FActorSpawnParameters SpawnInfo;
+		SpawnInfo.ObjectFlags = RF_Transient;
+		PreviewActor = GetTargetWorld()->SpawnActor<APreviewGeometryActor>(FVector::ZeroVector, Rotation, SpawnInfo);
+		USceneComponent* RootComponent = NewObject<USceneComponent>(PreviewActor);
+		PreviewActor->AddOwnedComponent(RootComponent);
+		PreviewActor->SetRootComponent(RootComponent);
+		RootComponent->RegisterComponent();
+		
+		WorkingSpline = CreateNewSplineInActor(PreviewActor);
 	};
 
 	// Set up the new preview
@@ -429,27 +497,15 @@ void UDrawSplineTool::TransitionOutputMode()
 			break;
 		case EDrawSplineOutputMode::ExistingActor:
 		{
-			if (!Settings->TargetActor.IsValid())
+			if (!Settings->TargetActor.IsValid()
+				|| !CreateDuplicatePreviewActor(Settings->TargetActor.Get()))
 			{
 				FallbackSplinePlacement();
 				break;
 			}
 
-			CreateDuplicatePreviewActor(Settings->TargetActor.Get());
-
-			// Hide the spline we're replacing, if we are replacing one.
-			// TODO: This isn't quite perfect because if the spline is selected, the component visualizer in the editor
-			// will still draw it even if the parent actor is hidden and the spline is set to not be drawn...
-			TInlineComponentArray<USplineComponent*> SplineComponents;
-			Settings->TargetActor->GetComponents<USplineComponent>(SplineComponents);
-			if (Settings->ExistingSplineIndexToReplace >= 0 && Settings->ExistingSplineIndexToReplace < SplineComponents.Num())
-			{
-				HiddenSpline = SplineComponents[Settings->ExistingSplineIndexToReplace];
-				bPreviousSplineVisibility = HiddenSpline->bDrawDebug;
-				HiddenSpline->bDrawDebug = false;
-			}
-
 			WorkingSpline = GetOrCreateTargetSpline(PreviewActor, Settings->ExistingSplineIndexToReplace);
+			bNeedToRerunConstructionScript = true;
 			break;
 		}
 		case EDrawSplineOutputMode::CreateBlueprint:
@@ -477,7 +533,12 @@ void UDrawSplineTool::TransitionOutputMode()
 				break;
 			}
 
+			// Hide this preview from outliner
+			FModelingToolsSetActorHiddenInSceneOutliner Hider(PreviewActor, true);
+			Hider.RefreshOutliner();
+
 			WorkingSpline = GetOrCreateTargetSpline(PreviewActor, Settings->ExistingSplineIndexToReplace);
+			bNeedToRerunConstructionScript = true;
 			break;
 		}
 		default:
@@ -514,23 +575,21 @@ void UDrawSplineTool::TransitionOutputMode()
 	// Now that we've copied over previous preview data, destroy the old previews
 	if (PreviousPreview)
 	{
-		PreviousPreview->Destroy();
-	}
-	if (PreviousPreviewRoot)
-	{
-		PreviousPreviewRoot->Destroy();
+		DestroyActor(PreviousPreview);
 	}
 }
 
 void UDrawSplineTool::Shutdown(EToolShutdownType ShutdownType)
 {
+	using namespace DrawSplineToolLocals;
+
 	LongTransactions.CloseAll(GetToolManager());
 
 	Settings->SaveProperties(this);
 
 	if (PreviousTargetActor)
 	{
-		PreviousTargetActor->SetIsTemporarilyHiddenInEditor(PreviousTargetActorVisibility);
+		PreviousTargetActor->GetRootComponent()->SetVisibility(true, true);
 		PreviousTargetActor = nullptr;
 	}
 	if (HiddenSpline.IsValid())
@@ -552,28 +611,10 @@ void UDrawSplineTool::Shutdown(EToolShutdownType ShutdownType)
 		WorkingSpline->DestroyComponent();
 	}
 
-	auto DestroyActor = [&](AActor* Actor)
-	{
-		if (UWorld* ActorWorld = Actor->GetWorld())
-		{
-			if (GIsEditor && GUnrealEd)
-			{
-				GUnrealEd->DeleteActors(TArray{ Actor }, ActorWorld, GUnrealEd->GetSelectedActors()->GetElementSelectionSet());
-			}
-			else
-			{
-				ActorWorld->DestroyActor(Actor);
-			}
-		}
-	};
 	if (PreviewActor)
 	{
 		DestroyActor(PreviewActor);
 	}
-	if (PreviewRootActor)
-	{
-		DestroyActor(PreviewRootActor);
-	}
 
 	Super::Shutdown(ShutdownType);
 }
diff --git a/Engine/Plugins/Experimental/MeshModelingToolsetExp/Source/MeshModelingToolsEditorOnlyExp/Public/DrawSplineTool.h b/Engine/Plugins/Experimental/MeshModelingToolsetExp/Source/MeshModelingToolsEditorOnlyExp/Public/DrawSplineTool.h
--- a/Engine/Plugins/Experimental/MeshModelingToolsetExp/Source/MeshModelingToolsEditorOnlyExp/Public/DrawSplineTool.h
+++ b/Engine/Plugins/Experimental/MeshModelingToolsetExp/Source/MeshModelingToolsEditorOnlyExp/Public/DrawSplineTool.h
@@ -214,7 +214,7 @@ public:
 	virtual void OnClickDrag(const FInputDeviceRay& DragPos);
 	virtual void OnClickRelease(const FInputDeviceRay& ReleasePos);
 	virtual void OnTerminateDragSequence();
-protected:
+private:
 
 	UPROPERTY()
 	TObjectPtr<UDrawSplineToolProperties> Settings = nullptr;
@@ -230,13 +230,8 @@ protected:
 	// This is only used to initialize TargetActor in the settings object
 	TWeakObjectPtr<AActor> SelectedActor = nullptr;
 
-	// PreviewRootActor either holds WorkingSpline inside it directly, or has some preview actor
-	// attached to it (so that the preview actor is hidden from outliner, like APreviewGeometryActor is).
-	UPROPERTY()
-	TObjectPtr<APreviewGeometryActor> PreviewRootActor = nullptr;
-
-	// The preview actor may be a duplicate of some target blueprint actor so that we can see the
-	// effects of the drawn spline immediately
+	// The preview actor is either a APreviewGeometryActor with a spline, or a duplicate of 
+	// some target blueprint actor so that we can see the effects of the drawn spline immediately.
 	UPROPERTY()
 	TObjectPtr<AActor> PreviewActor = nullptr;
 
diff --git a/Engine/Source/Runtime/Engine/Classes/GameFramework/Actor.h b/Engine/Source/Runtime/Engine/Classes/GameFramework/Actor.h
--- a/Engine/Source/Runtime/Engine/Classes/GameFramework/Actor.h
+++ b/Engine/Source/Runtime/Engine/Classes/GameFramework/Actor.h
@@ -4605,6 +4605,7 @@ private:
 	friend class FFoliageHelper;
 	friend class ULevelInstanceSubsystem;
 	friend class UExternalDataLayerInstance;
+	friend class FModelingToolsSetActorHiddenInSceneOutliner;
 };
 
 struct FSetActorGuid
