Fix mesh to collision support for custom polygroup layers in the multi-select case

diff --git a/Engine/Plugins/Experimental/MeshModelingToolsetExp/Source/MeshModelingToolsExp/Private/Physics/SetCollisionGeometryTool.cpp b/Engine/Plugins/Experimental/MeshModelingToolsetExp/Source/MeshModelingToolsExp/Private/Physics/SetCollisionGeometryTool.cpp
--- a/Engine/Plugins/Experimental/MeshModelingToolsetExp/Source/MeshModelingToolsExp/Private/Physics/SetCollisionGeometryTool.cpp
+++ b/Engine/Plugins/Experimental/MeshModelingToolsetExp/Source/MeshModelingToolsExp/Private/Physics/SetCollisionGeometryTool.cpp
@@ -341,14 +341,30 @@ void USetCollisionGeometryTool::Setup()
 	});
 	UE::ToolTarget::SetSourceObjectVisible(Targets.Last(), Settings->bShowTargetMesh);
 
+	PolygroupLayerProperties = NewObject<UPolygroupLayersProperties>(this);
+	PolygroupLayerProperties->RestoreProperties(this, TEXT("SetCollisionGeometryTool"));
 	if (InitialSourceMeshes.Num() == 1)
 	{
-		PolygroupLayerProperties = NewObject<UPolygroupLayersProperties>(this);
-		PolygroupLayerProperties->RestoreProperties(this, TEXT("SetCollisionGeometryTool"));
 		PolygroupLayerProperties->InitializeGroupLayers(&InitialSourceMeshes[0]);
-		PolygroupLayerProperties->WatchProperty(PolygroupLayerProperties->ActiveGroupLayer, [&](FName) { OnSelectedGroupLayerChanged(); });
-		AddToolPropertySource(PolygroupLayerProperties);
 	}
+	else
+	{
+		TSet<FName> LayerNames;
+		for (int32 MeshIdx = 0; MeshIdx < InitialSourceMeshes.Num(); ++MeshIdx)
+		{
+			if (InitialSourceMeshes[MeshIdx].Attributes())
+			{
+				for (int32 k = 0; k < InitialSourceMeshes[MeshIdx].Attributes()->NumPolygroupLayers(); k++)
+				{
+					FName Name = InitialSourceMeshes[MeshIdx].Attributes()->GetPolygroupLayer(k)->GetName();
+					LayerNames.Add(Name);
+				}
+			}
+		}
+		PolygroupLayerProperties->InitializeGroupLayers(LayerNames);
+	}
+	PolygroupLayerProperties->WatchProperty(PolygroupLayerProperties->ActiveGroupLayer, [&](FName) { OnSelectedGroupLayerChanged(); });
+	AddToolPropertySource(PolygroupLayerProperties);
 
 	VizSettings = NewObject<UCollisionGeometryVisualizationProperties>(this);
 	VizSettings->RestoreProperties(this);
@@ -695,18 +711,17 @@ void USetCollisionGeometryTool::OnSelectedGroupLayerChanged()
 }
 
 
-void USetCollisionGeometryTool::UpdateActiveGroupLayer(FDynamicMesh3* GroupLayersMesh)
+FPolygroupSet USetCollisionGeometryTool::GetActiveGroupLayer(const FDynamicMesh3* GroupLayersMesh)
 {
 	if (PolygroupLayerProperties->HasSelectedPolygroup() == false)
 	{
-		ActiveGroupSet = MakeUnique<UE::Geometry::FPolygroupSet>(GroupLayersMesh);
+		return FPolygroupSet(GroupLayersMesh);
 	}
 	else
 	{
 		FName SelectedName = PolygroupLayerProperties->ActiveGroupLayer;
-		FDynamicMeshPolygroupAttribute* FoundAttrib = UE::Geometry::FindPolygroupLayerByName(*GroupLayersMesh, SelectedName);
-		ensureMsgf(FoundAttrib, TEXT("Selected Attribute Not Found! Falling back to Default group layer."));
-		ActiveGroupSet = MakeUnique<UE::Geometry::FPolygroupSet>(GroupLayersMesh, FoundAttrib);
+		const FDynamicMeshPolygroupAttribute* FoundAttrib = UE::Geometry::FindPolygroupLayerByName(*GroupLayersMesh, SelectedName);
+		return FPolygroupSet(GroupLayersMesh, FoundAttrib);
 	}
 }
 
@@ -807,7 +822,8 @@ void USetCollisionGeometryTool::PrecomputeInputMeshes()
 	}
 	else
 	{
-		ParallelFor(SourceObjectIndices.Num(), [&](int32 k)
+		bool bCanDiscardAttributes = (PolygroupLayerProperties->HasSelectedPolygroup() == false);
+		ParallelFor(SourceObjectIndices.Num(), [this, &TargetTransform, bCanDiscardAttributes](int32 k)
 		{
 			FDynamicMesh3 SourceMesh = InitialSourceMeshes[k];
 			if (Settings->bUseWorldSpace)
@@ -816,7 +832,10 @@ void USetCollisionGeometryTool::PrecomputeInputMeshes()
 				MeshTransforms::ApplyTransform(SourceMesh, ToWorld, true);
 				MeshTransforms::ApplyTransformInverse(SourceMesh, TargetTransform, true);
 			}
-			SourceMesh.DiscardAttributes();
+			if (bCanDiscardAttributes)
+			{
+				SourceMesh.DiscardAttributes();
+			}
 
 			InputMeshes[k] = MakeShared<FDynamicMesh3, ESPMode::ThreadSafe>(MoveTemp(SourceMesh));
 		});
@@ -859,22 +878,30 @@ void USetCollisionGeometryTool::PrecomputeInputMeshes()
 		{
 			UseGroupLayerMesh = InputMeshes[0].Get();
 		}
-		UpdateActiveGroupLayer(UseGroupLayerMesh);
 
 		// Use the active polygroup layer when there is only one input
+		FPolygroupSet ActiveGroupSet = GetActiveGroupLayer(UseGroupLayerMesh);
 		InitializeDerivedMeshSet(InputMeshes, PerGroupInputMeshes,
-			[this](const FDynamicMesh3* Mesh, int32 Tri0, int32 Tri1)
+			[this, &ActiveGroupSet](const FDynamicMesh3* Mesh, int32 Tri0, int32 Tri1)
 			{
-				return ActiveGroupSet->GetTriangleGroup(Tri0) == ActiveGroupSet->GetTriangleGroup(Tri1);
+				return ActiveGroupSet.GetTriangleGroup(Tri0) == ActiveGroupSet.GetTriangleGroup(Tri1);
 			});
 	}
 	else
 	{
+		TMap<const FDynamicMesh3*, FPolygroupSet> ActiveGroupSets;
+		for (int32 Idx = 0; Idx < InputMeshes.Num(); ++Idx)
+		{
+			ActiveGroupSets.Add(InputMeshes[Idx].Get(), GetActiveGroupLayer(InputMeshes[Idx].Get()));
+		}
+
 		// Use the default polygroup layer when there is more than one input
 		InitializeDerivedMeshSet(InputMeshes, PerGroupInputMeshes,
-			[](const FDynamicMesh3* Mesh, int32 Tri0, int32 Tri1)
+			[&ActiveGroupSets](const FDynamicMesh3* Mesh, int32 Tri0, int32 Tri1)
 			{
-				return Mesh->GetTriangleGroup(Tri0) == Mesh->GetTriangleGroup(Tri1);
+				FPolygroupSet* ActiveGroupSet = ActiveGroupSets.Find(Mesh);
+				checkSlow(ActiveGroupSet);
+				return ActiveGroupSet->GetTriangleGroup(Tri0) == ActiveGroupSet->GetTriangleGroup(Tri1);
 			});
 	}
 	PerGroupMeshesApproximator = MakeShared<FMeshSimpleShapeApproximation, ESPMode::ThreadSafe>();
diff --git a/Engine/Plugins/Experimental/MeshModelingToolsetExp/Source/MeshModelingToolsExp/Public/Physics/SetCollisionGeometryTool.h b/Engine/Plugins/Experimental/MeshModelingToolsetExp/Source/MeshModelingToolsExp/Public/Physics/SetCollisionGeometryTool.h
--- a/Engine/Plugins/Experimental/MeshModelingToolsetExp/Source/MeshModelingToolsExp/Public/Physics/SetCollisionGeometryTool.h
+++ b/Engine/Plugins/Experimental/MeshModelingToolsetExp/Source/MeshModelingToolsExp/Public/Physics/SetCollisionGeometryTool.h
@@ -383,9 +383,7 @@ protected:
 		TArray<TSharedPtr<FDynamicMesh3, ESPMode::ThreadSafe>>& ToMeshes,
 		TFunctionRef<bool(const FDynamicMesh3*, int32, int32)> TrisConnectedPredicate);
 
-	TUniquePtr<UE::Geometry::FPolygroupSet> ActiveGroupSet;
 	void OnSelectedGroupLayerChanged();
-	void UpdateActiveGroupLayer(FDynamicMesh3* GroupLayersMesh);
 
 	FTransform OrigTargetTransform;
 	UE::Geometry::FTransformSequence3d TargetInverseTransform;
@@ -407,4 +405,7 @@ protected:
 
 	UPROPERTY()
 	TObjectPtr<UPreviewGeometry> GeometrySelectionViz = nullptr;
+
+private:
+	UE::Geometry::FPolygroupSet GetActiveGroupLayer(const FDynamicMesh3* GroupLayersMesh);
 };
