Skin paint tool: fix regression preventing LOD weights from being edited.

diff --git a/Engine/Plugins/Experimental/MeshModelingToolsetExp/Source/MeshModelingToolsEditorOnlyExp/Private/SkeletalMesh/SkinWeightsPaintTool.cpp b/Engine/Plugins/Experimental/MeshModelingToolsetExp/Source/MeshModelingToolsEditorOnlyExp/Private/SkeletalMesh/SkinWeightsPaintTool.cpp
--- a/Engine/Plugins/Experimental/MeshModelingToolsetExp/Source/MeshModelingToolsEditorOnlyExp/Private/SkeletalMesh/SkinWeightsPaintTool.cpp
+++ b/Engine/Plugins/Experimental/MeshModelingToolsetExp/Source/MeshModelingToolsEditorOnlyExp/Private/SkeletalMesh/SkinWeightsPaintTool.cpp
@@ -2119,7 +2119,7 @@ void USkinWeightsPaintTool::Setup()
 	WeightToolProperties->WeightTool = this;
 	WeightToolProperties->bSpecifyRadius = true;
 	// watch for skin weight layer changes
-	WeightToolProperties->ActiveLOD = "LOD0";
+	WeightToolProperties->ActiveLOD = GetLODName(DefaultLOD);
 	int32 WatcherIndex = WeightToolProperties->WatchProperty(WeightToolProperties->ActiveLOD, [this](FName) { OnActiveLODChanged(); });
 	WeightToolProperties->SilentUpdateWatcherAtIndex(WatcherIndex);
 	WeightToolProperties->ActiveSkinWeightProfile = FSkeletalMeshAttributesShared::DefaultSkinWeightProfileName;
@@ -4144,18 +4144,21 @@ void USkinWeightsPaintTool::OnActiveLODChanged()
 	Weights.ApplyCurrentWeightsToMeshDescription(GetCurrentlyEditedMeshDescription());
 
 	// update current mesh using the new LOD
-	const EMeshLODIdentifier LODId = GetLODId(WeightToolProperties->ActiveLOD);
-	const FGetMeshParameters Params(true, LODId);
-	FCleanedEditMesh* CleanedEditMesh = EditedMeshes.Find(LODId);
+	const EMeshLODIdentifier ActiveLODID = GetLODId(WeightToolProperties->ActiveLOD);
+	const FGetMeshParameters Params(true, ActiveLODID);
+	FCleanedEditMesh* CleanedEditMesh = EditedMeshes.Find(ActiveLODID);
 	if (!CleanedEditMesh)
 	{
 		const FDynamicMesh3 DynamicMesh = UE::ToolTarget::GetDynamicMeshCopy(Target, Params);
 		const FMeshDescription* MeshDescription = UE::ToolTarget::GetMeshDescription(Target, Params);
-		CleanedEditMesh = &EditedMeshes.Add(LODId, FCleanedEditMesh(DynamicMesh, *MeshDescription));
+		CleanedEditMesh = &EditedMeshes.Add(ActiveLODID, FCleanedEditMesh(DynamicMesh, *MeshDescription));
 	}
+	
+	// assign new LOD as the currently edited one
+	CurrentlyEditedLOD = ActiveLODID;
 
 	// reinitialize all mesh data structures
-	UpdateCurrentlyEditedMesh(Component, CleanedEditMesh->GetEditableMesh(), CleanedEditMesh->GetEditableMeshDescription());
+	UpdateCurrentlyEditedMesh(Component, *GetCurrentlyEditedDynamicMesh(), *GetCurrentlyEditedMeshDescription());
 }
 
 void USkinWeightsPaintTool::OnActiveSkinWeightProfileChanged()
