Fixes issue, skin weight tool could load non-normalized weights

diff --git a/Engine/Plugins/Experimental/MeshModelingToolsetExp/Source/MeshModelingToolsEditorOnlyExp/Private/SkeletalMesh/SkinWeightsPaintTool.cpp b/Engine/Plugins/Experimental/MeshModelingToolsetExp/Source/MeshModelingToolsEditorOnlyExp/Private/SkeletalMesh/SkinWeightsPaintTool.cpp
--- a/Engine/Plugins/Experimental/MeshModelingToolsetExp/Source/MeshModelingToolsEditorOnlyExp/Private/SkeletalMesh/SkinWeightsPaintTool.cpp
+++ b/Engine/Plugins/Experimental/MeshModelingToolsetExp/Source/MeshModelingToolsEditorOnlyExp/Private/SkeletalMesh/SkinWeightsPaintTool.cpp
@@ -1481,35 +1481,51 @@ void FSkinToolWeights::InitializeSkinWeights(
 	const FSkeletalMeshConstAttributes MeshAttribs(*Mesh);
 	const FSkinWeightsVertexAttributesConstRef VertexSkinWeights = MeshAttribs.GetVertexSkinWeights(Profile);
 	const int32 NumVertices = Mesh->Vertices().Num();
+	TMap<BoneIndex, float> NormalizedWeights;
 	CurrentWeights.SetNum(NumVertices);
 	for (int32 VertexIndex = 0; VertexIndex < NumVertices; VertexIndex++)
 	{
-		const FVertexID VertexID(VertexIndex);
-		int32 InfluenceIndex = 0;
-		for (UE::AnimationCore::FBoneWeight BoneWeight: VertexSkinWeights.Get(VertexID))
+		// we have to normalize here because there are edges cases where skeletal meshes are loaded with non-normalized weights.
+		// even when/if we fix these, it's best to be 100% sure that things are normalized since weight data can be generated by tools we don't control
+		NormalizedWeights.Reset();
+		int32 NumInfluences = 0;
+		const FVertexBoneWeightsConst BoneWeights = VertexSkinWeights.Get(FVertexID(VertexIndex));
+		for (const UE::AnimationCore::FBoneWeight BoneWeight: BoneWeights)
 		{
-			check(InfluenceIndex < MAX_TOTAL_INFLUENCES);
-			int32 BoneIndex = BoneWeight.GetBoneIndex();
-			if (!ensure(Deformer.InvCSRefPoseTransforms.IsValidIndex(BoneIndex)))
+			if (!ensure(NumInfluences < MAX_TOTAL_INFLUENCES))
 			{
-				UE_LOG(LogMeshModelingToolsEditor, Warning, TEXT("InitializeSkinWeights: Invalid bone index provided (%d); falling back to 0 as bone index."), BoneIndex);
-				BoneIndex = 0;
+				// cut-off weights if they exceed max influences (this shouldn't happen)
+				break;
 			}
-			const float Weight = BoneWeight.GetWeight();
-			const FVector& RefPoseVertexPosition = Deformer.RefPoseVertexPositions[VertexIndex];
-			const FTransform& InvRefPoseTransform = Deformer.InvCSRefPoseTransforms[BoneIndex];
-			const FVector& BoneLocalPositionInRefPose = InvRefPoseTransform.TransformPosition(RefPoseVertexPosition);
-			CurrentWeights[VertexIndex].Emplace(BoneIndex, BoneLocalPositionInRefPose, Weight);
-			++InfluenceIndex;
+			NormalizedWeights.Add(BoneWeight.GetBoneIndex(), BoneWeight.GetWeight());
+			++NumInfluences;
 		}
-
+		USkinWeightsPaintTool::NormalizeWeightMap(NormalizedWeights);
+		
 		// if there are no bone weights, default to root bone 
-		if (InfluenceIndex == 0)
+		if (NumInfluences == 0)
 		{
 			const FVector& RefPoseVertexPosition = Deformer.RefPoseVertexPositions[VertexIndex];
 			const FTransform& InvRefPoseTransform = Deformer.InvCSRefPoseTransforms[RootBoneIndex];
 			const FVector& BoneLocalPositionInRefPose = InvRefPoseTransform.TransformPosition(RefPoseVertexPosition);
 			CurrentWeights[VertexIndex].Emplace(RootBoneIndex, BoneLocalPositionInRefPose, FullWeight);
+			continue;
+		}
+
+		// load into the main weights data structure
+		for (TTuple<BoneIndex, float>& SingleBoneWeight : NormalizedWeights)
+		{
+			BoneIndex BoneIndex = SingleBoneWeight.Key;
+			const float Weight = SingleBoneWeight.Value;
+			if (!ensure(Deformer.InvCSRefPoseTransforms.IsValidIndex(BoneIndex)))
+			{
+				UE_LOG(LogMeshModelingToolsEditor, Warning, TEXT("InitializeSkinWeights: Invalid bone index provided (%d); falling back to 0 (root) as bone index."), BoneIndex);
+				BoneIndex = 0;
+			}
+			const FVector& RefPoseVertexPosition = Deformer.RefPoseVertexPositions[VertexIndex];
+			const FTransform& InvRefPoseTransform = Deformer.InvCSRefPoseTransforms[BoneIndex];
+			const FVector& BoneLocalPositionInRefPose = InvRefPoseTransform.TransformPosition(RefPoseVertexPosition);
+			CurrentWeights[VertexIndex].Emplace(BoneIndex, BoneLocalPositionInRefPose, Weight);
 		}
 	}
 	
@@ -4072,7 +4088,6 @@ void USkinWeightsPaintTool::NormalizeWeightMap(TMap<BoneIndex, float>& InOutWeig
 	{
 		Weight.Value /= TotalWeight > SMALL_NUMBER ? TotalWeight : 1.f;
 	}
-
 }
 
 void USkinWeightsPaintTool::AccumulateWeights(
diff --git a/Engine/Plugins/Experimental/MeshModelingToolsetExp/Source/MeshModelingToolsEditorOnlyExp/Public/SkeletalMesh/SkinWeightsPaintTool.h b/Engine/Plugins/Experimental/MeshModelingToolsetExp/Source/MeshModelingToolsEditorOnlyExp/Public/SkeletalMesh/SkinWeightsPaintTool.h
--- a/Engine/Plugins/Experimental/MeshModelingToolsetExp/Source/MeshModelingToolsEditorOnlyExp/Public/SkeletalMesh/SkinWeightsPaintTool.h
+++ b/Engine/Plugins/Experimental/MeshModelingToolsetExp/Source/MeshModelingToolsEditorOnlyExp/Public/SkeletalMesh/SkinWeightsPaintTool.h
@@ -738,6 +738,18 @@ public:
 	void NormalizeWeights();
 	void HammerWeights();
 
+	// HELPER functions for modifying weights
+	//
+	// given a map of BoneIndex > Weight values for a single vertex, modify the map by removing the smallest weights to fit in Max Influences
+	static void TruncateWeightMap(TMap<BoneIndex, float>& InOutWeights);
+	// given a map of BoneIndex > Weight values for a single vertex, modify the weights to sum to 1
+	static void NormalizeWeightMap(TMap<BoneIndex, float>& InOutWeights);
+	// sum all the weights on all bones for a given list of vertices (results we not be normalized!)
+	static void AccumulateWeights(
+		const TArray<SkinPaintTool::VertexWeights>& AllWeights,
+		const TArray<VertexIndex>& VerticesToAccumulate,
+		TMap<BoneIndex, float>& OutWeights);
+
 	// copy paste
 	void CopyWeights();
 	void PasteWeights();
@@ -922,18 +934,6 @@ protected:
 	TArray<FName> SelectedBoneNames;
 	TArray<BoneIndex> SelectedBoneIndices;
 
-	// HELPER functions for modifying weights
-	//
-	// given a map of BoneIndex > Weight values for a single vertex, modify the map by removing the smallest weights to fit in Max Influences
-	void TruncateWeightMap(TMap<BoneIndex, float>& InOutWeights);
-	// given a map of BoneIndex > Weight values for a single vertex, modify the weights to sum to 1
-	void NormalizeWeightMap(TMap<BoneIndex, float>& InOutWeights);
-	// sum all the weights on all bones for a given list of vertices (results we not be normalized!)
-	void AccumulateWeights(
-		const TArray<SkinPaintTool::VertexWeights>& AllWeights,
-		const TArray<VertexIndex>& VerticesToAccumulate,
-		TMap<BoneIndex, float>& OutWeights);
-
 	// ISkeletalMeshEditionInterface
 	virtual void HandleSkeletalMeshModified(const TArray<FName>& InBoneNames, const ESkeletalMeshNotifyType InNotifyType) override;
 
