MeshModelingTools: Fix spline tool creating multiple undeletable meshes and allowing user to pick a preview mesh as the output.

diff --git a/Engine/Plugins/Experimental/MeshModelingToolsetExp/Source/MeshModelingToolsEditorOnlyExp/Private/DrawSplineTool.cpp b/Engine/Plugins/Experimental/MeshModelingToolsetExp/Source/MeshModelingToolsEditorOnlyExp/Private/DrawSplineTool.cpp
--- a/Engine/Plugins/Experimental/MeshModelingToolsetExp/Source/MeshModelingToolsEditorOnlyExp/Private/DrawSplineTool.cpp
+++ b/Engine/Plugins/Experimental/MeshModelingToolsetExp/Source/MeshModelingToolsEditorOnlyExp/Private/DrawSplineTool.cpp
@@ -82,33 +82,6 @@ namespace DrawSplineToolLocals
 {
 	FText AddPointTransactionName = LOCTEXT("AddPointTransactionName", "Add Point");
 
-	// Helper for destroying actors in a "best practice" way. This is more valuable when destroying non-preview
-	//  actors to make sure that things are properly deselected, etc., so it's not critical in this particular tool,
-	//  but we'll use it.
-	// TODO: Move this to some helper file for use in multiple tools.
-	void DestroyActor(AActor* Actor)
-	{
-		if (!Actor)
-		{
-			return;
-		}
-		if (UWorld* ActorWorld = Actor->GetWorld())
-		{
-			if (GIsEditor && GUnrealEd)
-			{
-				GUnrealEd->DeleteActors(TArray{ Actor }, ActorWorld, GUnrealEd->GetSelectedActors()->GetElementSelectionSet());
-			}
-			else
-			{
-				ActorWorld->DestroyActor(Actor);
-			}
-		}
-		else
-		{
-			Actor->Destroy();
-		}
-	};
-
 	USplineComponent* CreateNewSplineInActor(AActor* Actor, bool bTransact = false, bool bSetAsRoot = false)
 	{
 		if (!ensure(Actor))
@@ -352,7 +325,7 @@ void UDrawSplineTool::Setup()
 	Settings->RestoreProperties(this);
 	AddToolPropertySource(Settings);
 
-	Settings->TargetActor = SelectedActor;
+	Settings->TargetActor = StartupSelectedActor;
 
 	SetToolDisplayName(LOCTEXT("DrawSplineToolName", "Draw Spline"));
 	GetToolManager()->DisplayMessage(
@@ -388,8 +361,23 @@ void UDrawSplineTool::Setup()
 	Settings->WatchProperty(Settings->OutputMode, [this](EDrawSplineOutputMode) {
 		TransitionOutputMode();
 	});
-	Settings->WatchProperty(Settings->TargetActor, [this](TWeakObjectPtr<AActor>) {
-		TransitionOutputMode();
+	TargetActorWatcherID = Settings->WatchProperty(Settings->TargetActor, [this](TWeakObjectPtr<AActor>) {
+		// It's possible for the user to use the actor picker to click on our preview, which we don't
+		//  want to be pickable via the actor picker... There doesn't currently seem to be a way prevent
+		//  that, so for now we'll just catch this case and keep whatever the previous value was.
+		if (Settings->TargetActor == PreviewActor)
+		{
+			Settings->TargetActor = IsValid(PreviousTargetActor) ? PreviousTargetActor : nullptr;
+			Settings->SilentUpdateWatcherAtIndex(TargetActorWatcherID);
+		}
+
+		if (PreviousTargetActor != Settings->TargetActor)
+		{
+			// Don't set PreviousTargetActor here because it needs to be made visible, etc inside
+			//  TransitionOutputMode
+			
+			TransitionOutputMode();
+		}
 	});
 	Settings->WatchProperty(Settings->ExistingSplineIndexToReplace, [this](int32) {
 		TransitionOutputMode();
@@ -575,7 +563,7 @@ void UDrawSplineTool::TransitionOutputMode()
 	// Now that we've copied over previous preview data, destroy the old previews
 	if (PreviousPreview)
 	{
-		DestroyActor(PreviousPreview);
+		PreviousPreview->Destroy();
 	}
 }
 
@@ -613,7 +601,7 @@ void UDrawSplineTool::Shutdown(EToolShutdownType ShutdownType)
 
 	if (PreviewActor)
 	{
-		DestroyActor(PreviewActor);
+		PreviewActor->Destroy();
 	}
 
 	Super::Shutdown(ShutdownType);
@@ -1138,7 +1126,7 @@ bool UDrawSplineTool::CanAccept() const
 // To be called by builder
 void UDrawSplineTool::SetSelectedActor(AActor* Actor)
 {
-	SelectedActor = Actor;
+	StartupSelectedActor = Actor;
 }
 void UDrawSplineTool::SetWorld(UWorld* World)
 {
diff --git a/Engine/Plugins/Experimental/MeshModelingToolsetExp/Source/MeshModelingToolsEditorOnlyExp/Public/DrawSplineTool.h b/Engine/Plugins/Experimental/MeshModelingToolsetExp/Source/MeshModelingToolsEditorOnlyExp/Public/DrawSplineTool.h
--- a/Engine/Plugins/Experimental/MeshModelingToolsetExp/Source/MeshModelingToolsEditorOnlyExp/Public/DrawSplineTool.h
+++ b/Engine/Plugins/Experimental/MeshModelingToolsetExp/Source/MeshModelingToolsEditorOnlyExp/Public/DrawSplineTool.h
@@ -228,7 +228,7 @@ private:
 	TWeakObjectPtr<UWorld> TargetWorld = nullptr;
 
 	// This is only used to initialize TargetActor in the settings object
-	TWeakObjectPtr<AActor> SelectedActor = nullptr;
+	TWeakObjectPtr<AActor> StartupSelectedActor = nullptr;
 
 	// The preview actor is either a APreviewGeometryActor with a spline, or a duplicate of 
 	// some target blueprint actor so that we can see the effects of the drawn spline immediately.
@@ -253,9 +253,12 @@ private:
 	void TransitionOutputMode();
 	void GenerateAsset();
 	
-	// Used to restore visibility of previous actor when switching to a different one
+	// Used to restore visibility of previous actor when switching to a different one, and to avoid switching
+	//  target actors if the new value is invalid (e.g., user clicked the preview actor with the actor picker)
 	UPROPERTY()
 	TObjectPtr<AActor> PreviousTargetActor = nullptr;
+
+	int32 TargetActorWatcherID = -1;
 	bool PreviousTargetActorVisibility = true;
 	// Used to restore visibility of previous spline when switching to a different one
 	TWeakObjectPtr<USplineComponent> HiddenSpline = nullptr;
