Skin Weight Paint Tool: improved transfer between the same Skeletal Mes / LOD / Profile

#rnx

diff --git a/Engine/Plugins/Experimental/MeshModelingToolsetExp/Source/MeshModelingToolsEditorOnlyExp/Private/SkeletalMesh/SkinWeightsPaintTool.cpp b/Engine/Plugins/Experimental/MeshModelingToolsetExp/Source/MeshModelingToolsEditorOnlyExp/Private/SkeletalMesh/SkinWeightsPaintTool.cpp
--- a/Engine/Plugins/Experimental/MeshModelingToolsetExp/Source/MeshModelingToolsEditorOnlyExp/Private/SkeletalMesh/SkinWeightsPaintTool.cpp
+++ b/Engine/Plugins/Experimental/MeshModelingToolsetExp/Source/MeshModelingToolsEditorOnlyExp/Private/SkeletalMesh/SkinWeightsPaintTool.cpp
@@ -854,26 +854,28 @@ void UWeightToolTransferManager::TransferWeights()
 	const bool bSameMesh = SourceSkeletalMesh == TargetSkeletalMesh;
 	const bool bSameLOD = GetLODId(ToolProperties->ActiveLOD) == GetLODId(ToolProperties->SourceLOD);
 	const bool bSameProfile = ToolProperties->ActiveSkinWeightProfile == ToolProperties->SourceSkinWeightProfile;
-
-	// cannot transfer between same mesh/LOD/profile (identical weights)
-	if (bSameMesh && bSameLOD && bSameProfile)
+	const bool bHasSelectedVertices = ToolProperties->EditingMode == EWeightEditMode::Mesh && !WeightTool->GetMainMeshSelector()->GetSelectedVertices().IsEmpty();
+	
+	// cannot transfer between same mesh/LOD/profile without selection (identical weights)
+	if (bSameMesh && bSameLOD && bSameProfile && !bHasSelectedVertices)
 	{
-		const FText NotificationText = LOCTEXT("IdenticalSourceAndTarget", "Cannot copy weights from the same mesh, LOD and profile. No weights were transferred.");
+		const FText NotificationText = LOCTEXT("IdenticalSourceAndTarget", "Cannot copy weights from the same mesh, LOD and profile without anything selected. No weights were transferred.");
 		ShowEditorMessage(ELogVerbosity::Error, NotificationText);
 		return;
 	}
-	
-	if (bSameMesh && bSameLOD)
+
+	// can weights be transferred copying the attribute directly?
+	if (bSameMesh && bSameLOD && !bSameProfile)
 	{
 		TransferWeightsFromSameMeshAndLOD();
 	}
 	else
 	{
-		TransferWeightsFromOtherMesh();
+		TransferWeightsFromOtherMeshOrSubset();
 	}
 }
 
-void UWeightToolTransferManager::TransferWeightsFromOtherMesh()
+void UWeightToolTransferManager::TransferWeightsFromOtherMeshOrSubset()
 {
 	using UE::Geometry::FDynamicMeshAttributeSet;
 	using UE::Geometry::FDynamicMeshVertexSkinWeightsAttribute;
@@ -891,13 +893,20 @@ void UWeightToolTransferManager::TransferWeightsFromOtherMesh()
 	const EMeshLODIdentifier TargetLODId = GetLODId(ToolProperties->ActiveLOD);
 	const EMeshLODIdentifier SourceLODId = GetLODId(ToolProperties->SourceLOD);
 
-	// if transferring between the same mesh, ensure that the LODs are different
+	// get selection
+	TArray<int32> SourceTrianglesToIsolate;
+	MeshSelector->GetSelectedTriangles(SourceTrianglesToIsolate);
+	
+	TArray<int> TargetSelectedVertices = WeightTool->GetMainMeshSelector()->GetSelectedVertices();
+	
+	// if transferring between the same mesh, ensure that the LODs are different and that the transfer is done on a subset
 	const USkeletalMesh* TargetSkeletalMesh = GetSkeletalMeshComponent(WeightTool->GetTarget())->GetSkeletalMeshAsset();
-	if (SourceSkeletalMesh == TargetSkeletalMesh)
+	const bool bSameMeshAndLOD = SourceSkeletalMesh == TargetSkeletalMesh && TargetLODId == SourceLODId;
+	if (bSameMeshAndLOD)
 	{
-		if (TargetLODId == SourceLODId)
+		if (SourceTrianglesToIsolate.IsEmpty() && TargetSelectedVertices.IsEmpty())
 		{
-			const FText NotificationText = LOCTEXT("SameLOD", "Cannot copy weights between the same LOD on the same mesh. No weights were transferred.");
+			const FText NotificationText = LOCTEXT("SameLODAndNoSelection", "Cannot copy weights between the same LOD on the same mesh without anything selected. No weights were transferred.");
 			ShowEditorMessage(ELogVerbosity::Error, NotificationText);
 			return;
 		}
@@ -908,39 +917,67 @@ void UWeightToolTransferManager::TransferWeightsFromOtherMesh()
 	FDynamicMesh3 TargetMesh = UE::ToolTarget::GetDynamicMeshCopy(WeightTool->GetTarget(), TargetParams);
 	
 	// get the source dynamic mesh and validate it
-	const FGetMeshParameters SourceParams(true, SourceLODId);
-	FDynamicMesh3 SourceMeshOrig = UE::ToolTarget::GetDynamicMeshCopy(SourceTarget, SourceParams);
-	FDynamicMesh3 SourceMesh;
-
-	// create a submesh from the selected triangles to filter the transfer so that it only copies from the selected components
-	TArray<int32> TrianglesToIsolate;
-	MeshSelector->GetSelectedTriangles(TrianglesToIsolate);
-	if (!TrianglesToIsolate.IsEmpty())
+	FDynamicMesh3 TmpSourceCopy;
+	auto GetSourceMesh = [this, bSameMeshAndLOD, SourceLODId, &TmpSourceCopy](const TArray<int32>& ToIsolate) -> const FDynamicMesh3*
 	{
-		UE::Geometry::FDynamicSubmesh3 PartialSubMesh = UE::Geometry::FDynamicSubmesh3(&SourceMeshOrig, TrianglesToIsolate);
-		SourceMesh = PartialSubMesh.GetSubmesh();
-		FDynamicMeshAttributeSet* SourceAttributes = SourceMesh.Attributes();
-		SourceAttributes->CopyBoneAttributes(*SourceMeshOrig.Attributes());
-	}
-	else
+		// use the current EditedMesh to get the current data without having to commit
+		FCleanedEditMesh* CleanedEditMesh = WeightTool->GetCurrentCleanedEditMesh();
+		if (const FDynamicMesh3* EditedMesh = bSameMeshAndLOD && CleanedEditMesh ? &CleanedEditMesh->GetOriginalMesh() : nullptr)
+		{
+			if (ToIsolate.IsEmpty())
+			{
+				return EditedMesh;
+			}
+
+			// create a sub-mesh from the selected triangles to filter the transfer so that it only copies from the selected components
+			const UE::Geometry::FDynamicSubmesh3 PartialSubMesh(EditedMesh, ToIsolate);
+			TmpSourceCopy = PartialSubMesh.GetSubmesh();			
+			FDynamicMeshAttributeSet* SourceAttributes = TmpSourceCopy.Attributes();
+			SourceAttributes->CopyBoneAttributes(*EditedMesh->Attributes());
+			return &TmpSourceCopy;
+		}
+
+		// otherwise, get the corresponding DynamicMesh for that LOD
+		const FGetMeshParameters SourceParams(true, SourceLODId);
+		FDynamicMesh3 SourceMeshOrig = UE::ToolTarget::GetDynamicMeshCopy(SourceTarget, SourceParams);
+		
+		// create a sub-mesh from the selected triangles to filter the transfer so that it only copies from the selected components
+		if (!ToIsolate.IsEmpty())
+		{
+			const UE::Geometry::FDynamicSubmesh3 PartialSubMesh(&SourceMeshOrig, ToIsolate);
+			TmpSourceCopy = PartialSubMesh.GetSubmesh();
+			FDynamicMeshAttributeSet* SourceAttributes = TmpSourceCopy.Attributes();
+			SourceAttributes->CopyBoneAttributes(*SourceMeshOrig.Attributes());
+		}
+		else
+		{
+			TmpSourceCopy = MoveTemp(SourceMeshOrig);
+		}
+		return &TmpSourceCopy;
+	};
+	
+	const FDynamicMesh3* SourceMesh = GetSourceMesh(SourceTrianglesToIsolate);
+	if (!ensure(SourceMesh))
 	{
-		SourceMesh = MoveTemp(SourceMeshOrig);
+		const FText NotificationText = LOCTEXT("NoSourceMesh", "Cannot retrieve any source mesh form the current properties.");
+		ShowEditorMessage(ELogVerbosity::Error, NotificationText);
+		return;
 	}
 	
-	if (!SourceMesh.HasAttributes() || !SourceMesh.Attributes()->HasBones())
+	if (!SourceMesh->HasAttributes() || !SourceMesh->Attributes()->HasBones())
 	{
 		const FText NotificationText = LOCTEXT("NoWeightsFoundInTransfer", "No skin weights were found in the source skeletal mesh. No weights were transferred.");
 		ShowEditorMessage(ELogVerbosity::Error, NotificationText);
 		return;
 	}
-	if (SourceMesh.Attributes()->GetNumBones() == 0)
+	if (SourceMesh->Attributes()->GetNumBones() == 0)
 	{
 		const FText NotificationText = LOCTEXT("NoBonesFoundInTransfer", "No bones were found in the source skeletal mesh. No weights were transferred.");
 		ShowEditorMessage(ELogVerbosity::Error, NotificationText);
 		return;
 	}
 
-	FTransferBoneWeights TransferBoneWeights(&SourceMesh, ToolProperties->SourceSkinWeightProfile);
+	FTransferBoneWeights TransferBoneWeights(SourceMesh, ToolProperties->SourceSkinWeightProfile);
 	TransferBoneWeights.TransferMethod = FTransferBoneWeights::ETransferBoneWeightsMethod::InpaintWeights;
 	
 	if (!TargetMesh.HasAttributes())
@@ -951,7 +988,7 @@ void UWeightToolTransferManager::TransferWeightsFromOtherMesh()
 	FDynamicMeshAttributeSet* TargetAttributes = TargetMesh.Attributes();
 	if (!TargetAttributes->HasBones())
 	{
-		TargetAttributes->CopyBoneAttributes(*SourceMesh.Attributes());
+		TargetAttributes->CopyBoneAttributes(*SourceMesh->Attributes());
 	}
 	else
 	{
@@ -981,7 +1018,7 @@ void UWeightToolTransferManager::TransferWeightsFromOtherMesh()
 
 	if (ToolProperties->EditingMode == EWeightEditMode::Mesh)
 	{
-		TransferBoneWeights.TargetVerticesSubset = WeightTool->GetMainMeshSelector()->GetSelectedVertices();
+		TransferBoneWeights.TargetVerticesSubset = MoveTemp(TargetSelectedVertices);
 	}
 
 	const FName TargetProfile = ToolProperties->GetActiveSkinWeightProfile();
@@ -4539,6 +4576,8 @@ FCleanedEditMesh::FCleanedEditMesh(
 	const FDynamicMesh3& InDynamicMesh,
 	const FMeshDescription& InMeshDescription) : OriginalDynamicMesh(InDynamicMesh), OriginalMeshDescription(InMeshDescription)
 {
+	using namespace UE::Geometry;
+	
 	// store copy of the original mesh data
 	OriginalDynamicMesh = InDynamicMesh;
 	OriginalMeshDescription = InMeshDescription;
@@ -4558,12 +4597,12 @@ FCleanedEditMesh::FCleanedEditMesh(
 	{
 		AllTris.Add(TriID);
 	}
-	CleanedSubMesh = UE::Geometry::FDynamicSubmesh3(&OriginalDynamicMesh, AllTris);
+	CleanedSubMesh = FDynamicSubmesh3(&OriginalDynamicMesh, AllTris);
 
 	// remap the non-manifold mapping from full mesh to partial mesh
-	UE::Geometry::FNonManifoldMappingSupport OriginalNonManifoldMapping(OriginalDynamicMesh);
+	FNonManifoldMappingSupport OriginalNonManifoldMapping(OriginalDynamicMesh);
 	TArray<int32> SubMeshVertexToNonManifoldVertexIDMap;
-	const FDynamicMesh3& SubMeshDynamicMesh = CleanedSubMesh.GetSubmesh();
+	FDynamicMesh3& SubMeshDynamicMesh = CleanedSubMesh.GetSubmesh();
 	for (int32 SubMeshVertexID=0; SubMeshVertexID<SubMeshDynamicMesh.VertexCount(); ++SubMeshVertexID)
 	{
 		int32 BaseMeshVertexID = CleanedSubMesh.MapVertexToBaseMesh(SubMeshVertexID);
@@ -4576,8 +4615,20 @@ FCleanedEditMesh::FCleanedEditMesh(
 		SubMeshVertexToNonManifoldVertexIDMap.Add(SrcSubMeshVertexID);
 	}
 
+	// copy bone attributes
+	const FDynamicMeshAttributeSet* OrigAttributes = InDynamicMesh.Attributes();
+	if (OrigAttributes && OrigAttributes->HasBones())
+	{
+		if (!SubMeshDynamicMesh.HasAttributes())
+		{
+			SubMeshDynamicMesh.EnableAttributes();
+		}
+		FDynamicMeshAttributeSet* SubMeshAttributes = SubMeshDynamicMesh.Attributes();
+		SubMeshAttributes->CopyBoneAttributes(*InDynamicMesh.Attributes());
+	}
+
 	// replace the non-manifold vertex map in the cleaned mesh attributes
-	UE::Geometry::FNonManifoldMappingSupport CleanedNonManifoldMapping(CleanedSubMesh.GetSubmesh());
+	FNonManifoldMappingSupport CleanedNonManifoldMapping(CleanedSubMesh.GetSubmesh());
 	CleanedNonManifoldMapping.AttachNonManifoldVertexMappingData(SubMeshVertexToNonManifoldVertexIDMap,CleanedSubMesh.GetSubmesh());
 
 	// create mesh description for sub-mesh
diff --git a/Engine/Plugins/Experimental/MeshModelingToolsetExp/Source/MeshModelingToolsEditorOnlyExp/Public/SkeletalMesh/SkinWeightsPaintTool.h b/Engine/Plugins/Experimental/MeshModelingToolsetExp/Source/MeshModelingToolsEditorOnlyExp/Public/SkeletalMesh/SkinWeightsPaintTool.h
--- a/Engine/Plugins/Experimental/MeshModelingToolsetExp/Source/MeshModelingToolsEditorOnlyExp/Public/SkeletalMesh/SkinWeightsPaintTool.h
+++ b/Engine/Plugins/Experimental/MeshModelingToolsetExp/Source/MeshModelingToolsEditorOnlyExp/Public/SkeletalMesh/SkinWeightsPaintTool.h
@@ -582,7 +582,7 @@ public:
 private:
 	
 	// actually run the weight transfer to copy weights from the source to the target
-	void TransferWeightsFromOtherMesh();
+	void TransferWeightsFromOtherMeshOrSubset();
 
 	// actually run the weight transfer to copy weights from the source to the target
 	void TransferWeightsFromSameMeshAndLOD();
