Fixes bug in Mirror weights tool causing non-normalized results.

diff --git a/Engine/Plugins/Experimental/MeshModelingToolsetExp/Source/MeshModelingToolsEditorOnlyExp/Private/SkeletalMesh/SkinWeightsPaintTool.cpp b/Engine/Plugins/Experimental/MeshModelingToolsetExp/Source/MeshModelingToolsEditorOnlyExp/Private/SkeletalMesh/SkinWeightsPaintTool.cpp
--- a/Engine/Plugins/Experimental/MeshModelingToolsetExp/Source/MeshModelingToolsEditorOnlyExp/Private/SkeletalMesh/SkinWeightsPaintTool.cpp
+++ b/Engine/Plugins/Experimental/MeshModelingToolsetExp/Source/MeshModelingToolsEditorOnlyExp/Private/SkeletalMesh/SkinWeightsPaintTool.cpp
@@ -3471,7 +3471,7 @@ void FSkinMirrorData::EnsureMirrorDataIsUpdated(
 			SearchRadius += HashGridCellSize;
 
 			// forcibly break out if search radius gets bigger than the maximum search radius
-			constexpr float MaxSearchRadius = 5.f; // TODO we may want to expose this value to the user...
+			static float MaxSearchRadius = 15.f; // TODO we may want to expose this value to the user...
 			if (SearchRadius >= MaxSearchRadius)
 			{
 				break;
@@ -3546,7 +3546,9 @@ void USkinWeightsPaintTool::MirrorWeights(EAxis::Type Axis, EMirrorDirection Dir
 	for (const VertexIndex SelectedVertex : SelectedVertices)
 	{
 		int32 TargetVertexID = INDEX_NONE;
-		if (VertexMirrorMap.Contains(SelectedVertex))
+		const bool bIsOnTargetSide = VertexMirrorMap.Contains(SelectedVertex);
+		
+		if (bIsOnTargetSide)
 		{
 			// vertex is located across the mirror plane (target side, to copy TO)
 			TargetVertexID = SelectedVertex;
@@ -3567,8 +3569,12 @@ void USkinWeightsPaintTool::MirrorWeights(EAxis::Type Axis, EMirrorDirection Dir
 		// selected vertex did not have a mirrored equivalent
 		if (TargetVertexID == INDEX_NONE)
 		{
-			MissingVertices.Add(TargetVertexID);
-			continue;;
+			if (bIsOnTargetSide)
+			{
+				MissingVertices.Add(TargetVertexID);
+			}
+			
+			continue;
 		}
 		
 		// add to the list of target vertices to set weights on
@@ -3577,6 +3583,8 @@ void USkinWeightsPaintTool::MirrorWeights(EAxis::Type Axis, EMirrorDirection Dir
 	
 	// spin through all target vertices to mirror and copy weights from source
 	FMultiBoneWeightEdits WeightEditsFromMirroring;
+	TMap<BoneIndex, float> NewBoneWeights;
+	NewBoneWeights.Reserve(MAX_TOTAL_INFLUENCES);
 	for (const VertexIndex TargetVertexID : TargetVertices)
 	{
 		const VertexIndex SourceVertexID = VertexMirrorMap[TargetVertexID];
@@ -3591,18 +3599,32 @@ void USkinWeightsPaintTool::MirrorWeights(EAxis::Type Axis, EMirrorDirection Dir
 		// remove all weight on vertex
 		for (const FVertexBoneWeight& TargetBoneWeight : Weights.PreChangeWeights[TargetVertexID])
 		{
-			constexpr bool bPruneInfluence = false;
-			constexpr float NewWeight = 0.f;
+			static bool bPruneInfluence = true;
+			static float NewWeight = 0.f;
 			WeightEditsFromMirroring.MergeSingleEdit(TargetBoneWeight.BoneID, TargetVertexID, NewWeight, bPruneInfluence, Weights.PreChangeWeights);
 		}
 
 		// copy source weights, but with mirrored bones
+		// NOTE: we have to normalize here because it's possible that while searching for mirrored bones, multiple source bones will be
+		// mapped to the same target bone. When that happens, only the last applied weight on that bone would be recorded in an edit
+		// If that happens, the final weight may not sum to 1.
+		NewBoneWeights.Reset();
 		for (const FVertexBoneWeight& SourceBoneWeight : Weights.PreChangeWeights[SourceVertexID])
 		{
 			const BoneIndex MirroredBoneIndex = BoneMap[SourceBoneWeight.BoneID];
-			constexpr bool bPruneInfluence = false;
 			const float NewWeight = SourceBoneWeight.Weight;
-			WeightEditsFromMirroring.MergeSingleEdit(MirroredBoneIndex, TargetVertexID, NewWeight, bPruneInfluence, Weights.PreChangeWeights);
+			float& Weight = NewBoneWeights.FindOrAdd(MirroredBoneIndex, 0.0f);
+			Weight += NewWeight;
+		}
+		NormalizeWeightMap(NewBoneWeights);
+
+		// apply weight edits
+		for (const TPair<BoneIndex, float>& NewBoneWeight : NewBoneWeights)
+		{
+			const BoneIndex BoneID = NewBoneWeight.Key;
+			const float NewWeight = NewBoneWeight.Value;
+			static bool bPruneInfluence = false;
+			WeightEditsFromMirroring.MergeSingleEdit(BoneID, TargetVertexID, NewWeight, bPruneInfluence, Weights.PreChangeWeights);
 		}
 	}
 
