Skin Weight Paint Tool: make sure to register the skin weight profile when copying weights from the cleaned mesh to the original one.

#rnx

diff --git a/Engine/Plugins/Experimental/MeshModelingToolsetExp/Source/MeshModelingToolsEditorOnlyExp/Private/SkeletalMesh/SkinWeightsPaintTool.cpp b/Engine/Plugins/Experimental/MeshModelingToolsetExp/Source/MeshModelingToolsEditorOnlyExp/Private/SkeletalMesh/SkinWeightsPaintTool.cpp
--- a/Engine/Plugins/Experimental/MeshModelingToolsetExp/Source/MeshModelingToolsEditorOnlyExp/Private/SkeletalMesh/SkinWeightsPaintTool.cpp
+++ b/Engine/Plugins/Experimental/MeshModelingToolsetExp/Source/MeshModelingToolsEditorOnlyExp/Private/SkeletalMesh/SkinWeightsPaintTool.cpp
@@ -843,7 +843,7 @@ void UWeightToolTransferManager::TransferWeights()
 	{
 		// to transfer weights from another mesh we need a source mesh
 		// (UI should prevent us from getting here)
-		const FText NotificationText = LOCTEXT("NoSourceMesh", "No source skeletal mesh specified. No weights were transferred.");
+		const FText NotificationText = LOCTEXT("NoSourceTarget", "No source skeletal mesh specified. No weights were transferred.");
 		ShowEditorMessage(ELogVerbosity::Error, NotificationText);
 		return;
 	}
@@ -4665,7 +4665,7 @@ FCleanedEditMesh::FCleanedEditMesh(
 	DnyToDescConverter.Convert(&CleanedSubMesh.GetSubmesh(), *CleanedSubMeshDescription, bCopyTangents);
 }
 
-void FCleanedEditMesh::CopyWeightsToOriginalMesh(FName Profile)
+void FCleanedEditMesh::CopyWeightsToOriginalMesh(const FName InProfile)
 {
 	// if the mesh was not duplicated and cleaned, we have nothing to copy from
 	if (!CleanedSubMeshDescription)
@@ -4673,11 +4673,23 @@ void FCleanedEditMesh::CopyWeightsToOriginalMesh(FName Profile)
 		return;
 	}
 
-	FSkeletalMeshAttributes SubMeshAttrs(*CleanedSubMeshDescription);
-	FSkinWeightsVertexAttributesRef SubMeshWeightAttrs = SubMeshAttrs.GetVertexSkinWeights(Profile);
+	// ensure that this profile is available on the cleaned mesh description
+	const FSkeletalMeshConstAttributes SubMeshAttrs(*CleanedSubMeshDescription);
+	const FSkinWeightsVertexAttributesConstRef SubMeshWeightAttrs = SubMeshAttrs.GetVertexSkinWeights(InProfile);
+	if (!SubMeshWeightAttrs.IsValid())
+	{
+		const FText NotificationText = FText::Format(LOCTEXT("NoSubProfile", "Failed copying skin weights, profile '{0}' not found on sub mesh."), FText::FromName(InProfile));
+		ShowEditorMessage(ELogVerbosity::Error, NotificationText);
+		return;
+	}
 
-	FSkeletalMeshAttributes BaseMeshAttrs(OriginalMeshDescription);
-	FSkinWeightsVertexAttributesRef BaseMeshWeightAttrs = BaseMeshAttrs.GetVertexSkinWeights(Profile);
+	FSkinWeightsVertexAttributesRef BaseMeshWeightAttrs = GetOrCreateSkinWeightsAttribute(OriginalMeshDescription, InProfile);
+	if (!BaseMeshWeightAttrs.IsValid())
+	{
+		const FText NotificationText = FText::Format(LOCTEXT("NoBaseProfile", "Failed copying skin weights, profile '{0}' not found on base mesh."), FText::FromName(InProfile));
+		ShowEditorMessage(ELogVerbosity::Error, NotificationText);
+		return;
+	}
 
 	UE::AnimationCore::FBoneWeightsSettings Settings;
 	Settings.SetNormalizeType(UE::AnimationCore::EBoneWeightNormalizeType::None);
@@ -4690,7 +4702,7 @@ void FCleanedEditMesh::CopyWeightsToOriginalMesh(FName Profile)
 	{
 		BoneWeightsToApply.Reset();
 		
-		FVertexBoneWeights SubVertexWeights = SubMeshWeightAttrs.Get(SubVertexID);
+		const FVertexBoneWeightsConst SubVertexWeights = SubMeshWeightAttrs.Get(SubVertexID);
 		for (UE::AnimationCore::FBoneWeight SingleBoneWeight : SubVertexWeights)
 		{
 			BoneWeightsToApply.Add(UE::AnimationCore::FBoneWeight(SingleBoneWeight.GetBoneIndex(), SingleBoneWeight.GetWeight()));
diff --git a/Engine/Plugins/Experimental/MeshModelingToolsetExp/Source/MeshModelingToolsEditorOnlyExp/Public/SkeletalMesh/SkinWeightsPaintTool.h b/Engine/Plugins/Experimental/MeshModelingToolsetExp/Source/MeshModelingToolsEditorOnlyExp/Public/SkeletalMesh/SkinWeightsPaintTool.h
--- a/Engine/Plugins/Experimental/MeshModelingToolsetExp/Source/MeshModelingToolsEditorOnlyExp/Public/SkeletalMesh/SkinWeightsPaintTool.h
+++ b/Engine/Plugins/Experimental/MeshModelingToolsetExp/Source/MeshModelingToolsEditorOnlyExp/Public/SkeletalMesh/SkinWeightsPaintTool.h
@@ -678,7 +678,7 @@ public:
 struct MESHMODELINGTOOLSEDITORONLYEXP_API FCleanedEditMesh
 {
 	FCleanedEditMesh(const FDynamicMesh3& InDynamicMesh, const FMeshDescription& InMeshDescription);
-	void CopyWeightsToOriginalMesh(FName Profile);
+	void CopyWeightsToOriginalMesh(const FName InProfile);
 	FDynamicMesh3& GetEditableMesh();
 	FMeshDescription& GetEditableMeshDescription();
 	FDynamicMesh3& GetOriginalMesh();
