MeshModelingTools: Add an "isolate selection" option to the attribute paint tool.

diff --git a/Engine/Plugins/Experimental/MeshModelingToolsetExp/Source/MeshModelingToolsExp/Private/MeshAttributePaintTool.cpp b/Engine/Plugins/Experimental/MeshModelingToolsetExp/Source/MeshModelingToolsExp/Private/MeshAttributePaintTool.cpp
--- a/Engine/Plugins/Experimental/MeshModelingToolsetExp/Source/MeshModelingToolsExp/Private/MeshAttributePaintTool.cpp
+++ b/Engine/Plugins/Experimental/MeshModelingToolsetExp/Source/MeshModelingToolsExp/Private/MeshAttributePaintTool.cpp
@@ -5,6 +5,8 @@
 #include "ToolBuilderUtil.h"
 #include "Drawing/MeshDebugDrawing.h"
 #include "ToolSetupUtil.h"
+#include "Selection/StoredMeshSelectionUtil.h" // GetCurrentGeometrySelectionForTarget
+#include "Selections/GeometrySelectionUtil.h"
 #include "Selections/MeshConnectedComponents.h"
 
 #include "MeshDescription.h"
@@ -168,6 +170,23 @@ UMeshSurfacePointTool* UMeshAttributePaintToolBuilder::CreateNewTool(const FTool
 	return SelectionTool;
 }
 
+// TODO: May want a base class for brush tools with selection?
+void UMeshAttributePaintToolBuilder::InitializeNewTool(UMeshSurfacePointTool* NewTool, const FToolBuilderState& SceneState) const
+{
+	Super::InitializeNewTool(NewTool, SceneState);
+
+	UMeshAttributePaintTool* Tool = Cast< UMeshAttributePaintTool>(NewTool);
+	if (ensure(Tool && Tool->GetTarget()))
+	{
+		UE::Geometry::FGeometrySelection Selection;
+		bool bHaveSelection = UE::Geometry::GetCurrentGeometrySelectionForTarget(SceneState, Tool->GetTarget(), Selection);
+		if (bHaveSelection)
+		{
+			Tool->SetGeometrySelection(MoveTemp(Selection));
+		}
+	}
+}
+
 bool UMeshAttributePaintToolBuilder::CanBuildTool(const FToolBuilderState& SceneState) const
 {
 	return UMeshSurfacePointMeshEditingToolBuilder::CanBuildTool(SceneState) &&
@@ -192,6 +211,11 @@ void UMeshAttributePaintTool::SetWorld(UWorld* World)
 	this->TargetWorld = World;
 }
 
+void UMeshAttributePaintTool::SetGeometrySelection(const UE::Geometry::FGeometrySelection& SelectionIn)
+{
+	GeometrySelection = SelectionIn;
+}
+
 
 
 void UMeshAttributePaintTool::Setup()
@@ -219,14 +243,40 @@ void UMeshAttributePaintTool::Setup()
 	//PreviewMesh->EnableWireframe(SelectionProps->bShowWireframe);
 	PreviewMesh->SetShadowsEnabled(false);
 
-	// enable vtx colors on preview mesh
-	PreviewMesh->EditMesh([](FDynamicMesh3& Mesh)
+	BrushActionProps->bToolHasSelection = GeometrySelection.IsSet();
+	SelectionTids.Reset();
+	SelectionVids.Reset();
+	
+	PreviewMesh->EditMesh([this](FDynamicMesh3& Mesh)
 	{
+		// enable vtx colors on preview mesh
 		Mesh.EnableAttributes();
 		Mesh.Attributes()->DisablePrimaryColors();
 		Mesh.Attributes()->EnablePrimaryColors();
 		// Create an overlay that has no split elements, init with zero value.
 		Mesh.Attributes()->PrimaryColors()->CreateFromPredicate([](int ParentVID, int TriIDA, int TriIDB){return true;}, 0.f);
+
+		if (GeometrySelection.IsSet())
+		{
+			UE::Geometry::EnumerateSelectionTriangles(GeometrySelection.GetValue(), *PreviewMesh->GetMesh(),
+				[this, &Mesh](int32 TriangleID)
+			{
+				SelectionTids.Add(TriangleID);
+				FIndex3i Triangle = Mesh.GetTriangle(TriangleID);
+				for (int i = 0; i < 3; ++i)
+				{
+					SelectionVids.Add(Triangle[i]);
+				}
+			});
+
+			PreviewMesh->EnableSecondaryTriangleBuffers([this](const UE::Geometry::FDynamicMesh3*, int32 Tid)
+			{
+				return !SelectionTids.Contains(Tid);
+			});
+
+			PreviewMesh->SetSecondaryBuffersVisibility(!BrushActionProps->bIsolateGeometrySelection);
+		}
+
 	});
 
 	// build octree
@@ -270,6 +320,11 @@ void UMeshAttributePaintTool::Setup()
 		[this](int32 NewValue) { PendingNewSelectedIndex = NewValue; }, AttribProps->GetSelectedAttributeIndex());
 
 	bVisibleAttributeValid = false;
+
+	BrushActionProps->WatchProperty(BrushActionProps->bIsolateGeometrySelection, [this](bool)
+	{
+		PreviewMesh->SetSecondaryBuffersVisibility(!BrushActionProps->bIsolateGeometrySelection);
+	});
 }
 
 
@@ -391,6 +446,16 @@ bool UMeshAttributePaintTool::OnUpdateHover(const FInputDeviceRay& DevicePos)
 	return true;
 }
 
+bool UMeshAttributePaintTool::HitTest(const FRay& Ray, FHitResult& OutHit)
+{
+	TFunction<bool(int32)> Filter = nullptr;
+	if (ShouldFilterTriangles())
+	{
+		Filter = [this](int32 Tid) { return SelectionTids.Contains(Tid); };
+	}
+		
+	return PreviewMesh->FindRayIntersection(FRay3d(Ray), OutHit, Filter);
+}
 
 
 void UMeshAttributePaintTool::CalculateVertexROI(const FBrushStampData& Stamp, TArray<int>& VertexROI)
@@ -402,9 +467,24 @@ void UMeshAttributePaintTool::CalculateVertexROI(const FBrushStampData& Stamp, T
 	float RadiusSqr = Radius * Radius;
 	const FDynamicMesh3* Mesh = PreviewMesh->GetPreviewDynamicMesh();
 	FAxisAlignedBox3d QueryBox(StampPosLocal, Radius);
-	VerticesOctree.RangeQuery(QueryBox,
-		[&](int32 VertexID) { return DistanceSquared(Mesh->GetVertex(VertexID), StampPosLocal) < RadiusSqr; },
-		VertexROI);
+	
+	if (ShouldFilterTriangles() && ensure(!SelectionVids.IsEmpty()))
+	{
+		VerticesOctree.RangeQuery(QueryBox,
+			[this, Mesh, &StampPosLocal, RadiusSqr](int32 VertexID)
+		{ 
+			return SelectionVids.Contains(VertexID) 
+				&& DistanceSquared(Mesh->GetVertex(VertexID), StampPosLocal) < RadiusSqr; 
+		}, VertexROI);
+	}
+	else
+	{
+		VerticesOctree.RangeQuery(QueryBox,
+			[Mesh, &StampPosLocal, RadiusSqr](int32 VertexID) 
+		{ 
+			return DistanceSquared(Mesh->GetVertex(VertexID), StampPosLocal) < RadiusSqr; 
+		}, VertexROI);
+	}
 }
 
 
@@ -642,7 +722,20 @@ void UMeshAttributePaintTool::ApplyStamp_FloodFill(const FBrushStampData& Stamp,
 	TSet<int32> RemainingTriangles;
 	for (int32 vid : ActionData.ROIVertices)
 	{
-		CurrentMesh->EnumerateVertexTriangles(vid, [&](int32 tid) { RemainingTriangles.Add(tid); });
+		if (ShouldFilterTriangles())
+		{
+			CurrentMesh->EnumerateVertexTriangles(vid, [&](int32 tid) 
+			{ 
+				if (SelectionTids.Contains(tid))
+				{ 
+					RemainingTriangles.Add(tid); 
+				}
+			});
+		}
+		else
+		{
+			CurrentMesh->EnumerateVertexTriangles(vid, [&](int32 tid) { RemainingTriangles.Add(tid); });
+		}
 	}
 
 	float SetValue = BrushProperties->BrushStrength * CurrentValueRange.Length();
@@ -663,7 +756,13 @@ void UMeshAttributePaintTool::ApplyStamp_FloodFill(const FBrushStampData& Stamp,
 		InputTriROI.Reset();
 		// get a single set element via an iterator
 		InputTriROI.Add(*RemainingTriangles.CreateConstIterator());
-		FMeshConnectedComponents::GrowToConnectedTriangles(CurrentMesh, InputTriROI, OutputTriROI, &QueueTempBuffer, &DoneTempBuffer);
+		TFunction<bool(int32, int32)> TidFilter = [](int32, int32){ return true; };
+		if (ShouldFilterTriangles())
+		{
+			TidFilter = [this](int32 Index, int32 Tid) { return SelectionTids.Contains(Tid); };
+		}
+		FMeshConnectedComponents::GrowToConnectedTriangles(CurrentMesh, InputTriROI, OutputTriROI, 
+			&QueueTempBuffer, &DoneTempBuffer, TidFilter);
 		for (int32 tid : OutputTriROI)
 		{
 			RemainingTriangles.Remove(tid);
@@ -810,6 +909,14 @@ void UMeshAttributePaintTool::ExternalUpdateValues(int32 AttribIndex, const TArr
 	}
 }
 
+bool UMeshAttributePaintTool::ShouldFilterTriangles() const
+{
+	return BrushActionProps->bIsolateGeometrySelection 
+		// Could check GeometrySelection.IsSet(), but might as well check SelectionTids instead
+		//  which will be nonempty if there was a selection, since that's what we'll actually use.
+		&& !SelectionTids.IsEmpty();
+}
+
 
 
 
diff --git a/Engine/Plugins/Experimental/MeshModelingToolsetExp/Source/MeshModelingToolsExp/Public/MeshAttributePaintTool.h b/Engine/Plugins/Experimental/MeshModelingToolsetExp/Source/MeshModelingToolsExp/Public/MeshAttributePaintTool.h
--- a/Engine/Plugins/Experimental/MeshModelingToolsetExp/Source/MeshModelingToolsExp/Public/MeshAttributePaintTool.h
+++ b/Engine/Plugins/Experimental/MeshModelingToolsetExp/Source/MeshModelingToolsExp/Public/MeshAttributePaintTool.h
@@ -8,13 +8,13 @@
 #include "Changes/IndexedAttributeChange.h"
 #include "DynamicMesh/DynamicVerticesOctree3.h"
 #include "MeshDescription.h"
-#include "MeshAttributePaintTool.generated.h"
+#include "Selections/GeometrySelection.h"
 
+#include "MeshAttributePaintTool.generated.h"
 
 struct FMeshDescription;
 class UMeshAttributePaintTool;
 
-
 /**
  * Maps float values to linear color ramp.
  */
@@ -92,6 +92,7 @@ public:
 	TUniqueFunction<TUniquePtr<FFloatAttributeColorMapper>()> ColorMapFactory;
 
 	virtual UMeshSurfacePointTool* CreateNewTool(const FToolBuilderState& SceneState) const override;
+	virtual void InitializeNewTool(UMeshSurfacePointTool* NewTool, const FToolBuilderState& SceneState) const override;
 	virtual bool CanBuildTool(const FToolBuilderState& SceneState) const override;
 
 protected:
@@ -116,8 +117,21 @@ class MESHMODELINGTOOLSEXP_API UMeshAttributePaintBrushOperationProperties : pub
 {
 	GENERATED_BODY()
 public:
+	/** 
+	 * If the tool was started with a mesh element selection, this setting hides everything
+	 *  except that selection, to make painting it easier. Requires that a mesh element
+	 *  selection exist on tool start.
+	 */
+	UPROPERTY(EditAnywhere, Category = Attribute, meta = (EditCondition = "bToolHasSelection", HideEditConditionToggle))
+	bool bIsolateGeometrySelection = false;
+	
+	//~ For the tool to set, to enable/disable bIsolateGeometrySelection
+	UPROPERTY()
+	bool bToolHasSelection = false;
+
 	UPROPERTY(EditAnywhere, Category = Attribute)
 	EBrushActionMode BrushAction = EBrushActionMode::Paint;
+
 };
 
 
@@ -210,6 +224,7 @@ class MESHMODELINGTOOLSEXP_API UMeshAttributePaintTool : public UDynamicMeshBrus
 
 public:
 	virtual void SetWorld(UWorld* World);
+	void SetGeometrySelection(const UE::Geometry::FGeometrySelection& SelectionIn);
 
 	virtual void RegisterActions(FInteractiveToolActionSet& ActionSet) override;
 
@@ -224,6 +239,9 @@ public:
 	virtual void OnUpdateDrag(const FRay& Ray) override;
 	virtual void OnEndDrag(const FRay& Ray) override;
 	virtual bool OnUpdateHover(const FInputDeviceRay& DevicePos) override;
+	
+	// UMeshSurfacePointTool
+	virtual bool HitTest(const FRay& Ray, FHitResult& OutHit) override;
 
 	virtual void RequestAction(EMeshAttributePaintToolActions ActionType);
 
@@ -316,6 +334,12 @@ protected:
 	void BeginChange();
 	TUniquePtr<FMeshAttributePaintChange> EndChange();
 	void ExternalUpdateValues(int32 AttribIndex, const TArray<int32>& VertexIndices, const TArray<float>& NewValues);
+
+private:
+	TOptional<UE::Geometry::FGeometrySelection> GeometrySelection;
+	TSet<int32> SelectionTids;
+	TSet<int32> SelectionVids;
+	bool ShouldFilterTriangles() const;
 };
 
 
diff --git a/Engine/Plugins/Runtime/MeshModelingToolset/Source/ModelingComponents/Private/PreviewMesh.cpp b/Engine/Plugins/Runtime/MeshModelingToolset/Source/ModelingComponents/Private/PreviewMesh.cpp
--- a/Engine/Plugins/Runtime/MeshModelingToolset/Source/ModelingComponents/Private/PreviewMesh.cpp
+++ b/Engine/Plugins/Runtime/MeshModelingToolset/Source/ModelingComponents/Private/PreviewMesh.cpp
@@ -415,13 +415,17 @@ TUniquePtr<FDynamicMesh3> UPreviewMesh::ExtractPreviewMesh() const
 
 
 
-bool UPreviewMesh::TestRayIntersection(const FRay3d& WorldRay)
+bool UPreviewMesh::TestRayIntersection(const FRay3d& WorldRay, TFunction<bool(int32 Tid)> TriangleFilterF)
 {
 	if (IsVisible() && TemporaryParentActor != nullptr && bBuildSpatialDataStructure)
 	{
 		FFrame3d TransformFrame(TemporaryParentActor->GetActorTransform());
 		FRay3d LocalRay = TransformFrame.ToFrame(WorldRay);
-		int HitTriID = MeshAABBTree.FindNearestHitTriangle(LocalRay);
+
+		UE::Geometry::IMeshSpatial::FQueryOptions QueryOptions;
+		QueryOptions.TriangleFilterF = TriangleFilterF;
+
+		int HitTriID = MeshAABBTree.FindNearestHitTriangle(LocalRay, QueryOptions);
 		if (HitTriID != FDynamicMesh3::InvalidID)
 		{
 			return true;
@@ -432,7 +436,7 @@ bool UPreviewMesh::TestRayIntersection(const FRay3d& WorldRay)
 
 
 
-bool UPreviewMesh::FindRayIntersection(const FRay3d& WorldRay, FHitResult& HitOut)
+bool UPreviewMesh::FindRayIntersection(const FRay3d& WorldRay, FHitResult& HitOut, TFunction<bool(int32)>TriangleFilterF)
 {
 	if (IsVisible() && TemporaryParentActor != nullptr && bBuildSpatialDataStructure)
 	{
@@ -440,7 +444,11 @@ bool UPreviewMesh::FindRayIntersection(const FRay3d& WorldRay, FHitResult& HitOu
 		FRay3d LocalRay(Transform.InverseTransformPosition(WorldRay.Origin),
 			Transform.InverseTransformVector(WorldRay.Direction));
 		UE::Geometry::Normalize(LocalRay.Direction);
-		int HitTriID = MeshAABBTree.FindNearestHitTriangle(LocalRay);
+
+		UE::Geometry::IMeshSpatial::FQueryOptions QueryOptions;
+		QueryOptions.TriangleFilterF = TriangleFilterF;
+
+		int HitTriID = MeshAABBTree.FindNearestHitTriangle(LocalRay, QueryOptions);
 		if (HitTriID != FDynamicMesh3::InvalidID)
 		{
 			const FDynamicMesh3* UseMesh = GetPreviewDynamicMesh();
diff --git a/Engine/Plugins/Runtime/MeshModelingToolset/Source/ModelingComponents/Public/PreviewMesh.h b/Engine/Plugins/Runtime/MeshModelingToolset/Source/ModelingComponents/Public/PreviewMesh.h
--- a/Engine/Plugins/Runtime/MeshModelingToolset/Source/ModelingComponents/Public/PreviewMesh.h
+++ b/Engine/Plugins/Runtime/MeshModelingToolset/Source/ModelingComponents/Public/PreviewMesh.h
@@ -266,20 +266,22 @@ public:
 	/**
 	 * Test for ray intersection with the preview mesh.
 	 * Requires that bBuildSpatialDataStructure = true
+	 * @param TriangleFilterF optional filter function for triangles
 	 * @return true if ray intersections mesh
 	 * @warning returns false if preview is not visible
 	 */
-	bool TestRayIntersection(const FRay3d& WorldRay);
+	bool TestRayIntersection(const FRay3d& WorldRay, TFunction<bool(int32 Tid)> TriangleFilterF = nullptr);
 
 	/**
 	 * Find ray intersection with the preview mesh.
 	 * Requires that bBuildSpatialDataStructure = true
 	 * @param WorldRay ray in world space
 	 * @param HitOut output hit data (only certain members are initialized, see implementation)
+	 * @param TriangleFilterF optional filter function for triangles
 	 * @return true if ray intersections mesh
 	 * @warning returns false if preview is not visible
 	 */
-	bool FindRayIntersection(const FRay3d& WorldRay, FHitResult& HitOut);
+	bool FindRayIntersection(const FRay3d& WorldRay, FHitResult& HitOut, TFunction<bool(int32 Tid)> TriangleFilterF = nullptr);
 
 
 	/**
