Fixes crash when selecting all vertices on non-manifold meshes

diff --git a/Engine/Plugins/Experimental/MeshModelingToolsetExp/Source/MeshModelingToolsEditorOnlyExp/Private/SkeletalMesh/SkinWeightsPaintTool.cpp b/Engine/Plugins/Experimental/MeshModelingToolsetExp/Source/MeshModelingToolsEditorOnlyExp/Private/SkeletalMesh/SkinWeightsPaintTool.cpp
--- a/Engine/Plugins/Experimental/MeshModelingToolsetExp/Source/MeshModelingToolsEditorOnlyExp/Private/SkeletalMesh/SkinWeightsPaintTool.cpp
+++ b/Engine/Plugins/Experimental/MeshModelingToolsetExp/Source/MeshModelingToolsEditorOnlyExp/Private/SkeletalMesh/SkinWeightsPaintTool.cpp
@@ -581,11 +581,14 @@ const TArray<int32>& UWeightToolMeshSelector::GetSelectedVertices()
 		// we have to make sure that the vertex ids are safe to use as PolygonSelectionMechanic does not act on the
 		// mesh description but on the dynamic mesh that can duplicate vertices when dealing with degenerate triangles.
 		// cf. FMeshDescriptionToDynamicMesh::Convert for more details.
-		Algo::CopyIf(VerticesToAdd, SelectedVerticesInternal, [&](int32 VertexID)
+		for (const int32 VertexToAdd : VerticesToAdd)
 		{
-			const int32 SrcVertexID = NonManifoldMappingSupport.GetOriginalNonManifoldVertexID(VertexID);
-			return AllVertices.IsValid(SrcVertexID);
-		});
+			const int32 SrcVertexID = NonManifoldMappingSupport.GetOriginalNonManifoldVertexID(VertexToAdd);
+			if (AllVertices.IsValid(SrcVertexID))
+			{
+				SelectedVerticesInternal.Add(SrcVertexID);
+			}
+		}
 	};
 
 	// add selected vertices
