Skin weight tool, fixes Hammer undo and UI buttons shifting while using Hammer (or any other edit weights button).

diff --git a/Engine/Plugins/Experimental/MeshModelingToolsetExp/Source/MeshModelingToolsEditorOnlyExp/Private/SkeletalMesh/SkinWeightsPaintTool.cpp b/Engine/Plugins/Experimental/MeshModelingToolsetExp/Source/MeshModelingToolsEditorOnlyExp/Private/SkeletalMesh/SkinWeightsPaintTool.cpp
--- a/Engine/Plugins/Experimental/MeshModelingToolsetExp/Source/MeshModelingToolsEditorOnlyExp/Private/SkeletalMesh/SkinWeightsPaintTool.cpp
+++ b/Engine/Plugins/Experimental/MeshModelingToolsetExp/Source/MeshModelingToolsEditorOnlyExp/Private/SkeletalMesh/SkinWeightsPaintTool.cpp
@@ -3927,15 +3927,15 @@ void USkinWeightsPaintTool::HammerWeights()
 		}
 		const int32 ClosestVertex = VertexPath.Last();
 
-		// remove all current weights (but don't prune)
+		// remove all current weights (pruning since this operation completely replaces the weight)
 		for (const FVertexBoneWeight& BoneWeight : Weights.PreChangeWeights[SelectedVertex])
 		{
 			constexpr float NewWeight = 0.f;
-			constexpr bool bPruneInfluence = false;
+			constexpr bool bPruneInfluence = true;
 			HammerWeightEdits.MergeSingleEdit(BoneWeight.BoneID, SelectedVertex, NewWeight, bPruneInfluence, Weights.PreChangeWeights);
 		}
 
-		// add weights from closest vertex
+		// replace weights with values from the closest vertex
 		for (const FVertexBoneWeight& BoneWeight : Weights.PreChangeWeights[ClosestVertex])
 		{
 			const float NewWeight = BoneWeight.Weight;
