MeshModelingTools: Add "isolate selection" option to vertex color paint tool.

diff --git a/Engine/Plugins/Experimental/MeshModelingToolsetExp/Source/MeshModelingToolsExp/Private/MeshVertexPaintTool.cpp b/Engine/Plugins/Experimental/MeshModelingToolsetExp/Source/MeshModelingToolsExp/Private/MeshVertexPaintTool.cpp
--- a/Engine/Plugins/Experimental/MeshModelingToolsetExp/Source/MeshModelingToolsExp/Private/MeshVertexPaintTool.cpp
+++ b/Engine/Plugins/Experimental/MeshModelingToolsetExp/Source/MeshModelingToolsExp/Private/MeshVertexPaintTool.cpp
@@ -15,8 +15,11 @@
 #include "MeshWeights.h"
 #include "DynamicMesh/MeshNormals.h"
 #include "DynamicMesh/MeshIndexUtil.h"
+#include "DynamicSubmesh3.h"
 #include "Util/BufferUtil.h"
 #include "Util/ColorConstants.h"
+#include "Selection/StoredMeshSelectionUtil.h" // GetCurrentGeometrySelectionForTarget
+#include "Selections/GeometrySelectionUtil.h"
 #include "Selections/MeshConnectedComponents.h"
 #include "Selections/MeshFaceSelection.h"
 #include "Selections/MeshVertexSelection.h"
@@ -73,6 +76,21 @@ UMeshSurfacePointTool* UMeshVertexPaintToolBuilder::CreateNewTool(const FToolBui
 	return SculptTool;
 }
 
+void UMeshVertexPaintToolBuilder::InitializeNewTool(UMeshSurfacePointTool* NewTool, const FToolBuilderState& SceneState) const
+{
+	Super::InitializeNewTool(NewTool, SceneState);
+	UMeshVertexPaintTool* Tool = Cast<UMeshVertexPaintTool>(NewTool);
+	if (ensure(Tool && Tool->GetTarget()))
+	{
+		UE::Geometry::FGeometrySelection Selection;
+		bool bHaveSelection = UE::Geometry::GetCurrentGeometrySelectionForTarget(SceneState, Tool->GetTarget(), Selection);
+		if (bHaveSelection)
+		{
+			Tool->SetGeometrySelection(MoveTemp(Selection));
+		}
+	}
+}
+
 bool UMeshVertexPaintToolBuilder::CanBuildTool(const FToolBuilderState& SceneState) const
 {
 	return UMeshSurfacePointMeshEditingToolBuilder::CanBuildTool(SceneState) &&
@@ -194,6 +212,12 @@ void UMeshVertexPaintTool::Setup()
 	//Mesh->EnableVertexColors(FVector3f::One());
 	FAxisAlignedBox3d Bounds = Mesh->GetBounds(true);
 
+	FilterProperties = NewObject<UVertexPaintBrushFilterProperties>(this);
+	FilterProperties->RestoreProperties(this);
+	FilterProperties->WatchProperty(FilterProperties->MaterialMode, [this](EMeshVertexPaintMaterialMode) { UpdateVertexPaintMaterialMode(); });
+	FilterProperties->bToolHasSelection = GeometrySelection.IsSet();
+	SelectionTids.Reset();
+
 	TFuture<void> PrecomputeFuture = Async(VertexPaintToolAsyncExecTarget, [&]()
 	{
 		PrecomputeFilterData();
@@ -252,11 +276,6 @@ void UMeshVertexPaintTool::Setup()
 	BasicProperties->WatchProperty(BasicProperties->SecondaryActionType,
 		[this](EMeshVertexPaintSecondaryActionType NewType) { UpdateSecondaryBrushType(NewType); });
 
-	// initialize other properties
-	FilterProperties = NewObject<UVertexPaintBrushFilterProperties>(this);
-	FilterProperties->RestoreProperties(this);
-	FilterProperties->WatchProperty(FilterProperties->MaterialMode, [this](EMeshVertexPaintMaterialMode) { UpdateVertexPaintMaterialMode(); });
-
 	InitializeIndicator();
 
 	AddToolPropertySource(BasicProperties);
@@ -329,8 +348,24 @@ void UMeshVertexPaintTool::Setup()
 		MeshElementsDisplay->Settings->RestoreProperties(this, TEXT("MeshVertexPaintTool"));
 		AddToolPropertySource(MeshElementsDisplay->Settings);
 	}
-	MeshElementsDisplay->SetMeshAccessFunction([this](UMeshElementsVisualizer::ProcessDynamicMeshFunc ProcessFunc) {
-		ProcessFunc(*GetSculptMesh());
+	MeshElementsDisplay->SetMeshAccessFunction([this](UMeshElementsVisualizer::ProcessDynamicMeshFunc ProcessFunc) 
+	{
+		const UE::Geometry::FDynamicMesh3* const FullMesh = GetSculptMesh();
+		if (ShouldFilterTrianglesBySelection())
+		{
+			// Arguably there should be a way to add a filter to the visualizer, but we've done it this way elsewhere...
+			UE::Geometry::FDynamicSubmesh3 Submesh(FullMesh, SelectionTids.Array());
+			ProcessFunc(Submesh.GetSubmesh());
+		}
+		else
+		{
+			ProcessFunc(*FullMesh);
+		}
+	});
+	FilterProperties->WatchProperty(FilterProperties->bIsolateGeometrySelection, [this](bool)
+	{
+		DynamicMeshComponent->SetSecondaryBuffersVisibility(!ShouldFilterTrianglesBySelection());
+		MeshElementsDisplay->NotifyMeshChanged();
 	});
 
 	// initialize weightmap list in Utility actions
@@ -449,6 +484,11 @@ void UMeshVertexPaintTool::CommitResult(UBaseDynamicMeshComponent* Component, bo
 }
 
 
+void UMeshVertexPaintTool::SetGeometrySelection(const UE::Geometry::FGeometrySelection& SelectionIn)
+{
+	GeometrySelection = SelectionIn;
+}
+
 void UMeshVertexPaintTool::RegisterActions(FInteractiveToolActionSet& ActionSet)
 {
 	UMeshSculptToolBase::RegisterActions(ActionSet);
@@ -609,6 +649,10 @@ void UMeshVertexPaintTool::UpdateROI(const FSculptBrushStamp& BrushStamp)
 
 	TriangleROI.Reset();
 
+	// Theoretically we shouldn't need to do this filtering here since we already filter out
+	//  unselected triangles in ApplyVisibilityFilter if needed. But seems odd not to.
+	bool bShouldFilterTrianglesBySelection = ShouldFilterTrianglesBySelection();
+
 	// With Lazy Brush, GetBrushTriangleID() may not necessarily return a triangle that contains BrushPos,
 	// so in that case we will try to find the triangle containing BrushPos. Note that this may not
 	// be correct in some geometric cases...potentially we could check if the initial CenterTID 
@@ -617,7 +661,12 @@ void UMeshVertexPaintTool::UpdateROI(const FSculptBrushStamp& BrushStamp)
 	if (BrushProperties->Lazyness > 0)
 	{
 		double NearDistSqr;
-		CenterTID = AABBTree.FindNearestTriangle(BrushPos, NearDistSqr);
+		IMeshSpatial::FQueryOptions QueryOptions;
+		if (bShouldFilterTrianglesBySelection)
+		{
+			QueryOptions.TriangleFilterF = [this](int32 Tid) { return SelectionTids.Contains(Tid); };
+		}
+		CenterTID = AABBTree.FindNearestTriangle(BrushPos, NearDistSqr, QueryOptions);
 	}
 
 	if (Mesh->IsTriangle(CenterTID))
@@ -656,7 +705,10 @@ void UMeshVertexPaintTool::UpdateROI(const FSculptBrushStamp& BrushStamp)
 	{
 		Octree.RangeQuery(BrushBox,
 			[&](int TriIdx) {
-
+			if (bShouldFilterTrianglesBySelection && !SelectionTids.Contains(TriIdx))
+			{
+				return;
+			}
 			if ((Mesh->GetTriCentroid(TriIdx) - BrushPos).SquaredLength() < RadiusSqr)
 			{
 				TriangleROI.Add(TriIdx);
@@ -672,6 +724,10 @@ void UMeshVertexPaintTool::UpdateROI(const FSculptBrushStamp& BrushStamp)
 			FMeshConnectedComponents::GrowToConnectedTriangles(Mesh, StartROI, TriangleROI, &TempROIBuffer,
 				[&](int t1, int t2) 
 			{ 
+				if (bShouldFilterTrianglesBySelection && (!SelectionTids.Contains(t1) || !SelectionTids.Contains(t2)))
+				{
+					return false;
+				}
 				if ((Mesh->GetTriCentroid(t2) - BrushPos).SquaredLength() < RadiusSqr)
 				{
 					return CheckEdgeCriteria(t1, t2);
@@ -691,6 +747,10 @@ void UMeshVertexPaintTool::UpdateROI(const FSculptBrushStamp& BrushStamp)
 		FMeshConnectedComponents::GrowToConnectedTriangles(Mesh, StartROI, TriangleROI, &TempROIBuffer,
 														   [&](int t1, int t2)
 		{
+			if (bShouldFilterTrianglesBySelection && (!SelectionTids.Contains(t1) || !SelectionTids.Contains(t2)))
+			{
+				return false;
+			}
 			return CheckEdgeCriteria(t1, t2);
 		});
 	}
@@ -706,6 +766,10 @@ void UMeshVertexPaintTool::UpdateROI(const FSculptBrushStamp& BrushStamp)
 		FMeshConnectedComponents::GrowToConnectedTriangles(Mesh, StartROI, TriangleROI, &TempROIBuffer,
 			[&](int t1, int t2)
 		{
+			if (bShouldFilterTrianglesBySelection && (!SelectionTids.Contains(t1) || !SelectionTids.Contains(t2)))
+			{
+				return false;
+			}
 			return (FillGroups.Contains(ActiveGroupSet->GetGroup(t2)));
 		});
 	}
@@ -1034,7 +1098,8 @@ void UMeshVertexPaintTool::InitializeElementROIFromTriangleROI(const TArray<int3
 
 bool UMeshVertexPaintTool::HaveVisibilityFilter() const
 {
-	return FilterProperties->VisibilityFilter != EMeshVertexPaintVisibilityType::None;
+	return ShouldFilterTrianglesBySelection()
+		|| FilterProperties->VisibilityFilter != EMeshVertexPaintVisibilityType::None;
 }
 
 
@@ -1072,16 +1137,23 @@ void UMeshVertexPaintTool::ApplyVisibilityFilter(const TArray<int32>& Triangles,
 	const FDynamicMesh3* Mesh = GetSculptMesh();
 
 	int32 NumTriangles = Triangles.Num();
+	bool bShouldFilterUnselected = ShouldFilterTrianglesBySelection();
 
 	VisibilityFilterBuffer.SetNum(NumTriangles, EAllowShrinking::No);
 	ParallelFor(NumTriangles, [&](int32 idx)
 	{
 		VisibilityFilterBuffer[idx] = true;
+		if (bShouldFilterUnselected && !SelectionTids.Contains(idx))
+		{
+			VisibilityFilterBuffer[idx] = false;
+			return;
+		}
 		FVector3d Centroid = Mesh->GetTriCentroid(Triangles[idx]);
 		FVector3d FaceNormal = Mesh->GetTriNormal(Triangles[idx]);
 		if (FaceNormal.Dot((Centroid - LocalEyePosition)) > 0)
 		{
 			VisibilityFilterBuffer[idx] = false;
+			return;
 		}
 		if (FilterProperties->VisibilityFilter == EMeshVertexPaintVisibilityType::Unoccluded)
 		{
@@ -1112,8 +1184,13 @@ int32 UMeshVertexPaintTool::FindHitSculptMeshTriangle(const FRay3d& LocalRay)
 		return IndexConstants::InvalidID;
 	}
 
+	bool bShouldFilterUnselected = ShouldFilterTrianglesBySelection();
 	if (GetBrushCanHitBackFaces())
 	{
+		if (bShouldFilterUnselected)
+		{
+			return Octree.FindNearestHitObject(LocalRay, [this](int32 Tid) { return SelectionTids.Contains(Tid); });
+		}
 		return Octree.FindNearestHitObject(LocalRay);
 	}
 	else
@@ -1124,7 +1201,11 @@ int32 UMeshVertexPaintTool::FindHitSculptMeshTriangle(const FRay3d& LocalRay)
 		GetToolManager()->GetContextQueriesAPI()->GetCurrentViewState(StateOut);
 		FVector3d LocalEyePosition(CurTargetTransform.InverseTransformPosition((FVector3d)StateOut.Position));
 		int HitTID = Octree.FindNearestHitObject(LocalRay,
-			[this, Mesh, &LocalEyePosition](int TriangleID) {
+			[this, Mesh, &LocalEyePosition, bShouldFilterUnselected](int TriangleID) {
+			if (bShouldFilterUnselected && !SelectionTids.Contains(TriangleID))
+			{
+				return false;
+			}
 			FVector3d Normal, Centroid;
 			double Area;
 			Mesh->GetTriInfo(TriangleID, Normal, Area, Centroid);
@@ -1210,7 +1291,9 @@ void UMeshVertexPaintTool::DrawHUD(FCanvas* Canvas, IToolsContextRenderAPI* Rend
 	if (FilterProperties->bShowHitColor)
 	{
 		FRay3d LocalRay(LocalEyePosition, Normalized(HoverStamp.LocalFrame.Origin - LocalEyePosition));
-		int CursorHitTID = Octree.FindNearestHitObject(LocalRay);
+		int CursorHitTID = ShouldFilterTrianglesBySelection() ? 
+			Octree.FindNearestHitObject(LocalRay, [this](int32 Tid) {return SelectionTids.Contains(Tid); })
+			: Octree.FindNearestHitObject(LocalRay);
 		if (Mesh->IsTriangle(CursorHitTID))
 		{
 			FIntrRay3Triangle3d Intersection = TMeshQueries<FDynamicMesh3>::RayTriangleIntersection(*GetSculptMesh(), CursorHitTID, LocalRay);
@@ -1967,6 +2050,26 @@ void UMeshVertexPaintTool::PrecomputeFilterData()
 			NormalSeamEdges[eid] = Normals->IsSeamEdge(eid);
 		}
 	});
+
+	if (GeometrySelection.IsSet())
+	{
+		UE::Geometry::EnumerateSelectionTriangles(GeometrySelection.GetValue(), *Mesh,
+			[this, &Mesh](int32 TriangleID)
+			{
+				SelectionTids.Add(TriangleID);
+			});
+
+		DynamicMeshComponent->EnableSecondaryTriangleBuffers([this](const UE::Geometry::FDynamicMesh3*, int32 Tid)
+		{
+			return !SelectionTids.Contains(Tid);
+		});
+		DynamicMeshComponent->SetSecondaryBuffersVisibility(!ShouldFilterTrianglesBySelection());
+	}
+}
+
+bool UMeshVertexPaintTool::ShouldFilterTrianglesBySelection() const
+{
+	return FilterProperties->bIsolateGeometrySelection && !SelectionTids.IsEmpty();
 }
 
 
diff --git a/Engine/Plugins/Experimental/MeshModelingToolsetExp/Source/MeshModelingToolsExp/Public/MeshVertexPaintTool.h b/Engine/Plugins/Experimental/MeshModelingToolsetExp/Source/MeshModelingToolsExp/Public/MeshVertexPaintTool.h
--- a/Engine/Plugins/Experimental/MeshModelingToolsetExp/Source/MeshModelingToolsExp/Public/MeshVertexPaintTool.h
+++ b/Engine/Plugins/Experimental/MeshModelingToolsetExp/Source/MeshModelingToolsExp/Public/MeshVertexPaintTool.h
@@ -11,6 +11,7 @@
 #include "PropertySets/PolygroupLayersProperties.h"
 #include "PropertySets/ColorChannelFilterPropertyType.h"
 #include "Mechanics/PolyLassoMarqueeMechanic.h"
+#include "Selections/GeometrySelection.h"
 #include "TargetInterfaces/MeshTargetInterfaceTypes.h"
 
 #include "Sculpting/MeshSculptToolBase.h"
@@ -42,6 +43,7 @@ class MESHMODELINGTOOLSEXP_API UMeshVertexPaintToolBuilder : public UMeshSurface
 
 public:
 	virtual UMeshSurfacePointTool* CreateNewTool(const FToolBuilderState& SceneState) const override;
+	virtual void InitializeNewTool(UMeshSurfacePointTool* NewTool, const FToolBuilderState& SceneState) const override;
 	virtual bool CanBuildTool(const FToolBuilderState& SceneState) const override;
 
 protected:
@@ -237,6 +239,18 @@ public:
 	UPROPERTY(EditAnywhere, Category = Filters)
 	EMeshVertexPaintVisibilityType VisibilityFilter = EMeshVertexPaintVisibilityType::None;
 
+	/**
+	 * If the tool was started with a mesh element selection, this setting hides everything
+	 *  except that selection, to make painting it easier. Requires that a mesh element
+	 *  selection exist on tool start.
+	 */
+	UPROPERTY(EditAnywhere, Category = Filters, meta = (EditCondition = "bToolHasSelection", HideEditConditionToggle))
+	bool bIsolateGeometrySelection = false;
+
+	//~ For the tool to set, to enable/disable bIsolateGeometrySelection
+	UPROPERTY()
+	bool bToolHasSelection = false;
+
 	/** Number of vertices in a triangle the Lasso must hit to be counted as "inside" */
 	//UPROPERTY(EditAnywhere, Category = Filters, AdvancedDisplay, meta = (UIMin = 1, UIMax = 3, EditCondition = "CurrentSubToolType == EMeshVertexPaintInteractionType::PolyLasso"))
 	UPROPERTY()
@@ -466,6 +480,8 @@ class MESHMODELINGTOOLSEXP_API UMeshVertexPaintTool : public UMeshSculptToolBase
 	GENERATED_BODY()
 
 public:
+	void SetGeometrySelection(const UE::Geometry::FGeometrySelection& SelectionIn);
+
 	virtual void RegisterActions(FInteractiveToolActionSet& ActionSet) override;
 
 	virtual void Setup() override;
@@ -690,6 +706,11 @@ protected:
 
 	// Currently using flow rate as 'brush strength', so disable temporal stamp spacing
 	virtual float GetStampTemporalFlowRate() const override { return 1.0f; }
+
+private:
+	TOptional<UE::Geometry::FGeometrySelection> GeometrySelection;
+	TSet<int32> SelectionTids;
+	bool ShouldFilterTrianglesBySelection() const;
 };
 
 
