Skeleton Modifier: update all skin weights profile (and not just the default one) when updating the skeleton

#rnx

diff --git a/Engine/Plugins/Experimental/MeshModelingToolsetExp/Source/SkeletalMeshModifiers/Private/SkeletonModifier.cpp b/Engine/Plugins/Experimental/MeshModelingToolsetExp/Source/SkeletalMeshModifiers/Private/SkeletonModifier.cpp
--- a/Engine/Plugins/Experimental/MeshModelingToolsetExp/Source/SkeletalMeshModifiers/Private/SkeletonModifier.cpp
+++ b/Engine/Plugins/Experimental/MeshModelingToolsetExp/Source/SkeletalMeshModifiers/Private/SkeletonModifier.cpp
@@ -608,69 +608,7 @@ bool USkeletonModifier::CommitSkeletonToSkeletalMesh()
 	const TArray<FMeshBoneInfo>& BoneInfos = ReferenceSkeleton->GetRawRefBoneInfo();
 	
 	// update mesh description
-	{
-		FSkeletalMeshAttributes MeshAttributes(*MeshDescription);
-
-		// update bone data
-		if (!MeshAttributes.HasBones())
-		{
-			MeshAttributes.Register(true);
-		}
-
-		MeshAttributes.Bones().Reset(BoneInfos.Num());
-
-		FSkeletalMeshAttributes::FBoneNameAttributesRef BoneNames = MeshAttributes.GetBoneNames();
-		FSkeletalMeshAttributes::FBoneParentIndexAttributesRef BoneParentIndices = MeshAttributes.GetBoneParentIndices();
-		FSkeletalMeshAttributes::FBonePoseAttributesRef BonePoses = MeshAttributes.GetBonePoses();
-
-		const TArray<FTransform> Transforms = ReferenceSkeleton->GetRawRefBonePose();
-		for (int Index = 0; Index < BoneInfos.Num(); ++Index)
-		{
-			const FMeshBoneInfo& Info = BoneInfos[Index];
-			const FBoneID BoneID = MeshAttributes.CreateBone();
-			BoneNames.Set(BoneID, Info.Name);
-			BoneParentIndices.Set(BoneID, Info.ParentIndex);
-			BonePoses.Set(BoneID, Transforms[Index]);
-		}
-		
-		// update skin data if needed
-		if (EnumHasAnyFlags(Modifications, ESkeletalMeshModificationType::IndicesUpdated))
-		{
-			using namespace UE::AnimationCore;
-			FBoneWeightsSettings BoneSettings; BoneSettings.SetNormalizeType(EBoneWeightNormalizeType::None);
-		
-			FSkinWeightsVertexAttributesRef SkinWeights = MeshAttributes.GetVertexSkinWeights();
-			for (const FVertexID& VertexID: MeshDescription->Vertices().GetElementIDs())
-			{
-				FVertexBoneWeights BoneWeights = SkinWeights.Get(VertexID);
-				if (const int32 NumBoneWeights = BoneWeights.Num())
-				{
-					TArray<FBoneWeight> NewWeights;
-					for (int32 Idx = 0; Idx < NumBoneWeights; ++Idx)
-					{
-						const FBoneWeight& OldBoneWeight = BoneWeights[Idx];
-						const int32 BoneIndex = OldBoneWeight.GetBoneIndex();
-
-						int32 NewBoneIndex = 0;
-						if (ensure(BoneIndexTracker.IsValidIndex(BoneIndex)))
-						{
-							NewBoneIndex = BoneIndexTracker[BoneIndex];							
-						}
-						else
-						{
-							UE_LOG(LogAnimation, Warning, TEXT("Skeleton Modifier - Commit: Invalid bone index provided (%d); falling back to 0 as bone index."), BoneIndex);
-						}
-
-						if (NewBoneIndex != INDEX_NONE)
-						{
-							NewWeights.Add(FBoneWeight(NewBoneIndex, OldBoneWeight.GetRawWeight()));
-						}
-					}
-					SkinWeights.Set(VertexID, FBoneWeights::Create(NewWeights, BoneSettings));
-				}
-			}
-		}
-	}
+	CommitChangesToMeshDescription(Modifications);
 
 	// store retargeting modes
 	USkeleton* Skeleton = SkeletalMesh->GetSkeleton();
@@ -766,6 +704,85 @@ bool USkeletonModifier::CommitSkeletonToSkeletalMesh()
 	return false;
 }
 
+void USkeletonModifier::CommitChangesToMeshDescription(const ESkeletalMeshModificationType InSkeletalMeshModifications)
+{
+	if (!SkeletalMesh.IsValid() || !ReferenceSkeleton || !MeshDescription)
+	{
+		// this is supposed to be tested earlier
+		return;
+	}
+	
+	const TArray<FMeshBoneInfo>& BoneInfos = ReferenceSkeleton->GetRawRefBoneInfo();
+	const TArray<FTransform>& Transforms = ReferenceSkeleton->GetRawRefBonePose();
+	
+	FSkeletalMeshAttributes MeshAttributes(*MeshDescription);
+
+	// update bone data
+	if (!MeshAttributes.HasBones())
+	{
+		MeshAttributes.Register(true);
+	}
+
+	MeshAttributes.Bones().Reset(BoneInfos.Num());
+
+	FSkeletalMeshAttributes::FBoneNameAttributesRef BoneNames = MeshAttributes.GetBoneNames();
+	FSkeletalMeshAttributes::FBoneParentIndexAttributesRef BoneParentIndices = MeshAttributes.GetBoneParentIndices();
+	FSkeletalMeshAttributes::FBonePoseAttributesRef BonePoses = MeshAttributes.GetBonePoses();
+
+	for (int Index = 0; Index < BoneInfos.Num(); ++Index)
+	{
+		const FMeshBoneInfo& Info = BoneInfos[Index];
+		const FBoneID BoneID = MeshAttributes.CreateBone();
+		BoneNames.Set(BoneID, Info.Name);
+		BoneParentIndices.Set(BoneID, Info.ParentIndex);
+		BonePoses.Set(BoneID, Transforms[Index]);
+	}
+		
+	// update skin weight data if needed
+	if (EnumHasAnyFlags(InSkeletalMeshModifications, ESkeletalMeshModificationType::IndicesUpdated))
+	{
+		using namespace UE::AnimationCore;
+		FBoneWeightsSettings BoneSettings; BoneSettings.SetNormalizeType(EBoneWeightNormalizeType::None);
+
+		for (const FName SkinWeightProfile: MeshAttributes.GetSkinWeightProfileNames())
+		{
+			FSkinWeightsVertexAttributesRef SkinWeights = MeshAttributes.GetVertexSkinWeights(SkinWeightProfile);
+			if (SkinWeights.IsValid())
+			{
+				for (const FVertexID& VertexID: MeshDescription->Vertices().GetElementIDs())
+				{
+					FVertexBoneWeights BoneWeights = SkinWeights.Get(VertexID);
+					if (const int32 NumBoneWeights = BoneWeights.Num())
+					{
+						TArray<FBoneWeight> NewWeights;
+						for (int32 Idx = 0; Idx < NumBoneWeights; ++Idx)
+						{
+							const FBoneWeight& OldBoneWeight = BoneWeights[Idx];
+							const int32 BoneIndex = OldBoneWeight.GetBoneIndex();
+
+							int32 NewBoneIndex = 0;
+							if (ensure(BoneIndexTracker.IsValidIndex(BoneIndex)))
+							{
+								NewBoneIndex = BoneIndexTracker[BoneIndex];							
+							}
+							else
+							{
+								UE_LOG(LogAnimation, Warning, TEXT("Skeleton Modifier - Commit: Invalid bone index provided (%d); falling back to 0 as bone index."), BoneIndex);
+							}
+
+							if (NewBoneIndex != INDEX_NONE)
+							{
+								NewWeights.Add(FBoneWeight(NewBoneIndex, OldBoneWeight.GetRawWeight()));
+							}
+						}
+						SkinWeights.Set(VertexID, FBoneWeights::Create(NewWeights, BoneSettings));
+					}
+				}
+			}
+		}
+	}
+}
+
 void USkeletonModifier::PostCommitSkeleton(const ESkeletonModificationType InSkeletonModifications) const
 {
 	if (InSkeletonModifications != ESkeletonModificationType::FullMergeAll)
diff --git a/Engine/Plugins/Experimental/MeshModelingToolsetExp/Source/SkeletalMeshModifiers/Public/SkeletonModifier.h b/Engine/Plugins/Experimental/MeshModelingToolsetExp/Source/SkeletalMeshModifiers/Public/SkeletonModifier.h
--- a/Engine/Plugins/Experimental/MeshModelingToolsetExp/Source/SkeletalMeshModifiers/Public/SkeletonModifier.h
+++ b/Engine/Plugins/Experimental/MeshModelingToolsetExp/Source/SkeletalMeshModifiers/Public/SkeletonModifier.h
@@ -268,6 +268,9 @@ private:
 	ESkeletonModificationType PreCommitSkeleton(const ESkeletalMeshModificationType InSkeletalMeshModifications) const;
 	void PostCommitSkeleton(const ESkeletonModificationType InSkeletonModifications) const;
 
+	// update mesh description on commit
+	void CommitChangesToMeshDescription(const ESkeletalMeshModificationType InSkeletalMeshModifications);
+	
 	// notification
 	void NotifyFromSkeletonChanges() const;
 	
