add 'show target' option to the collision inspector tool

diff --git a/Engine/Plugins/Experimental/MeshModelingToolsetExp/Source/MeshModelingToolsExp/Private/Physics/PhysicsInspectorTool.cpp b/Engine/Plugins/Experimental/MeshModelingToolsetExp/Source/MeshModelingToolsExp/Private/Physics/PhysicsInspectorTool.cpp
--- a/Engine/Plugins/Experimental/MeshModelingToolsetExp/Source/MeshModelingToolsExp/Private/Physics/PhysicsInspectorTool.cpp
+++ b/Engine/Plugins/Experimental/MeshModelingToolsetExp/Source/MeshModelingToolsExp/Private/Physics/PhysicsInspectorTool.cpp
@@ -57,6 +57,17 @@ void UPhysicsInspectorTool::Setup()
 		InitializePreviewGeometry(false /*bClearExisting*/);
 	}
 
+	Settings = NewObject<UPhysicsInspectorToolProperties>(this);
+	Settings->RestoreProperties(this);
+	AddToolPropertySource(Settings);
+	Settings->WatchProperty(Settings->bShowTargetMesh, [this](bool bNewValue)
+	{
+		for (UToolTarget* Target : Targets)
+		{
+			UE::ToolTarget::SetSourceObjectVisible(Target, bNewValue);
+		}
+	});
+
 	OnCreatePhysicsDelegateHandle = UActorComponent::GlobalCreatePhysicsDelegate.AddUObject(this, &UPhysicsInspectorTool::OnCreatePhysics);
 
 	SetToolDisplayName(LOCTEXT("ToolName", "Physics Inspector"));
@@ -130,6 +141,7 @@ void UPhysicsInspectorTool::OnCreatePhysics(UActorComponent* Component)
 void UPhysicsInspectorTool::OnShutdown(EToolShutdownType ShutdownType)
 {
 	VizSettings->SaveProperties(this);
+	Settings->SaveProperties(this);
 
 	UActorComponent::GlobalCreatePhysicsDelegate.Remove(OnCreatePhysicsDelegateHandle);
 
@@ -137,6 +149,14 @@ void UPhysicsInspectorTool::OnShutdown(EToolShutdownType ShutdownType)
 	{
 		Preview->Disconnect();
 	}
+
+	if (!Settings->bShowTargetMesh)
+	{
+		for (UToolTarget* Target : Targets)
+		{
+			UE::ToolTarget::ShowSourceObject(Target);
+		}
+	}
 }
 
 
diff --git a/Engine/Plugins/Experimental/MeshModelingToolsetExp/Source/MeshModelingToolsExp/Public/Physics/PhysicsInspectorTool.h b/Engine/Plugins/Experimental/MeshModelingToolsetExp/Source/MeshModelingToolsExp/Public/Physics/PhysicsInspectorTool.h
--- a/Engine/Plugins/Experimental/MeshModelingToolsetExp/Source/MeshModelingToolsExp/Public/Physics/PhysicsInspectorTool.h
+++ b/Engine/Plugins/Experimental/MeshModelingToolsetExp/Source/MeshModelingToolsExp/Public/Physics/PhysicsInspectorTool.h
@@ -23,6 +23,16 @@ protected:
 };
 
 
+UCLASS()
+class MESHMODELINGTOOLSEXP_API UPhysicsInspectorToolProperties : public UInteractiveToolPropertySet
+{
+	GENERATED_BODY()
+public:
+
+	/** Show/Hide target mesh */
+	UPROPERTY(EditAnywhere, Category = TargetVisualization)
+	bool bShowTargetMesh = true;
+};
 
 
 /**
@@ -55,6 +65,9 @@ protected:
 	UPROPERTY()
 	TObjectPtr<UCollisionGeometryVisualizationProperties> VizSettings = nullptr;
 
+	UPROPERTY()
+	TObjectPtr<UPhysicsInspectorToolProperties> Settings = nullptr;
+
 	UPROPERTY()
 	TArray<TObjectPtr<UPhysicsObjectToolPropertySet>> ObjectData;
 
