Skin paint tool: fixed isolated selections to prevent the undo/redo queue from getting out of sync with the mesh.

diff --git a/Engine/Plugins/Experimental/MeshModelingToolsetExp/Source/MeshModelingToolsEditorOnlyExp/Private/SkeletalMesh/SkinWeightsPaintTool.cpp b/Engine/Plugins/Experimental/MeshModelingToolsetExp/Source/MeshModelingToolsEditorOnlyExp/Private/SkeletalMesh/SkinWeightsPaintTool.cpp
--- a/Engine/Plugins/Experimental/MeshModelingToolsetExp/Source/MeshModelingToolsEditorOnlyExp/Private/SkeletalMesh/SkinWeightsPaintTool.cpp
+++ b/Engine/Plugins/Experimental/MeshModelingToolsetExp/Source/MeshModelingToolsEditorOnlyExp/Private/SkeletalMesh/SkinWeightsPaintTool.cpp
@@ -2428,7 +2428,7 @@
 {
 	if (SelectionIsolator)
 	{
-		SelectionIsolator->OnTick(DeltaTime);
+		SelectionIsolator->UpdateIsolatedSelection();
 	}
 	
 	if (bStampPending)
@@ -2540,7 +2540,7 @@
 	USkinWeightsPaintTool* Tool = Cast<USkinWeightsPaintTool>(Object);
 	if (Tool)
 	{
-		Tool->GetSelectionIsolator()->SetIsolatedTriangles(IsolatedTrianglesAfter);
+		Tool->GetSelectionIsolator()->SetTrianglesToIsolate(IsolatedTrianglesAfter);
 	}	
 }
 
@@ -2549,7 +2549,7 @@
 	USkinWeightsPaintTool* Tool = Cast<USkinWeightsPaintTool>(Object);
 	if (Tool)
 	{
-		Tool->GetSelectionIsolator()->SetIsolatedTriangles(IsolatedTrianglesBefore);
+		Tool->GetSelectionIsolator()->SetTrianglesToIsolate(IsolatedTrianglesBefore);
 	}
 }
 
@@ -4382,13 +4382,21 @@
 	WeightTool = InTool;
 }
 
-void UWeightToolSelectionIsolator::OnTick(float DeltaTime)
+void UWeightToolSelectionIsolator::UpdateIsolatedSelection()
 {
-	if (bRestoreFullMeshOnNextTick)
+	// this is queued to run on Tick() because modifying the mesh from other threads can cause the tool's Render() to be out of sync
+	if (bIsolatedMeshNeedsUpdated)
 	{
-		// this is queued to run on Tick() because modifying the mesh from other threads can cause the tool's Render() to be out of sync
-		RestoreFullMesh();
-		bRestoreFullMeshOnNextTick = false;
+		if (CurrentlyIsolatedTriangles.IsEmpty())
+		{
+			RestoreFullMesh();	
+		}
+		else
+		{
+			CreatePartialMesh();
+		}
+		
+		bIsolatedMeshNeedsUpdated = false;
 	}
 }
 
@@ -4405,7 +4413,8 @@
 	ActiveChange->IsolatedTrianglesBefore = GetIsolatedTriangles();
 	WeightTool->GetMainMeshSelector()->GetSelectedTriangles(ActiveChange->IsolatedTrianglesAfter);
 
-	SetIsolatedTriangles(ActiveChange->IsolatedTrianglesAfter);
+	UpdateIsolatedSelection();
+	SetTrianglesToIsolate(ActiveChange->IsolatedTrianglesAfter);
 	
 	UInteractiveToolManager* ToolManager = WeightTool->GetToolManager();
 	ToolManager->BeginUndoTransaction(TransactionLabel);
@@ -4421,7 +4430,8 @@
 	ActiveChange->IsolatedTrianglesBefore = GetIsolatedTriangles();
 	ActiveChange->IsolatedTrianglesAfter = {};
 
-	SetIsolatedTriangles(ActiveChange->IsolatedTrianglesAfter);
+	UpdateIsolatedSelection();
+	SetTrianglesToIsolate(ActiveChange->IsolatedTrianglesAfter);
 	
 	UInteractiveToolManager* ToolManager = WeightTool->GetToolManager();
 	ToolManager->BeginUndoTransaction(TransactionLabel);
@@ -4429,87 +4439,19 @@
 	ToolManager->EndUndoTransaction();
 }
 
-void UWeightToolSelectionIsolator::SetIsolatedTriangles(const TArray<int32>& TrianglesToIsolate)
+void UWeightToolSelectionIsolator::SetTrianglesToIsolate(const TArray<int32>& TrianglesToIsolate)
 {
-	// if we are turning off an isolated selection, we must queue the Tick() to update the full mesh
-	if (TrianglesToIsolate.IsEmpty() && PartialMeshDescription.IsValid())
+	if (bIsolatedMeshNeedsUpdated)
 	{
-		bRestoreFullMeshOnNextTick = true;
+		// cannot queue up multiple changes
+		// we must allow previous selection isolations to be applied to keep the weight buffer in-sync
 		return;
 	}
 	
-	if (PartialMeshDescription.IsValid())
-	{
-		// should be reset to null
-		// NOTE: we do not support isolating the selection twice in a row
-		// ie. you must isolate the selection, then un-isolate it before isolating it again
-		ensure(false); 
-		return;
-	}
-	
-	UPolygonSelectionMechanic* SelectionMechanic = WeightTool->GetMainMeshSelector()->GetSelectionMechanic();
-	if (!ensure(SelectionMechanic))
-	{
-		return;
-	}
-
-	const USkeletalMeshComponent* SkeletalMeshComponent = GetSkeletalMeshComponent(WeightTool->GetTarget());
-	if (!ensure(SkeletalMeshComponent))
-	{
-		return;
-	}
-
-	const FMeshDescription* EditedMesh = WeightTool->GetCurrentlyEditedMeshDescription();
-	if (!ensure(EditedMesh))
-	{
-		return;
-	}
-
-	// record the triangles we are isolating
+	// record the triangles we are isolating (if empty, this will restore the full mesh)
 	CurrentlyIsolatedTriangles = TrianglesToIsolate;
-
-	// get the weights
-	FSkinToolWeights& Weights = WeightTool->GetWeights();
-
-	// apply previous changes
-	Weights.ApplyCurrentWeightsToMeshDescription(WeightTool->GetCurrentlyEditedMeshDescription());
-	
-	// put into ref pose, BEFORE copying the mesh, so that submesh deformer initializes with vertices in ref pose
-	Weights.Deformer.SetToRefPose(WeightTool);
-
-	// store selection to be restored
-	IsolatedSelectionToRestoreVertices.Reset();
-	IsolatedSelectionToRestoreEdges.Reset();
-	IsolatedSelectionToRestoreFaces.Reset();
-	IsolatedSelectionToRestoreVertices.ElementType = UE::Geometry::EGeometryElementType::Vertex;
-	IsolatedSelectionToRestoreEdges.ElementType = UE::Geometry::EGeometryElementType::Edge;
-	IsolatedSelectionToRestoreFaces.ElementType = UE::Geometry::EGeometryElementType::Face;
-	SelectionMechanic->GetSelection_AsTriangleTopology(IsolatedSelectionToRestoreVertices);
-	SelectionMechanic->GetSelection_AsTriangleTopology(IsolatedSelectionToRestoreEdges);
-	SelectionMechanic->GetSelection_AsTriangleTopology(IsolatedSelectionToRestoreFaces);
-
-	// store copy of original FDynamicMesh to restore
-	FDynamicMesh3 DynamicMesh;
-	FMeshDescriptionToDynamicMesh Converter;
-	Converter.Convert(EditedMesh, DynamicMesh);
-	FullDynamicMesh = MoveTemp(DynamicMesh);
-
-	// create a submesh from the given triangles
-	PartialSubMesh = UE::Geometry::FDynamicSubmesh3(&FullDynamicMesh, TrianglesToIsolate);
-
-	// create mesh description for sub-mesh
-	PartialMeshDescription = MakeShared<FMeshDescription>();
-	// registering skeletal mesh attributes is required to create room to copy attributes during conversion from dynamic mesh
-	FSkeletalMeshAttributes Attributes(*PartialMeshDescription);
-	Attributes.Register();
-	// convert the partial dynamic mesh to a mesh description
-	// NOTE: this copies vertex weights to partial mesh description (later used to load weights into the tool)
-	FDynamicMeshToMeshDescription DnyToDescConverter;
-	constexpr bool bCopyTangents = true;
-	DnyToDescConverter.Convert(&PartialSubMesh.GetSubmesh(), *PartialMeshDescription, bCopyTangents);
-	
-	// reinitialize all mesh data structures
-	WeightTool->UpdateCurrentlyEditedMesh(SkeletalMeshComponent, PartialSubMesh.GetSubmesh(), *PartialMeshDescription.Get());
+	// queue an update for next tick
+	bIsolatedMeshNeedsUpdated = true;
 }
 
 void UWeightToolSelectionIsolator::RestoreFullMesh()
@@ -4597,6 +4539,70 @@
 	return PartialSubMesh.MapVertexToSubmesh(FullMeshVertexIndex);
 }
 
+void UWeightToolSelectionIsolator::CreatePartialMesh()
+{
+	UPolygonSelectionMechanic* SelectionMechanic = WeightTool->GetMainMeshSelector()->GetSelectionMechanic();
+	if (!ensure(SelectionMechanic))
+	{
+		return;
+	}
+
+	const USkeletalMeshComponent* SkeletalMeshComponent = GetSkeletalMeshComponent(WeightTool->GetTarget());
+	if (!ensure(SkeletalMeshComponent))
+	{
+		return;
+	}
+
+	const FMeshDescription* EditedMesh = WeightTool->GetCurrentlyEditedMeshDescription();
+	if (!ensure(EditedMesh))
+	{
+		return;
+	}
+	
+	// get the weights
+	FSkinToolWeights& Weights = WeightTool->GetWeights();
+
+	// apply previous changes
+	Weights.ApplyCurrentWeightsToMeshDescription(WeightTool->GetCurrentlyEditedMeshDescription());
+	
+	// put into ref pose, BEFORE copying the mesh, so that submesh deformer initializes with vertices in ref pose
+	Weights.Deformer.SetToRefPose(WeightTool);
+
+	// store selection to be restored
+	IsolatedSelectionToRestoreVertices.Reset();
+	IsolatedSelectionToRestoreEdges.Reset();
+	IsolatedSelectionToRestoreFaces.Reset();
+	IsolatedSelectionToRestoreVertices.ElementType = UE::Geometry::EGeometryElementType::Vertex;
+	IsolatedSelectionToRestoreEdges.ElementType = UE::Geometry::EGeometryElementType::Edge;
+	IsolatedSelectionToRestoreFaces.ElementType = UE::Geometry::EGeometryElementType::Face;
+	SelectionMechanic->GetSelection_AsTriangleTopology(IsolatedSelectionToRestoreVertices);
+	SelectionMechanic->GetSelection_AsTriangleTopology(IsolatedSelectionToRestoreEdges);
+	SelectionMechanic->GetSelection_AsTriangleTopology(IsolatedSelectionToRestoreFaces);
+
+	// store copy of original FDynamicMesh to restore
+	FDynamicMesh3 DynamicMesh;
+	FMeshDescriptionToDynamicMesh Converter;
+	Converter.Convert(EditedMesh, DynamicMesh);
+	FullDynamicMesh = MoveTemp(DynamicMesh);
+
+	// create a submesh from the given triangles
+	PartialSubMesh = UE::Geometry::FDynamicSubmesh3(&FullDynamicMesh, CurrentlyIsolatedTriangles);
+
+	// create mesh description for sub-mesh
+	PartialMeshDescription = MakeShared<FMeshDescription>();
+	// registering skeletal mesh attributes is required to create room to copy attributes during conversion from dynamic mesh
+	FSkeletalMeshAttributes Attributes(*PartialMeshDescription);
+	Attributes.Register();
+	// convert the partial dynamic mesh to a mesh description
+	// NOTE: this copies vertex weights to partial mesh description (later used to load weights into the tool)
+	FDynamicMeshToMeshDescription DnyToDescConverter;
+	constexpr bool bCopyTangents = true;
+	DnyToDescConverter.Convert(&PartialSubMesh.GetSubmesh(), *PartialMeshDescription, bCopyTangents);
+	
+	// reinitialize all mesh data structures
+	WeightTool->UpdateCurrentlyEditedMesh(SkeletalMeshComponent, PartialSubMesh.GetSubmesh(), *PartialMeshDescription.Get());
+}
+
 FCleanedEditMesh::FCleanedEditMesh(
 	const FDynamicMesh3& InDynamicMesh,
 	const FMeshDescription& InMeshDescription) : OriginalDynamicMesh(InDynamicMesh), OriginalMeshDescription(InMeshDescription)
diff --git a/Engine/Plugins/Experimental/MeshModelingToolsetExp/Source/MeshModelingToolsEditorOnlyExp/Public/SkeletalMesh/SkinWeightsPaintTool.h b/Engine/Plugins/Experimental/MeshModelingToolsetExp/Source/MeshModelingToolsEditorOnlyExp/Public/SkeletalMesh/SkinWeightsPaintTool.h
--- a/Engine/Plugins/Experimental/MeshModelingToolsetExp/Source/MeshModelingToolsEditorOnlyExp/Public/SkeletalMesh/SkinWeightsPaintTool.h
+++ b/Engine/Plugins/Experimental/MeshModelingToolsetExp/Source/MeshModelingToolsEditorOnlyExp/Public/SkeletalMesh/SkinWeightsPaintTool.h
@@ -615,7 +615,7 @@ public:
 	void InitialSetup(USkinWeightsPaintTool* InTool);
 
 	// call every tick to apply deferred changes to mesh
-	void OnTick(float DeltaTime);
+	void UpdateIsolatedSelection();
 
 	// returns true if any triangles are currently isolated
 	bool IsSelectionIsolated() const;
@@ -627,7 +627,7 @@ public:
 	void UnIsolateSelectionAsTransaction();
 
 	// isolate the array of triangles
-	void SetIsolatedTriangles(const TArray<int32>& TrianglesToIsolate);
+	void SetTrianglesToIsolate(const TArray<int32>& TrianglesToIsolate);
 	
 	// restores the whole mesh
 	void RestoreFullMesh();
@@ -640,6 +640,8 @@ public:
 	int32 FullToPartialMeshVertexIndex(int32 FullMeshVertexIndex) const;
 
 private:
+
+	void CreatePartialMesh();
 	
 	UPROPERTY()
 	TObjectPtr<USkinWeightsPaintTool> WeightTool;
@@ -651,7 +653,7 @@ private:
 	UE::Geometry::FGeometrySelection IsolatedSelectionToRestoreVertices;
 	UE::Geometry::FGeometrySelection IsolatedSelectionToRestoreEdges;
 	UE::Geometry::FGeometrySelection IsolatedSelectionToRestoreFaces;
-	bool bRestoreFullMeshOnNextTick = false;
+	bool bIsolatedMeshNeedsUpdated = false;
 	// isolate selection sub-meshes
 	UE::Geometry::FDynamicSubmesh3 PartialSubMesh;
 	UE::Geometry::FDynamicMesh3 FullDynamicMesh;
