Skin Weight Paint Tool: fix for non-manifold isolation
    1 - there's no need to convert the full mesh once again, use the currently edited mesh (which has already been cleaned) to avoid inconsistencies
    2 - the sub-mesh must update its non-manifold mapping, otherwise it still refers to the original mapping
    3 - mapping back for sub-mesh to original one must invert the non-manifold mapping as well

NOTE: FDynamicSubmesh3 should probably take care of all of this when the original mesh has a FNonManifoldSourceVertexIdAttribute attached to it so we wouldn't have to make all that jazz

#rnx

diff --git a/Engine/Plugins/Experimental/MeshModelingToolsetExp/Source/MeshModelingToolsEditorOnlyExp/Private/SkeletalMesh/SkinWeightsPaintTool.cpp b/Engine/Plugins/Experimental/MeshModelingToolsetExp/Source/MeshModelingToolsEditorOnlyExp/Private/SkeletalMesh/SkinWeightsPaintTool.cpp
--- a/Engine/Plugins/Experimental/MeshModelingToolsetExp/Source/MeshModelingToolsEditorOnlyExp/Private/SkeletalMesh/SkinWeightsPaintTool.cpp
+++ b/Engine/Plugins/Experimental/MeshModelingToolsetExp/Source/MeshModelingToolsEditorOnlyExp/Private/SkeletalMesh/SkinWeightsPaintTool.cpp
@@ -4476,17 +4476,21 @@ void UWeightToolSelectionIsolator::RestoreFullMesh()
 
 	// reinitialize with full mesh
 	FMeshDescription* CurrentlyEditedMesh = WeightTool->GetCurrentlyEditedMeshDescription();
-	WeightTool->UpdateCurrentlyEditedMesh(SkeletalMeshComponent, FullDynamicMesh, *CurrentlyEditedMesh);
+	FDynamicMesh3* CurrentlyEditedDynaMesh = WeightTool->GetCurrentlyEditedDynamicMesh();
+	WeightTool->UpdateCurrentlyEditedMesh(SkeletalMeshComponent, *CurrentlyEditedDynaMesh, *CurrentlyEditedMesh);
 
 	// copy the remapped weights back to the full mesh
 	const FSkeletalMeshConstAttributes MeshAttribs(*PartialMeshDescription.Get());
 	const FName ActiveProfile = WeightTool->GetWeightToolProperties()->GetActiveSkinWeightProfile();
 	const FSkinWeightsVertexAttributesConstRef AllVertexWeights = MeshAttribs.GetVertexSkinWeights(ActiveProfile);
 	const int32 NumVerticesInPartialMesh = PartialMeshDescription.Get()->Vertices().Num();
+	UE::Geometry::FNonManifoldMappingSupport NonManifoldMapping(*CurrentlyEditedDynaMesh);
 	for (int32 VertexIndexPartial = 0; VertexIndexPartial < NumVerticesInPartialMesh; VertexIndexPartial++)
 	{
 		// get the equivalent vertex on the full mesh
-		const int32 VertexIndexFull = PartialSubMesh.MapVertexToBaseMesh(VertexIndexPartial);
+		int32 VertexIndexFull = PartialSubMesh.MapVertexToBaseMesh(VertexIndexPartial);
+		VertexIndexFull = NonManifoldMapping.GetOriginalNonManifoldVertexID(VertexIndexFull);
+
 		// clear all the weights on this vertex
 		Weights.CurrentWeights[VertexIndexFull].Reset();
 		// replace with weights from partial mesh
@@ -4554,7 +4558,8 @@ void UWeightToolSelectionIsolator::CreatePartialMesh()
 	}
 
 	const FMeshDescription* EditedMesh = WeightTool->GetCurrentlyEditedMeshDescription();
-	if (!ensure(EditedMesh))
+	const FDynamicMesh3* EditedDynaMesh = WeightTool->GetCurrentlyEditedDynamicMesh();
+	if (!ensure(EditedMesh && EditedDynaMesh))
 	{
 		return;
 	}
@@ -4580,13 +4585,50 @@ void UWeightToolSelectionIsolator::CreatePartialMesh()
 	SelectionMechanic->GetSelection_AsTriangleTopology(IsolatedSelectionToRestoreFaces);
 
 	// store copy of original FDynamicMesh to restore
-	FDynamicMesh3 DynamicMesh;
-	FMeshDescriptionToDynamicMesh Converter;
-	Converter.Convert(EditedMesh, DynamicMesh);
-	FullDynamicMesh = MoveTemp(DynamicMesh);
+	auto GetIsolatedMesh = [this, EditedDynaMesh](const TArray<int32>& ToIsolate) -> const FDynamicMesh3*
+	{
+		using namespace UE::Geometry;
 
-	// create a submesh from the given triangles
-	PartialSubMesh = UE::Geometry::FDynamicSubmesh3(&FullDynamicMesh, CurrentlyIsolatedTriangles);
+		// create a sub-mesh from the triangles to isolate
+		PartialSubMesh = FDynamicSubmesh3(EditedDynaMesh, ToIsolate);
+		
+		if (ToIsolate.IsEmpty())
+		{
+			return EditedDynaMesh;
+		}
+
+		FDynamicMesh3& IsolatedMesh = PartialSubMesh.GetSubmesh();
+
+		// remap the non-manifold mapping from full mesh to partial mesh if needed
+		const FNonManifoldMappingSupport OriginalNonManifoldMapping(*EditedDynaMesh);
+		if (OriginalNonManifoldMapping.IsNonManifoldVertexInSource())
+		{
+			TArray<int32> IsolatedMeshVertexToNonManifoldVertexID;
+			IsolatedMeshVertexToNonManifoldVertexID.Reserve(IsolatedMesh.VertexCount());
+			for (int32 IsolatedVertexID = 0; IsolatedVertexID < IsolatedMesh.VertexCount(); ++IsolatedVertexID)
+			{
+				const int32 OrigMeshVertexID = PartialSubMesh.MapVertexToBaseMesh(IsolatedVertexID);
+				const int32 SrcBaseMeshVertexID = OriginalNonManifoldMapping.GetOriginalNonManifoldVertexID(OrigMeshVertexID);
+				int32 SrcSubMeshVertexID = PartialSubMesh.MapVertexToSubmesh(SrcBaseMeshVertexID);
+				if (SrcSubMeshVertexID == INDEX_NONE)
+				{
+					SrcSubMeshVertexID = IsolatedVertexID;
+				}
+				IsolatedMeshVertexToNonManifoldVertexID.Add(SrcSubMeshVertexID);
+			}
+
+			// replace the non-manifold vertex map in the cleaned mesh attributes
+			FNonManifoldMappingSupport CleanedNonManifoldMapping(IsolatedMesh);
+			CleanedNonManifoldMapping.AttachNonManifoldVertexMappingData(IsolatedMeshVertexToNonManifoldVertexID, IsolatedMesh);
+		}
+
+		// copy bone attributes
+		FDynamicMeshAttributeSet* SourceAttributes = IsolatedMesh.Attributes();
+		SourceAttributes->CopyBoneAttributes(*EditedDynaMesh->Attributes());
+		
+		return &IsolatedMesh;
+	};
+	const FDynamicMesh3* IsolatedMesh = GetIsolatedMesh(CurrentlyIsolatedTriangles);
 
 	// create mesh description for sub-mesh
 	PartialMeshDescription = MakeShared<FMeshDescription>();
diff --git a/Engine/Plugins/Experimental/MeshModelingToolsetExp/Source/MeshModelingToolsEditorOnlyExp/Public/SkeletalMesh/SkinWeightsPaintTool.h b/Engine/Plugins/Experimental/MeshModelingToolsetExp/Source/MeshModelingToolsEditorOnlyExp/Public/SkeletalMesh/SkinWeightsPaintTool.h
--- a/Engine/Plugins/Experimental/MeshModelingToolsetExp/Source/MeshModelingToolsEditorOnlyExp/Public/SkeletalMesh/SkinWeightsPaintTool.h
+++ b/Engine/Plugins/Experimental/MeshModelingToolsetExp/Source/MeshModelingToolsEditorOnlyExp/Public/SkeletalMesh/SkinWeightsPaintTool.h
@@ -656,7 +656,6 @@ private:
 	bool bIsolatedMeshNeedsUpdated = false;
 	// isolate selection sub-meshes
 	UE::Geometry::FDynamicSubmesh3 PartialSubMesh;
-	UE::Geometry::FDynamicMesh3 FullDynamicMesh;
 };
 
 class MESHMODELINGTOOLSEDITORONLYEXP_API FIsolateSelectionChange : public FToolCommandChange
