Fixes skin weights minor prune issue that can cause influences to stay when they should be pruned. Causes weird normalization after some operations (like Hammer and mirror weights)
Also fixes rare crash when loading tool.

diff --git a/Engine/Plugins/Experimental/MeshModelingToolsetExp/Source/MeshModelingToolsEditorOnlyExp/Private/SkeletalMesh/SkinWeightsPaintTool.cpp b/Engine/Plugins/Experimental/MeshModelingToolsetExp/Source/MeshModelingToolsEditorOnlyExp/Private/SkeletalMesh/SkinWeightsPaintTool.cpp
--- a/Engine/Plugins/Experimental/MeshModelingToolsetExp/Source/MeshModelingToolsEditorOnlyExp/Private/SkeletalMesh/SkinWeightsPaintTool.cpp
+++ b/Engine/Plugins/Experimental/MeshModelingToolsetExp/Source/MeshModelingToolsEditorOnlyExp/Private/SkeletalMesh/SkinWeightsPaintTool.cpp
@@ -1245,27 +1245,24 @@ void FMultiBoneWeightEdits::MergeSingleEdit(
 			break;
 		}
 	}
-
-	// don't prune influence unless they are actually connected to the vertex
-	if (bPruneInfluence && !bWasAlreadyAnInfluence)
-	{
-		bPruneInfluence = false;
-	}
-
+	
 	// record a modification of this vertex weight
 	FSingleBoneWeightEdits& BoneWeightEdit = PerBoneWeightEdits.FindOrAdd(BoneIndex);
 	BoneWeightEdit.BoneIndex = BoneIndex;
 	BoneWeightEdit.NewWeights.Add(VertexID, NewWeight);
 	BoneWeightEdit.OldWeights.FindOrAdd(VertexID, OldWeight);
 
-	if (bPruneInfluence)
+	// record when an influence is REMOVED (unless it was not connected to the vertex)
+	if (bPruneInfluence && bWasAlreadyAnInfluence)
 	{
 		// record that this influence was REMOVED from this vertex
 		BoneWeightEdit.VerticesRemovedFrom.AddUnique(VertexID);
 		// remove from the AddedTo list in the off chance that vertex was added by a prior edit (last edit takes precedence)
 		BoneWeightEdit.VerticesAddedTo.Remove(VertexID);
 	}
-	else if (!bWasAlreadyAnInfluence && NewWeight > MinimumWeightThreshold)
+
+	// record when an influence is ADDED
+	if (!bPruneInfluence && !bWasAlreadyAnInfluence)
 	{
 		// record that this influence was ADDED to this vertex
 		BoneWeightEdit.VerticesAddedTo.AddUnique(VertexID);
@@ -1949,6 +1946,16 @@ float FSkinToolWeights::SetCurrentFalloffAndGetMaxFalloffThisStroke(int32 Vertex
 
 void FSkinToolWeights::ApplyEditsToCurrentWeights(const FMultiBoneWeightEdits& Edits)
 {
+	// remove influences so that SetWeightOfBoneOnVertex doesn't have to
+	for (const TTuple<BoneIndex, FSingleBoneWeightEdits>& BoneWeightEdits : Edits.PerBoneWeightEdits)
+	{
+		BoneIndex InfluenceToRemove = BoneWeightEdits.Key;
+		for (const VertexIndex VertexID : BoneWeightEdits.Value.VerticesRemovedFrom)
+		{
+			RemoveInfluenceFromVertex(VertexID, InfluenceToRemove, CurrentWeights);
+		}
+	}
+	
 	// apply weight edits to the CurrentWeights data
 	for (const TTuple<BoneIndex, FSingleBoneWeightEdits>& BoneWeightEdits : Edits.PerBoneWeightEdits)
 	{
@@ -3619,6 +3626,7 @@ void USkinWeightsPaintTool::MirrorWeights(EAxis::Type Axis, EMirrorDirection Dir
 			float& Weight = NewBoneWeights.FindOrAdd(MirroredBoneIndex, 0.0f);
 			Weight += NewWeight;
 		}
+		TruncateWeightMap(NewBoneWeights);
 		NormalizeWeightMap(NewBoneWeights);
 
 		// apply weight edits
@@ -4146,6 +4154,7 @@ void USkinWeightsPaintTool::NormalizeWeightMap(TMap<BoneIndex, float>& InOutWeig
 	{
 		TotalWeight += Weight.Value;
 	}
+	
 	for (TTuple<BoneIndex, float>& Weight : InOutWeights)
 	{
 		Weight.Value /= TotalWeight > SMALL_NUMBER ? TotalWeight : 1.f;
@@ -4787,6 +4796,17 @@ FMeshDescription& FCleanedEditMesh::GetOriginalMeshDescription()
 
 bool USkinWeightsPaintTool::HasActiveSelectionOnMainMesh()
 {
+	if (!WeightToolProperties)
+	{
+		return false;
+	}
+
+	UWeightToolMeshSelector* MainMeshSelector = GetMainMeshSelector();
+	if (!MainMeshSelector)
+	{
+		return false;
+	}
+	
 	const bool bSelectingTargetMesh = WeightToolProperties->MeshSelectMode == EMeshTransferOption::Target;
 	const bool bTargetMeshHasSelection = GetMainMeshSelector()->IsAnyComponentSelected();
 	return bSelectingTargetMesh && bTargetMeshHasSelection;
