Skin Weight Editing: Added ability to select vertices based on influence count.
Added ability to clamp number of influences on selected vertices.

diff --git a/Engine/Plugins/Experimental/Animation/SkeletalMeshModelingTools/Source/SkeletalMeshModelingTools/Private/DetailCustomization/SkinWeightDetailCustomization.cpp b/Engine/Plugins/Experimental/Animation/SkeletalMeshModelingTools/Source/SkeletalMeshModelingTools/Private/DetailCustomization/SkinWeightDetailCustomization.cpp
--- a/Engine/Plugins/Experimental/Animation/SkeletalMeshModelingTools/Source/SkeletalMeshModelingTools/Private/DetailCustomization/SkinWeightDetailCustomization.cpp
+++ b/Engine/Plugins/Experimental/Animation/SkeletalMeshModelingTools/Source/SkeletalMeshModelingTools/Private/DetailCustomization/SkinWeightDetailCustomization.cpp
@@ -556,6 +556,53 @@ void FSkinWeightDetailCustomization::AddSelectionUI(IDetailLayoutBuilder& Detail
 				})
 			]
 		]
+
+		// SELECT BY INFLUENCE COUNT
+		+SVerticalBox::Slot()
+		[
+			SNew(SHorizontalBox)
+		
+			+SHorizontalBox::Slot()
+			.Padding(2.f, WeightEditVerticalPadding)
+			[
+				SNew(SButton)
+				.IsEnabled_Lambda([this]
+				{
+					// only allow selecting affected vertices on the target/main mesh
+					return ToolSettings->MeshSelectMode == EMeshTransferOption::Target;
+				})
+				.HAlign(HAlign_Center)
+				.VAlign(VAlign_Center)
+				.Text(LOCTEXT("InfluenceCountSelectionButtonLabel", "Influence Count"))
+				.ToolTipText(LOCTEXT("InfluenceCountSelectionTooltip",
+						"Select vertices that are affected by at least N influences."))
+				.OnClicked_Lambda([this]()
+				{
+					Tool->SelectByInfluenceCount(Tool->GetWeightToolProperties()->ClampSelectValue);
+					return FReply::Handled();
+				})
+			]
+
+			+SHorizontalBox::Slot()
+			.Padding(2.f, WeightEditVerticalPadding)
+			[
+				SNew(SSpinBox<float>)
+				.MinValue(1)
+				.MaxValue(24)
+				.Value_Lambda([this]()
+				{
+					return ToolSettings->ClampSelectValue;
+				})
+				.OnValueChanged_Lambda([this](float NewValue)
+				{
+					ToolSettings->ClampSelectValue = NewValue;
+				})
+				.OnValueCommitted_Lambda([this](float NewValue, ETextCommit::Type CommitType)
+				{
+					ToolSettings->SaveConfig();
+				})
+			]
+		]
 	];
 	
 	IDetailCategoryBuilder& EditWeightsCategory = DetailBuilder.EditCategory("Edit Weights", FText::GetEmpty(), ECategoryPriority::Important);
@@ -1124,6 +1171,51 @@ void FSkinWeightDetailCustomization::AddSelectionUI(IDetailLayoutBuilder& Detail
 		]
 	];
 
+	// CLAMP INFLUENCES category
+	EditWeightsCategory.AddCustomRow(LOCTEXT("ClampInfluencesRow", "Clamp"), false)
+	.WholeRowContent()
+	[
+		SNew(SHorizontalBox)
+		.IsEnabled_Lambda([this]{ return Tool->HasActiveSelectionOnMainMesh(); })
+
+		+SHorizontalBox::Slot()
+		.Padding(2.f, WeightEditVerticalPadding)
+		[
+			SNew(SButton)
+			.HAlign(HAlign_Center)
+			.Text(LOCTEXT("ClampInfluencesButtonLabel", "Clamp"))
+			.ToolTipText(LOCTEXT("ClampInfluencesButtonTooltip",
+				"Clamp the number of influences to not exceed the target value.\n"
+				"Removes smallest influences first.\n"
+				"This command operates on the selected vertices."))
+			.OnClicked_Lambda([this]()
+			{
+				ToolSettings->WeightTool->ClampInfluences(ToolSettings->ClampValue);
+				return FReply::Handled();
+			})
+		]
+
+		+SHorizontalBox::Slot()
+		.Padding(2.f, WeightEditVerticalPadding)
+		[
+			SNew(SSpinBox<float>)
+			.MinValue(1)
+			.MaxValue(24)
+			.Value_Lambda([this]()
+			{
+				return ToolSettings->ClampValue;
+			})
+			.OnValueChanged_Lambda([this](float NewValue)
+			{
+				ToolSettings->ClampValue = NewValue;
+			})
+			.OnValueCommitted_Lambda([this](float NewValue, ETextCommit::Type CommitType)
+			{
+				ToolSettings->SaveConfig();
+			})
+		]
+	];
+
 	// COPY/PASTE WEIGHTS category
 	EditWeightsCategory.AddCustomRow(LOCTEXT("CopyPasteWeightsRow", "Copy Paste"), false)
 	.WholeRowContent()
diff --git a/Engine/Plugins/Experimental/MeshModelingToolsetExp/Source/MeshModelingToolsEditorOnlyExp/Private/SkeletalMesh/SkinWeightsPaintTool.cpp b/Engine/Plugins/Experimental/MeshModelingToolsetExp/Source/MeshModelingToolsEditorOnlyExp/Private/SkeletalMesh/SkinWeightsPaintTool.cpp
--- a/Engine/Plugins/Experimental/MeshModelingToolsetExp/Source/MeshModelingToolsEditorOnlyExp/Private/SkeletalMesh/SkinWeightsPaintTool.cpp
+++ b/Engine/Plugins/Experimental/MeshModelingToolsetExp/Source/MeshModelingToolsEditorOnlyExp/Private/SkeletalMesh/SkinWeightsPaintTool.cpp
@@ -3960,6 +3960,79 @@ void USkinWeightsPaintTool::HammerWeights()
 	Weights.Deformer.SetAllVerticesToBeUpdated();
 }
 
+void USkinWeightsPaintTool::ClampInfluences(const int32 MaxInfluences)
+{
+	if (!ensure(MaxInfluences >= 1))
+	{
+		return;
+	}
+
+	// get all the selected vertices
+	const TArray<VertexIndex>& VerticesToClamp = MeshSelector->GetSelectedVertices();
+	if (VerticesToClamp.IsEmpty())
+	{
+		const FText NotificationText = FText::FromString(TEXT("No vertices were selected. No weights were clamped."));
+		ShowEditorMessage(ELogVerbosity::Error, NotificationText);
+		return;
+	}
+	
+	// create set of weight edits that clamp the smallest influences
+	FMultiBoneWeightEdits ClampedWeightEdits;
+
+	// clamp each vertex to MaxInfluences (discarding the smallest weights)
+	for (const VertexIndex VertexID : VerticesToClamp)
+	{
+		if (Weights.PreChangeWeights[VertexID].Num() <= MaxInfluences)
+		{
+			continue;
+		}
+		
+		VertexWeights WeightsToClamp = Weights.PreChangeWeights[VertexID];
+
+		// sort in descending order by weight
+		Algo::Sort(WeightsToClamp, [](const FVertexBoneWeight& A, const FVertexBoneWeight& B)
+			{
+				return A.Weight > B.Weight; 
+			});
+		// cull all the smallest weights up to MaxInfluences
+		WeightsToClamp.SetNum(MaxInfluences);
+
+		// normalize remaining influences
+		float TotalWeight = 0.f;
+		for (const FVertexBoneWeight& VertexWeight : WeightsToClamp)
+		{
+			TotalWeight += VertexWeight.Weight;
+		}
+		for (FVertexBoneWeight& VertexWeight : WeightsToClamp)
+		{
+			VertexWeight.Weight /= TotalWeight > SMALL_NUMBER ? TotalWeight : 1.f;
+		}
+
+		// remove all current weights (pruning since this operation completely replaces the weight)
+		for (const FVertexBoneWeight& BoneWeight : Weights.PreChangeWeights[VertexID])
+		{
+			constexpr float NewWeight = 0.f;
+			constexpr bool bPruneInfluence = true;
+			ClampedWeightEdits.MergeSingleEdit(BoneWeight.BoneID, VertexID, NewWeight, bPruneInfluence, Weights.PreChangeWeights);
+		}
+
+		// replace weights with clamped and renormalized weights
+		for (const FVertexBoneWeight& ClampedBoneWeight : WeightsToClamp)
+		{
+			const float NewWeight = ClampedBoneWeight.Weight;
+			constexpr bool bPruneInfluence = false;
+			ClampedWeightEdits.MergeSingleEdit(ClampedBoneWeight.BoneID, VertexID, NewWeight, bPruneInfluence, Weights.PreChangeWeights);
+		}
+	}
+	
+	// apply the changes
+	const FText TransactionLabel = LOCTEXT("ClampInfluencesChange", "Clamped influences.");
+	ApplyWeightEditsAsTransaction(ClampedWeightEdits, TransactionLabel);
+
+	// update deformations
+	Weights.Deformer.SetAllVerticesToBeUpdated();
+}
+
 const FString USkinWeightsPaintTool::CopyPasteWeightsIdentifier = TEXT("UNREAL_VERTEX_WEIGHTS:");
 
 void USkinWeightsPaintTool::CopyWeights()
@@ -4859,6 +4932,37 @@ void USkinWeightsPaintTool::SelectAffected() const
 	GetToolManager()->EndUndoTransaction();
 }
 
+void USkinWeightsPaintTool::SelectByInfluenceCount(const int32 MinInfluences) const
+{
+	UPolygonSelectionMechanic* SelectionMechanic = MeshSelector->GetSelectionMechanic();
+	if (!ensure(SelectionMechanic))
+	{
+		return;
+	}
+	
+	GetToolManager()->BeginUndoTransaction(LOCTEXT("InfluenceCountSelectionChange", "Select by Influence Count"));
+	SelectionMechanic->BeginChange();
+
+	// create selection set
+	FGroupTopologySelection Selection;
+
+	// spin through all vertices and find the ones that have at least MinInfluences number of bones affecting them
+	for (VertexIndex VertexID=0; VertexID<Weights.PreChangeWeights.Num(); ++VertexID)
+	{
+		const VertexWeights& SingleVertexWeights = Weights.PreChangeWeights[VertexID];
+		if (SingleVertexWeights.Num() >= MinInfluences)
+		{
+			Selection.SelectedCornerIDs.Add(VertexID);
+		}
+	}
+	
+	// select vertices
+	constexpr bool bBroadcast = true;
+	SelectionMechanic->SetSelection(Selection, bBroadcast);
+	SelectionMechanic->EndChangeAndEmitIfModified();
+	GetToolManager()->EndUndoTransaction();
+}
+
 void USkinWeightsPaintTool::GetVerticesAffectedByBone(BoneIndex IndexOfBone, TSet<int32>& OutVertexIndices) const
 {
 	VertexIndex VertexID = 0;
diff --git a/Engine/Plugins/Experimental/MeshModelingToolsetExp/Source/MeshModelingToolsEditorOnlyExp/Public/SkeletalMesh/SkinWeightsPaintTool.h b/Engine/Plugins/Experimental/MeshModelingToolsetExp/Source/MeshModelingToolsEditorOnlyExp/Public/SkeletalMesh/SkinWeightsPaintTool.h
--- a/Engine/Plugins/Experimental/MeshModelingToolsetExp/Source/MeshModelingToolsEditorOnlyExp/Public/SkeletalMesh/SkinWeightsPaintTool.h
+++ b/Engine/Plugins/Experimental/MeshModelingToolsetExp/Source/MeshModelingToolsEditorOnlyExp/Public/SkeletalMesh/SkinWeightsPaintTool.h
@@ -401,6 +401,10 @@ public:
 	UPROPERTY(Config)
 	float PruneValue = 0.01;
 	UPROPERTY(Config)
+	int32 ClampValue = 8;
+	UPROPERTY(Config)
+	int32 ClampSelectValue = 8;
+	UPROPERTY(Config)
 	float AddStrength = 1.0;
 	UPROPERTY(Config)
 	float ReplaceValue = 1.0;
@@ -738,6 +742,7 @@ public:
 	void AverageWeights(const float Strength);
 	void NormalizeWeights();
 	void HammerWeights();
+	void ClampInfluences(const int32 MaxInfluences);
 
 	// HELPER functions for modifying weights
 	//
@@ -778,6 +783,8 @@ public:
 	bool HasActiveSelectionOnMainMesh();
 	// select all vertices affected by the currently selected bone(s)
 	void SelectAffected() const;
+	// select all vertices affected by at least MinInfluences number of bones
+	void SelectByInfluenceCount(const int32 MinInfluences) const;
 
 	// get the average weight value of each influence on the given vertices
 	void GetInfluences(const TArray<int32>& VertexIndices, TArray<BoneIndex>& OutBoneIndices);
