Skeleton Modifier: fixed skeleton not being duplicated...

#rnx

diff --git a/Engine/Plugins/Experimental/MeshModelingToolsetExp/Source/SkeletalMeshModifiers/Private/SkeletonModifier.cpp b/Engine/Plugins/Experimental/MeshModelingToolsetExp/Source/SkeletalMeshModifiers/Private/SkeletonModifier.cpp
--- a/Engine/Plugins/Experimental/MeshModelingToolsetExp/Source/SkeletalMeshModifiers/Private/SkeletonModifier.cpp
+++ b/Engine/Plugins/Experimental/MeshModelingToolsetExp/Source/SkeletalMeshModifiers/Private/SkeletonModifier.cpp
@@ -5,6 +5,7 @@
 
 #include "AssetNotifications.h"
 #include "BoneWeights.h"
+#include "FileHelpers.h"
 #include "MeshDescription.h"
 #include "Animation/Skeleton.h"
 #include "ReferenceSkeleton.h"
@@ -604,6 +605,22 @@ bool USkeletonModifier::CommitSkeletonToSkeletalMesh()
 		UE_LOG(LogAnimation, Warning, TEXT("Skeleton Modifier: Skeleton can't be modified."));
 		return false;
 	}
+
+	// duplicate the skeleton if needed
+	if (SkeletonModifications == ESkeletonModificationType::DuplicateAndMerge)
+	{
+		TArray<UObject*> SavedObjects;
+		FEditorFileUtils::SaveAssetsAs({SkeletalMesh->GetSkeleton()}, SavedObjects);
+
+		USkeleton* NewSkeleton = SavedObjects.IsEmpty() ? nullptr : Cast<USkeleton>(SavedObjects[0]);
+		if (!NewSkeleton)
+		{
+			UE_LOG(LogAnimation, Warning, TEXT("Skeleton Modifier: Skeleton can't be duplicated."));
+			return false;
+		}
+
+		SkeletalMesh->SetSkeleton(NewSkeleton);
+	}
 	
 	const TArray<FMeshBoneInfo>& BoneInfos = ReferenceSkeleton->GetRawRefBoneInfo();
 	
