Skin weight tool: fixes regression causes weight edits made during an isolated selection to be lost.

diff --git a/Engine/Plugins/Experimental/MeshModelingToolsetExp/Source/MeshModelingToolsEditorOnlyExp/Private/SkeletalMesh/SkinWeightsPaintTool.cpp b/Engine/Plugins/Experimental/MeshModelingToolsetExp/Source/MeshModelingToolsEditorOnlyExp/Private/SkeletalMesh/SkinWeightsPaintTool.cpp
--- a/Engine/Plugins/Experimental/MeshModelingToolsetExp/Source/MeshModelingToolsEditorOnlyExp/Private/SkeletalMesh/SkinWeightsPaintTool.cpp
+++ b/Engine/Plugins/Experimental/MeshModelingToolsetExp/Source/MeshModelingToolsEditorOnlyExp/Private/SkeletalMesh/SkinWeightsPaintTool.cpp
@@ -4559,6 +4559,7 @@ void UWeightToolSelectionIsolator::RestoreFullMesh()
 	// reinitialize with full mesh
 	FMeshDescription* CurrentlyEditedMesh = WeightTool->GetCurrentlyEditedMeshDescription();
 	FDynamicMesh3* CurrentlyEditedDynaMesh = WeightTool->GetCurrentlyEditedDynamicMesh();
+	// this resizes Weights to the full mesh size
 	WeightTool->UpdateCurrentlyEditedMesh(SkeletalMeshComponent, *CurrentlyEditedDynaMesh, *CurrentlyEditedMesh);
 
 	// copy the remapped weights back to the full mesh
@@ -4639,9 +4640,14 @@ void UWeightToolSelectionIsolator::CreatePartialMesh()
 		return;
 	}
 
-	const FMeshDescription* EditedMesh = WeightTool->GetCurrentlyEditedMeshDescription();
-	const FDynamicMesh3* EditedDynaMesh = WeightTool->GetCurrentlyEditedDynamicMesh();
-	if (!ensure(EditedMesh && EditedDynaMesh))
+	if (!ensure(!CurrentlyIsolatedTriangles.IsEmpty()))
+	{
+		return;
+	}
+
+	FMeshDescription* EditedMeshDesc = WeightTool->GetCurrentlyEditedMeshDescription();
+	FDynamicMesh3* EditedDynaMesh = WeightTool->GetCurrentlyEditedDynamicMesh();
+	if (!ensure(EditedMeshDesc && EditedDynaMesh))
 	{
 		return;
 	}
@@ -4649,8 +4655,8 @@ void UWeightToolSelectionIsolator::CreatePartialMesh()
 	// get the weights
 	FSkinToolWeights& Weights = WeightTool->GetWeights();
 
-	// apply previous changes
-	Weights.ApplyCurrentWeightsToMeshDescription(WeightTool->GetCurrentlyEditedMeshDescription());
+	// apply the current weights to the full mesh description
+	Weights.ApplyCurrentWeightsToMeshDescription(EditedMeshDesc);
 	
 	// put into ref pose, BEFORE copying the mesh, so that submesh deformer initializes with vertices in ref pose
 	Weights.Deformer.SetToRefPose(WeightTool);
@@ -4666,52 +4672,43 @@ void UWeightToolSelectionIsolator::CreatePartialMesh()
 	SelectionMechanic->GetSelection_AsTriangleTopology(IsolatedSelectionToRestoreEdges);
 	SelectionMechanic->GetSelection_AsTriangleTopology(IsolatedSelectionToRestoreFaces);
 
-	// store copy of original FDynamicMesh to restore
-	auto GetIsolatedMesh = [this, EditedDynaMesh](const TArray<int32>& ToIsolate) -> const FDynamicMesh3*
-	{
-		using namespace UE::Geometry;
-
-		// create a sub-mesh from the triangles to isolate
-		PartialSubMesh = FDynamicSubmesh3(EditedDynaMesh, ToIsolate);
-		
-		if (ToIsolate.IsEmpty())
-		{
-			return EditedDynaMesh;
-		}
+	// copy the full mesh description data to the full dynamic mesh
+	// 
+	// NOTE: this copies skin weights to the dynamic mesh.
+	//	 	 this is important because the partial mesh is created from the dynamic mesh NOT the mesh description where the weight have been applied
+	// NOTE: this conversion also restores NonManifold mappings to the converted dynamic mesh
+	EditedDynaMesh->Clear();
+	FMeshDescriptionToDynamicMesh Converter;
+	Converter.Convert(EditedMeshDesc, *EditedDynaMesh);
+	
+	// create a partial sub-mesh from a subset of triangles on the dynamic mesh
+	PartialSubMesh = UE::Geometry::FDynamicSubmesh3(EditedDynaMesh, CurrentlyIsolatedTriangles);
 
+	// remap the non-manifold mapping from full mesh to partial mesh if needed
+	const UE::Geometry::FNonManifoldMappingSupport OriginalNonManifoldMapping(*EditedDynaMesh);
+	if (OriginalNonManifoldMapping.IsNonManifoldVertexInSource())
+	{
 		FDynamicMesh3& IsolatedMesh = PartialSubMesh.GetSubmesh();
-
-		// remap the non-manifold mapping from full mesh to partial mesh if needed
-		const FNonManifoldMappingSupport OriginalNonManifoldMapping(*EditedDynaMesh);
-		if (OriginalNonManifoldMapping.IsNonManifoldVertexInSource())
-		{
-			TArray<int32> IsolatedMeshVertexToNonManifoldVertexID;
-			IsolatedMeshVertexToNonManifoldVertexID.Reserve(IsolatedMesh.VertexCount());
-			for (int32 IsolatedVertexID = 0; IsolatedVertexID < IsolatedMesh.VertexCount(); ++IsolatedVertexID)
+		
+		TArray<int32> IsolatedMeshVertexToNonManifoldVertexID;
+		IsolatedMeshVertexToNonManifoldVertexID.Reserve(IsolatedMesh.VertexCount());
+		for (int32 IsolatedVertexID = 0; IsolatedVertexID < IsolatedMesh.VertexCount(); ++IsolatedVertexID)
+		{
+			const int32 OrigMeshVertexID = PartialSubMesh.MapVertexToBaseMesh(IsolatedVertexID);
+			const int32 SrcBaseMeshVertexID = OriginalNonManifoldMapping.GetOriginalNonManifoldVertexID(OrigMeshVertexID);
+			int32 SrcSubMeshVertexID = PartialSubMesh.MapVertexToSubmesh(SrcBaseMeshVertexID);
+			if (SrcSubMeshVertexID == INDEX_NONE)
 			{
-				const int32 OrigMeshVertexID = PartialSubMesh.MapVertexToBaseMesh(IsolatedVertexID);
-				const int32 SrcBaseMeshVertexID = OriginalNonManifoldMapping.GetOriginalNonManifoldVertexID(OrigMeshVertexID);
-				int32 SrcSubMeshVertexID = PartialSubMesh.MapVertexToSubmesh(SrcBaseMeshVertexID);
-				if (SrcSubMeshVertexID == INDEX_NONE)
-				{
-					SrcSubMeshVertexID = IsolatedVertexID;
-				}
-				IsolatedMeshVertexToNonManifoldVertexID.Add(SrcSubMeshVertexID);
+				SrcSubMeshVertexID = IsolatedVertexID;
 			}
-
-			// replace the non-manifold vertex map in the cleaned mesh attributes
-			FNonManifoldMappingSupport CleanedNonManifoldMapping(IsolatedMesh);
-			CleanedNonManifoldMapping.AttachNonManifoldVertexMappingData(IsolatedMeshVertexToNonManifoldVertexID, IsolatedMesh);
+			IsolatedMeshVertexToNonManifoldVertexID.Add(SrcSubMeshVertexID);
 		}
 
-		// copy bone attributes
-		FDynamicMeshAttributeSet* SourceAttributes = IsolatedMesh.Attributes();
-		SourceAttributes->CopyBoneAttributes(*EditedDynaMesh->Attributes());
-		
-		return &IsolatedMesh;
-	};
-	const FDynamicMesh3* IsolatedMesh = GetIsolatedMesh(CurrentlyIsolatedTriangles);
-
+		// replace the non-manifold vertex map in the cleaned mesh attributes
+		UE::Geometry::FNonManifoldMappingSupport CleanedNonManifoldMapping(IsolatedMesh);
+		CleanedNonManifoldMapping.AttachNonManifoldVertexMappingData(IsolatedMeshVertexToNonManifoldVertexID, IsolatedMesh);
+	}
+	
 	// create mesh description for sub-mesh
 	PartialMeshDescription = MakeShared<FMeshDescription>();
 	// registering skeletal mesh attributes is required to create room to copy attributes during conversion from dynamic mesh
