MeshModelingTools: Minor tweaks to MeshSculptToolBase to make subclasses able to disable work plane grid line drawing, and make brushes able to share a property set.

diff --git a/Engine/Plugins/Experimental/MeshModelingToolsetExp/Source/MeshModelingToolsExp/Private/Sculpting/MeshSculptToolBase.cpp b/Engine/Plugins/Experimental/MeshModelingToolsetExp/Source/MeshModelingToolsExp/Private/Sculpting/MeshSculptToolBase.cpp
--- a/Engine/Plugins/Experimental/MeshModelingToolsetExp/Source/MeshModelingToolsExp/Private/Sculpting/MeshSculptToolBase.cpp
+++ b/Engine/Plugins/Experimental/MeshModelingToolsetExp/Source/MeshModelingToolsExp/Private/Sculpting/MeshSculptToolBase.cpp
@@ -280,7 +280,7 @@ void UMeshSculptToolBase::Render(IToolsContextRenderAPI* RenderAPI)
 		BrushIndicatorMaterial->SetScalarParameterValue(TEXT("FalloffWidth"), FixedDimScale);
 	}
 
-	if (ShowWorkPlane())
+	if (ShowWorkPlane() && bDrawWorkPlaneGridLines)
 	{
 		FPrimitiveDrawInterface* PDI = RenderAPI->GetPrimitiveDrawInterface();
 		FColor GridColor(128, 128, 128, 32);
@@ -340,7 +340,10 @@ void UMeshSculptToolBase::RegisterBrushType(int32 Identifier, FText Name, TUniqu
 	BrushOpPropSets.Add(Identifier, PropSet);
 	BrushOpFactories.Add(Identifier, MoveTemp(Factory));
 
-	AddToolPropertySource(PropSet);
+	if (!ToolPropertyObjects.Contains(PropSet))
+	{
+		AddToolPropertySource(PropSet);
+	}
 	SetToolPropertySourceEnabled(PropSet, false);
 }
 
@@ -355,7 +358,10 @@ void UMeshSculptToolBase::RegisterSecondaryBrushType(int32 Identifier, FText Nam
 	SecondaryBrushOpPropSets.Add(Identifier, PropSet);
 	SecondaryBrushOpFactories.Add(Identifier, MoveTemp(Factory));
 
-	AddToolPropertySource(PropSet);
+	if (!ToolPropertyObjects.Contains(PropSet))
+	{
+		AddToolPropertySource(PropSet);
+	}
 	SetToolPropertySourceEnabled(PropSet, false);
 }
 
@@ -390,7 +396,7 @@ void UMeshSculptToolBase::SetActivePrimaryBrushType(int32 Identifier)
 	TUniquePtr<FMeshSculptBrushOpFactory>* Factory = BrushOpFactories.Find(Identifier);
 	if (Factory == nullptr)
 	{
-		check(false);
+		ensure(false);
 		return;
 	}
 
@@ -420,7 +426,7 @@ void UMeshSculptToolBase::SetActiveSecondaryBrushType(int32 Identifier)
 	TUniquePtr<FMeshSculptBrushOpFactory>* Factory = SecondaryBrushOpFactories.Find(Identifier);
 	if (Factory == nullptr)
 	{
-		check(false);
+		ensure(false);
 		return;
 	}
 
diff --git a/Engine/Plugins/Experimental/MeshModelingToolsetExp/Source/MeshModelingToolsExp/Public/Sculpting/MeshSculptToolBase.h b/Engine/Plugins/Experimental/MeshModelingToolsetExp/Source/MeshModelingToolsExp/Public/Sculpting/MeshSculptToolBase.h
--- a/Engine/Plugins/Experimental/MeshModelingToolsetExp/Source/MeshModelingToolsExp/Public/Sculpting/MeshSculptToolBase.h
+++ b/Engine/Plugins/Experimental/MeshModelingToolsetExp/Source/MeshModelingToolsExp/Public/Sculpting/MeshSculptToolBase.h
@@ -583,6 +583,9 @@ public:
 protected:
 	virtual void UpdateWorkPlane();
 	virtual bool ShowWorkPlane() const { return false; };
+	// Only relevant when ShowWorkPlane() returns true. Determines whether the plane
+	//  grid lines are drawn.
+	bool bDrawWorkPlaneGridLines = true;
 
 protected:
 	TValueWatcher<FVector> GizmoPositionWatcher;
