Skin paint tool: removed auto-toggle to bone colors when entering tool

diff --git a/Engine/Plugins/Experimental/MeshModelingToolsetExp/Source/MeshModelingToolsEditorOnlyExp/Private/SkeletalMesh/SkinWeightsPaintTool.cpp b/Engine/Plugins/Experimental/MeshModelingToolsetExp/Source/MeshModelingToolsEditorOnlyExp/Private/SkeletalMesh/SkinWeightsPaintTool.cpp
--- a/Engine/Plugins/Experimental/MeshModelingToolsetExp/Source/MeshModelingToolsEditorOnlyExp/Private/SkeletalMesh/SkinWeightsPaintTool.cpp
+++ b/Engine/Plugins/Experimental/MeshModelingToolsetExp/Source/MeshModelingToolsEditorOnlyExp/Private/SkeletalMesh/SkinWeightsPaintTool.cpp
@@ -2242,9 +2242,7 @@ void USkinWeightsPaintTool::Setup()
 	{
 		ViewportClient->SetViewMode(VMI_Lit_Wireframe);
 	}
-	// turn on bone colors
-	bBoneColorsToRestore = GetDefault<UPersonaOptions>()->bShowBoneColors;
-	GetMutableDefault<UPersonaOptions>()->bShowBoneColors = true;
+	
 	// set focus to viewport so brush hotkey works
 	SetFocusInViewport();
 	
@@ -3238,7 +3236,6 @@ void USkinWeightsPaintTool::OnShutdown(EToolShutdownType ShutdownType)
 	// restore viewport show flags and preview settings
 	FPreviewProfileController PreviewProfileController;
 	PreviewProfileController.SetActiveProfile(PreviewProfileToRestore);
-	GetMutableDefault<UPersonaOptions>()->bShowBoneColors = bBoneColorsToRestore;
 
 	if (EditorContext.IsValid())
 	{
diff --git a/Engine/Plugins/Experimental/MeshModelingToolsetExp/Source/MeshModelingToolsEditorOnlyExp/Public/SkeletalMesh/SkinWeightsPaintTool.h b/Engine/Plugins/Experimental/MeshModelingToolsetExp/Source/MeshModelingToolsEditorOnlyExp/Public/SkeletalMesh/SkinWeightsPaintTool.h
--- a/Engine/Plugins/Experimental/MeshModelingToolsetExp/Source/MeshModelingToolsEditorOnlyExp/Public/SkeletalMesh/SkinWeightsPaintTool.h
+++ b/Engine/Plugins/Experimental/MeshModelingToolsetExp/Source/MeshModelingToolsEditorOnlyExp/Public/SkeletalMesh/SkinWeightsPaintTool.h
@@ -973,7 +973,6 @@ protected:
 	
 	// editor state to restore when exiting the paint tool
 	FString PreviewProfileToRestore;
-	bool bBoneColorsToRestore;
 
 	friend SkinPaintTool::FSkinToolDeformer;
 };
