Skin paint tool, fixed incorrect ensure() condition in transfer weights feature preventing transfer from same mesh/LOD with different profile.

diff --git a/Engine/Plugins/Experimental/MeshModelingToolsetExp/Source/MeshModelingToolsEditorOnlyExp/Private/SkeletalMesh/SkinWeightsPaintTool.cpp b/Engine/Plugins/Experimental/MeshModelingToolsetExp/Source/MeshModelingToolsEditorOnlyExp/Private/SkeletalMesh/SkinWeightsPaintTool.cpp
--- a/Engine/Plugins/Experimental/MeshModelingToolsetExp/Source/MeshModelingToolsEditorOnlyExp/Private/SkeletalMesh/SkinWeightsPaintTool.cpp
+++ b/Engine/Plugins/Experimental/MeshModelingToolsetExp/Source/MeshModelingToolsEditorOnlyExp/Private/SkeletalMesh/SkinWeightsPaintTool.cpp
@@ -1071,9 +1071,10 @@ void UWeightToolTransferManager::TransferWeightsFromSameMeshAndLOD()
 	const FDynamicMeshVertexSkinWeightsAttribute* SourceAttributes = TargetAttributes->GetSkinWeightsAttribute(ToolProperties->SourceSkinWeightProfile);
 	check(SourceAttributes != nullptr);
 
-	if (!ensure(SourceAttributes == TransferredSkinWeights))
+	// this function assumes we are transferring between: Same Mesh, Same LOD, DIFFERENT Profile so attributes must not be identical.
+	if (!ensure(SourceAttributes != TransferredSkinWeights))
 	{
-		const FText NotificationText = LOCTEXT("TransferBetweenSame", "Mesh attribute mismatch. No weights were transferred.");
+		const FText NotificationText = LOCTEXT("TransferBetweenSame", "Cannot transfer between same LOD & profile. No weights were transferred.");
 		ShowEditorMessage(ELogVerbosity::Error, NotificationText);
 		return;
 	}
