Approximate Actors: Prevent issue where there's no mesh generated in the "solidify" phase, causing subsequent passes to crash.

diff --git a/Engine/Plugins/Experimental/MeshModelingToolsetExp/Source/GeometryProcessingAdapters/Private/ApproximateActorsImpl.cpp b/Engine/Plugins/Experimental/MeshModelingToolsetExp/Source/GeometryProcessingAdapters/Private/ApproximateActorsImpl.cpp
--- a/Engine/Plugins/Experimental/MeshModelingToolsetExp/Source/GeometryProcessingAdapters/Private/ApproximateActorsImpl.cpp
+++ b/Engine/Plugins/Experimental/MeshModelingToolsetExp/Source/GeometryProcessingAdapters/Private/ApproximateActorsImpl.cpp
@@ -468,6 +468,14 @@ static TSharedPtr<FApproximationMeshData> GenerateApproximationMesh(
 		SolidMesh.DiscardAttributes();
 	}
 
+	// if solid mesh has no triangles, something has gone wrong
+	if (SolidMesh.TriangleCount() == 0)
+	{
+		UE_LOG(LogApproximateActors, Warning, TEXT("Solidify mesh has no triangles - unable to generate mesh"));
+		Result->ResultCode = IGeometryProcessing_ApproximateActors::EResultCode::MeshGenerationFailed;
+		return Result;
+	}
+
 	// CurResultMesh will point to the "current" result and we will update this pointer as we
 	// step through various stages of the generation process
 	FDynamicMesh3* CurResultMesh = &SolidMesh;		
@@ -1430,4 +1438,4 @@ UStaticMesh* FApproximateActorsImpl::EmitGeneratedMeshAsset(
 }
 
 
-#undef LOCTEXT_NAMESPACE
\ No newline at end of file
+#undef LOCTEXT_NAMESPACE
