ModelingTools: Fixed issue where focusing the viewport in sculpt tools could focus on world origin.

diff --git a/Engine/Plugins/Experimental/MeshModelingToolsetExp/Source/MeshModelingToolsExp/Public/Sculpting/MeshSculptToolBase.h b/Engine/Plugins/Experimental/MeshModelingToolsetExp/Source/MeshModelingToolsExp/Public/Sculpting/MeshSculptToolBase.h
--- a/Engine/Plugins/Experimental/MeshModelingToolsetExp/Source/MeshModelingToolsExp/Public/Sculpting/MeshSculptToolBase.h
+++ b/Engine/Plugins/Experimental/MeshModelingToolsetExp/Source/MeshModelingToolsExp/Public/Sculpting/MeshSculptToolBase.h
@@ -426,7 +426,7 @@ public:
 protected:
 	FFrame3d LastBrushFrameWorld;
 	FFrame3d LastBrushFrameLocal;
-	int32 LastBrushTriangleID;
+	int32 LastBrushTriangleID = IndexConstants::InvalidID;
 
 	const FFrame3d& GetBrushFrameWorld() const { return LastBrushFrameWorld; }
 	const FFrame3d& GetBrushFrameLocal() const { return LastBrushFrameLocal; }
