Skin paint tool: fixed transfer when using same mesh+LOD, with different profile and different selection on source/target.

diff --git a/Engine/Plugins/Experimental/MeshModelingToolsetExp/Source/MeshModelingToolsEditorOnlyExp/Private/SkeletalMesh/SkinWeightsPaintTool.cpp b/Engine/Plugins/Experimental/MeshModelingToolsetExp/Source/MeshModelingToolsEditorOnlyExp/Private/SkeletalMesh/SkinWeightsPaintTool.cpp
--- a/Engine/Plugins/Experimental/MeshModelingToolsetExp/Source/MeshModelingToolsEditorOnlyExp/Private/SkeletalMesh/SkinWeightsPaintTool.cpp
+++ b/Engine/Plugins/Experimental/MeshModelingToolsetExp/Source/MeshModelingToolsEditorOnlyExp/Private/SkeletalMesh/SkinWeightsPaintTool.cpp
@@ -857,10 +857,13 @@ void UWeightToolTransferManager::TransferWeights()
 	const bool bSameMesh = SourceSkeletalMesh == TargetSkeletalMesh;
 	const bool bSameLOD = GetLODId(ToolProperties->ActiveLOD) == GetLODId(ToolProperties->SourceLOD);
 	const bool bSameProfile = ToolProperties->ActiveSkinWeightProfile == ToolProperties->SourceSkinWeightProfile;
-	const bool bHasSelectedVertices = ToolProperties->EditingMode == EWeightEditMode::Mesh && !WeightTool->GetMainMeshSelector()->GetSelectedVertices().IsEmpty();
+	const bool bVerticesSelectedOnTarget = ToolProperties->EditingMode == EWeightEditMode::Mesh && !WeightTool->GetMainMeshSelector()->GetSelectedVertices().IsEmpty();
+	const bool bVerticesSelectedOnSource = ToolProperties->EditingMode == EWeightEditMode::Mesh && !MeshSelector->GetSelectedVertices().IsEmpty();
+	const bool bHasAnySelectedVertices = bVerticesSelectedOnSource || bVerticesSelectedOnTarget;
+	const bool bOneMeshHasSelectedVertices = bHasAnySelectedVertices && !(bVerticesSelectedOnSource && bVerticesSelectedOnTarget);
 	
 	// cannot transfer between same mesh/LOD/profile without selection (identical weights)
-	if (bSameMesh && bSameLOD && bSameProfile && !bHasSelectedVertices)
+	if (bSameMesh && bSameLOD && bSameProfile && !bHasAnySelectedVertices)
 	{
 		const FText NotificationText = LOCTEXT("IdenticalSourceAndTarget", "Cannot copy weights from the same mesh, LOD and profile without anything selected. No weights were transferred.");
 		ShowEditorMessage(ELogVerbosity::Error, NotificationText);
@@ -868,7 +871,7 @@ void UWeightToolTransferManager::TransferWeights()
 	}
 
 	// can weights be transferred copying the attribute directly?
-	if (bSameMesh && bSameLOD && !bSameProfile)
+	if (bSameMesh && bSameLOD && !bSameProfile && bOneMeshHasSelectedVertices)
 	{
 		TransferWeightsFromSameMeshAndLOD();
 	}
@@ -1079,10 +1082,12 @@ void UWeightToolTransferManager::TransferWeightsFromSameMeshAndLOD()
 		return;
 	}
 
-	TArray<int32> VertexSubset;
-	if (ToolProperties->EditingMode == EWeightEditMode::Mesh)
+	// get subset of vertices to transfer weights on (from either source or target)
+	TArray<int32> VertexSubset = MeshSelector->GetSelectedVertices();
+	if (VertexSubset.IsEmpty())
 	{
-		VertexSubset = MeshSelector->GetSelectedVertices();
+		// get vertex selection from the target mesh instead
+		VertexSubset = WeightTool->GetMainMeshSelector()->GetSelectedVertices();
 	}
 
 	if (VertexSubset.IsEmpty())
