MeshModelingTools: Make the vertex sculpt "align to normal" hotkey work properly when the hover stamp is not aligned to mesh.

diff --git a/Engine/Plugins/Experimental/MeshModelingToolsetExp/Source/MeshModelingToolsExp/Private/Sculpting/MeshSculptToolBase.cpp b/Engine/Plugins/Experimental/MeshModelingToolsetExp/Source/MeshModelingToolsExp/Private/Sculpting/MeshSculptToolBase.cpp
--- a/Engine/Plugins/Experimental/MeshModelingToolsetExp/Source/MeshModelingToolsExp/Private/Sculpting/MeshSculptToolBase.cpp
+++ b/Engine/Plugins/Experimental/MeshModelingToolsetExp/Source/MeshModelingToolsExp/Private/Sculpting/MeshSculptToolBase.cpp
@@ -1261,7 +1261,8 @@ void UMeshSculptToolBase::UpdateWorkPlane()
 		FHitResult Result;
 		bool bWorldHit = ToolSceneQueriesUtil::FindNearestVisibleObjectHit(this, Result, CursorWorldRay);
 		FRay3d LocalRay = GetLocalRay(CursorWorldRay);
-		bool bObjectHit = ( FindHitSculptMeshTriangle(LocalRay) != IndexConstants::InvalidID );
+		int32 SculptMeshTid = FindHitSculptMeshTriangle(LocalRay);
+		bool bObjectHit = (SculptMeshTid != IndexConstants::InvalidID);
 		if (bWorldHit && 
 			(bObjectHit == false || (CursorWorldRay.GetParameter(Result.ImpactPoint) < CursorWorldRay.GetParameter((FVector)HoverStamp.WorldFrame.Origin))))
 		{
@@ -1269,7 +1270,20 @@ void UMeshSculptToolBase::UpdateWorkPlane()
 		}
 		else
 		{
-			SetFixedSculptPlaneFromWorldPos((FVector)HoverStamp.WorldFrame.Origin, (FVector)HoverStamp.WorldFrame.Z(), PendingWorkPlaneUpdate);
+			FVector3d Normal = HoverStamp.WorldFrame.Z();
+			// We're updating from a sculpt mesh hit. We can expect HoverStamp to be at the proper location,
+			//  but the alignment might differ depending on how the cient is brushing, so get the triangle
+			//  normal if we need it.
+			if (bObjectHit && PendingWorkPlaneUpdate == EPendingWorkPlaneUpdate::MoveToHitPositionNormal)
+			{
+				FDynamicMesh3* Mesh = GetSculptMesh();
+				if (ensure(Mesh))
+				{
+					Normal = Mesh->GetTriNormal(SculptMeshTid);
+				}
+			}
+
+			SetFixedSculptPlaneFromWorldPos((FVector)HoverStamp.WorldFrame.Origin, Normal, PendingWorkPlaneUpdate);
 		}
 		PendingWorkPlaneUpdate = EPendingWorkPlaneUpdate::NoUpdatePending;
 	}
