MeshModelingTools: Refactors to MeshVertexSculptTool and its brush system to allow subclasses to register their own set of brushes. This involves moving various brush-specific information to live on the brushes instead of being hardcoded into the tool. Also add support for cylindrical brush regions.

diff --git a/Engine/Plugins/Editor/ModelingToolsEditorMode/Source/ModelingToolsEditorMode/Private/DetailsCustomizations/MeshVertexSculptToolCustomizations.cpp b/Engine/Plugins/Editor/ModelingToolsEditorMode/Source/ModelingToolsEditorMode/Private/DetailsCustomizations/MeshVertexSculptToolCustomizations.cpp
--- a/Engine/Plugins/Editor/ModelingToolsEditorMode/Source/ModelingToolsEditorMode/Private/DetailsCustomizations/MeshVertexSculptToolCustomizations.cpp
+++ b/Engine/Plugins/Editor/ModelingToolsEditorMode/Source/ModelingToolsEditorMode/Private/DetailsCustomizations/MeshVertexSculptToolCustomizations.cpp
@@ -119,7 +119,7 @@ FVertexBrushSculptPropertiesDetails::~FVertexBrushSculptPropertiesDetails()
 
 void FVertexBrushSculptPropertiesDetails::CustomizeDetails(IDetailLayoutBuilder& DetailBuilder)
 {
-	TSharedPtr<IPropertyHandle> BrushTypeHandle = DetailBuilder.GetProperty(GET_MEMBER_NAME_CHECKED(UVertexBrushSculptProperties, PrimaryBrushType), UVertexBrushSculptProperties::StaticClass());
+	TSharedPtr<IPropertyHandle> BrushTypeHandle = DetailBuilder.GetProperty(GET_MEMBER_NAME_CHECKED(UVertexBrushSculptProperties, PrimaryBrushID), UVertexBrushSculptProperties::StaticClass());
 	ensure(BrushTypeHandle->IsValidHandle());
 
 	DetailBuilder.GetObjectsBeingCustomized(ObjectsBeingCustomized);
@@ -129,7 +129,7 @@ void FVertexBrushSculptPropertiesDetails::CustomizeDetails(IDetailLayoutBuilder&
 	TargetTool = Tool;
 
 	const TArray<UMeshSculptToolBase::FBrushTypeInfo>& BrushTypeInfos = Tool->GetRegisteredPrimaryBrushTypes();
-	int32 CurrentBrushType = (int32)BrushProperties->PrimaryBrushType;
+	int32 CurrentBrushType = (int32)BrushProperties->PrimaryBrushID;
 	int32 CurrentBrushTypeIndex = 0;
 
 	TArray<TSharedPtr<SComboPanel::FComboPanelItem>> BrushTypeItems;
diff --git a/Engine/Plugins/Experimental/MeshModelingToolsetExp/Source/MeshModelingToolsExp/MeshModelingToolsExp.Build.cs b/Engine/Plugins/Experimental/MeshModelingToolsetExp/Source/MeshModelingToolsExp/MeshModelingToolsExp.Build.cs
--- a/Engine/Plugins/Experimental/MeshModelingToolsetExp/Source/MeshModelingToolsExp/MeshModelingToolsExp.Build.cs
+++ b/Engine/Plugins/Experimental/MeshModelingToolsetExp/Source/MeshModelingToolsExp/MeshModelingToolsExp.Build.cs
@@ -46,6 +46,7 @@ public class MeshModelingToolsExp : ModuleRules
 			{
 				"CoreUObject",
 				"Engine",
+				"GeometryAlgorithms",
 				"RenderCore",
 				"InputCore",
 				"PhysicsCore",
diff --git a/Engine/Plugins/Experimental/MeshModelingToolsetExp/Source/MeshModelingToolsExp/Private/MeshVertexSculptTool.cpp b/Engine/Plugins/Experimental/MeshModelingToolsetExp/Source/MeshModelingToolsExp/Private/MeshVertexSculptTool.cpp
--- a/Engine/Plugins/Experimental/MeshModelingToolsetExp/Source/MeshModelingToolsExp/Private/MeshVertexSculptTool.cpp
+++ b/Engine/Plugins/Experimental/MeshModelingToolsetExp/Source/MeshModelingToolsExp/Private/MeshVertexSculptTool.cpp
@@ -4,6 +4,8 @@
 #include "Engine/World.h"
 #include "InteractiveToolManager.h"
 #include "InteractiveGizmoManager.h"
+#include "Intersection/ContainmentQueries3.h"
+#include "Intersection/IntrCylinderBox3.h"
 #include "ToolDataVisualizer.h"
 #include "Async/ParallelFor.h"
 #include "Async/Async.h"
@@ -121,18 +123,19 @@ void UMeshVertexSculptTool::Setup()
 
 	FDynamicMesh3* SculptMesh = GetSculptMesh();
 	FAxisAlignedBox3d Bounds = SculptMesh->GetBounds(true);
+	InitialBoundsMaxDim = Bounds.MaxDim();
 
 	// initialize dynamic octree
 	TFuture<void> InitializeOctree = Async(VertexSculptToolAsyncExecTarget, [SculptMesh, Bounds, this]()
 	{
 		if (SculptMesh->TriangleCount() > 100000)
 		{
-			Octree.RootDimension = Bounds.MaxDim() / 10.0;
+			Octree.RootDimension = InitialBoundsMaxDim / 10.0;
 			Octree.SetMaxTreeDepth(4);
 		}
 		else
 		{
-			Octree.RootDimension = Bounds.MaxDim() / 2.0;
+			Octree.RootDimension = InitialBoundsMaxDim / 2.0;
 			Octree.SetMaxTreeDepth(8);
 		}
 		Octree.Initialize(SculptMesh);
@@ -211,17 +214,17 @@ void UMeshVertexSculptTool::Setup()
 	AddToolPropertySource(UMeshSculptToolBase::BrushProperties);
 	UMeshSculptToolBase::BrushProperties->bShowPerBrushProps = false;
 	UMeshSculptToolBase::BrushProperties->bShowFalloff = false;
-	SculptProperties->RestoreProperties(this);
+	SculptProperties->RestoreProperties(this, GetPropertyCacheIdentifier());
 	AddToolPropertySource(SculptProperties);
 	CalculateBrushRadius();
 
 	AlphaProperties = NewObject<UVertexBrushAlphaProperties>(this);
-	AlphaProperties->RestoreProperties(this);
+	AlphaProperties->RestoreProperties(this, GetPropertyCacheIdentifier());
 	AlphaProperties->Tool = this;
 	AddToolPropertySource(AlphaProperties);
 
 	SymmetryProperties = NewObject<UMeshSymmetryProperties>(this);
-	SymmetryProperties->RestoreProperties(this);
+	SymmetryProperties->RestoreProperties(this, GetPropertyCacheIdentifier());
 	SymmetryProperties->bSymmetryCanBeEnabled = false;
 	AddToolPropertySource(SymmetryProperties);
 
@@ -230,6 +233,56 @@ void UMeshVertexSculptTool::Setup()
 		return GetBaseMeshNearest(VertexID, Position, MaxDist, PosOut, NormalOut);
 	};
 
+	RegisterBrushes();
+
+	// falloffs
+	RegisterStandardFalloffTypes();
+
+	AddToolPropertySource(UMeshSculptToolBase::GizmoProperties);
+	SetToolPropertySourceEnabled(UMeshSculptToolBase::GizmoProperties, false);
+	// Move the gizmo toward the center of the mesh, without changing the plane it represents
+	UMeshSculptToolBase::GizmoProperties->RecenterGizmoIfFar(GetSculptMeshComponent()->GetComponentTransform().TransformPosition(Bounds.Center()), Bounds.MaxDim());
+
+	AddToolPropertySource(UMeshSculptToolBase::ViewProperties);
+
+	// register watchers
+	SculptProperties->WatchProperty( SculptProperties->PrimaryBrushID,
+		[this](int32 NewType) { UpdateBrushType(NewType); });
+
+	SculptProperties->WatchProperty( SculptProperties->PrimaryFalloffType,
+		[this](EMeshSculptFalloffType NewType) { 
+			SetPrimaryFalloffType(NewType);
+			// Request to have the details panel rebuilt to ensure the new falloff property value is propagated to the details customization
+			OnDetailsPanelRequestRebuild.Broadcast();
+		});
+
+	SculptProperties->WatchProperty(AlphaProperties->Alpha,
+		[this](UTexture2D* NewAlpha) {
+			UpdateBrushAlpha(AlphaProperties->Alpha);
+			// Request to have the details panel rebuilt to ensure the new alpha property value is propagated to the details customization
+			OnDetailsPanelRequestRebuild.Broadcast();
+		});
+
+	// must call before updating brush type so that we register all brush properties?
+	UMeshSculptToolBase::OnCompleteSetup();
+
+	UpdateBrushType(SculptProperties->PrimaryBrushID);
+	SetPrimaryFalloffType(SculptProperties->PrimaryFalloffType);
+	UpdateBrushAlpha(AlphaProperties->Alpha);
+	SetActiveSecondaryBrushType(0);
+
+	StampRandomStream = FRandomStream(31337);
+
+	// update symmetry state based on validity, and then update internal apply-symmetry state
+	SymmetryProperties->bSymmetryCanBeEnabled = bMeshSymmetryIsValid;
+	bApplySymmetry = bMeshSymmetryIsValid && SymmetryProperties->bEnableSymmetry;
+
+	SymmetryProperties->WatchProperty(SymmetryProperties->bEnableSymmetry,
+		[this](bool bNewValue) { bApplySymmetry = bMeshSymmetryIsValid && bNewValue; });
+}
+
+void UMeshVertexSculptTool::RegisterBrushes()
+{
 	RegisterBrushType((int32)EMeshVertexSculptBrushType::Smooth, LOCTEXT("SmoothBrush", "Smooth"),
 		MakeUnique<TBasicMeshSculptBrushOpFactory<FSmoothBrushOp>>(),
 		NewObject<USmoothBrushOpProps>(this));
@@ -247,7 +300,7 @@ void UMeshVertexSculptTool::Setup()
 		NewObject<UStandardSculptBrushOpProps>(this));
 
 	RegisterBrushType((int32)EMeshVertexSculptBrushType::SculptView, LOCTEXT("SculptView", "SculptV"),
-		MakeUnique<FLambdaMeshSculptBrushOpFactory>( [this]() { return MakeUnique<FViewAlignedSculptBrushOp>(BaseMeshQueryFunc); } ),
+		MakeUnique<FLambdaMeshSculptBrushOpFactory>([this]() { return MakeUnique<FViewAlignedSculptBrushOp>(BaseMeshQueryFunc); }),
 		NewObject<UViewAlignedSculptBrushOpProps>(this));
 
 	RegisterBrushType((int32)EMeshVertexSculptBrushType::SculptMax, LOCTEXT("SculptMax", "SculptMx"),
@@ -278,7 +331,7 @@ void UMeshVertexSculptTool::Setup()
 		MakeUnique<TBasicMeshSculptBrushOpFactory<FPlaneBrushOp>>(),
 		NewObject<UFixedPlaneBrushOpProps>(this));
 
-	RegisterBrushType((int32)EMeshVertexSculptBrushType::ScaleKelvin , LOCTEXT("ScaleKelvin", "Scale"),
+	RegisterBrushType((int32)EMeshVertexSculptBrushType::ScaleKelvin, LOCTEXT("ScaleKelvin", "Scale"),
 		MakeUnique<TBasicMeshSculptBrushOpFactory<FScaleKelvinletBrushOp>>(),
 		NewObject<UScaleKelvinletBrushOpProps>(this));
 
@@ -295,54 +348,15 @@ void UMeshVertexSculptTool::Setup()
 		NewObject<UTwistKelvinletBrushOpProps>(this));
 
 	// secondary brushes
-	RegisterSecondaryBrushType((int32)EMeshVertexSculptBrushType::Smooth, LOCTEXT("Smooth", "Smooth"),
+	// We activate ID 0 as our default secondary brush, so use that as the registration ID
+	RegisterSecondaryBrushType(0, LOCTEXT("Smooth", "Smooth"),
 		MakeUnique<TBasicMeshSculptBrushOpFactory<FSmoothBrushOp>>(),
 		NewObject<USecondarySmoothBrushOpProps>(this));
+}
 
-	// falloffs
-	RegisterStandardFalloffTypes();
-
-	AddToolPropertySource(UMeshSculptToolBase::GizmoProperties);
-	SetToolPropertySourceEnabled(UMeshSculptToolBase::GizmoProperties, false);
-	// Move the gizmo toward the center of the mesh, without changing the plane it represents
-	UMeshSculptToolBase::GizmoProperties->RecenterGizmoIfFar(GetSculptMeshComponent()->GetComponentTransform().TransformPosition(Bounds.Center()), Bounds.MaxDim());
-
-	AddToolPropertySource(UMeshSculptToolBase::ViewProperties);
-
-	// register watchers
-	SculptProperties->WatchProperty( SculptProperties->PrimaryBrushType,
-		[this](EMeshVertexSculptBrushType NewType) { UpdateBrushType(NewType); });
-
-	SculptProperties->WatchProperty( SculptProperties->PrimaryFalloffType,
-		[this](EMeshSculptFalloffType NewType) { 
-			SetPrimaryFalloffType(NewType);
-			// Request to have the details panel rebuilt to ensure the new falloff property value is propagated to the details customization
-			OnDetailsPanelRequestRebuild.Broadcast();
-		});
-
-	SculptProperties->WatchProperty(AlphaProperties->Alpha,
-		[this](UTexture2D* NewAlpha) {
-			UpdateBrushAlpha(AlphaProperties->Alpha);
-			// Request to have the details panel rebuilt to ensure the new alpha property value is propagated to the details customization
-			OnDetailsPanelRequestRebuild.Broadcast();
-		});
-
-	// must call before updating brush type so that we register all brush properties?
-	UMeshSculptToolBase::OnCompleteSetup();
-
-	UpdateBrushType(SculptProperties->PrimaryBrushType);
-	SetPrimaryFalloffType(SculptProperties->PrimaryFalloffType);
-	UpdateBrushAlpha(AlphaProperties->Alpha);
-	SetActiveSecondaryBrushType((int32)EMeshVertexSculptBrushType::Smooth);
-
-	StampRandomStream = FRandomStream(31337);
-
-	// update symmetry state based on validity, and then update internal apply-symmetry state
-	SymmetryProperties->bSymmetryCanBeEnabled = bMeshSymmetryIsValid;
-	bApplySymmetry = bMeshSymmetryIsValid && SymmetryProperties->bEnableSymmetry;
-
-	SymmetryProperties->WatchProperty(SymmetryProperties->bEnableSymmetry,
-		[this](bool bNewValue) { bApplySymmetry = bMeshSymmetryIsValid && bNewValue; });
+FString UMeshVertexSculptTool::GetPropertyCacheIdentifier() const
+{
+	return TEXT("UMeshVertexSculptTool");
 }
 
 void UMeshVertexSculptTool::Shutdown(EToolShutdownType ShutdownType)
@@ -352,9 +366,9 @@ void UMeshVertexSculptTool::Shutdown(EToolShutdownType ShutdownType)
 		DynamicMeshComponent->OnMeshChanged.Remove(OnDynamicMeshComponentChangedHandle);
 	}
 
-	SculptProperties->SaveProperties(this);
-	AlphaProperties->SaveProperties(this);
-	SymmetryProperties->SaveProperties(this);
+	SculptProperties->SaveProperties(this, GetPropertyCacheIdentifier());
+	AlphaProperties->SaveProperties(this, GetPropertyCacheIdentifier());
+	SymmetryProperties->SaveProperties(this, GetPropertyCacheIdentifier());
 
 	if (PreviewMeshActor != nullptr)
 	{
@@ -416,11 +430,10 @@ void UMeshVertexSculptTool::InitializeIndicator()
 
 void UMeshVertexSculptTool::SetActiveBrushType(int32 Identifier)
 {
-	EMeshVertexSculptBrushType NewBrushType =  static_cast<EMeshVertexSculptBrushType>(Identifier);
-	if (SculptProperties->PrimaryBrushType != NewBrushType)
+	if (SculptProperties->PrimaryBrushID != Identifier)
 	{
-		SculptProperties->PrimaryBrushType = NewBrushType;
-		UpdateBrushType(SculptProperties->PrimaryBrushType);
+		SculptProperties->PrimaryBrushID = Identifier;
+		UpdateBrushType(SculptProperties->PrimaryBrushID);
 		SculptProperties->SilentUpdateWatched();
 	}
 
@@ -454,14 +467,16 @@ void UMeshVertexSculptTool::OnBeginStroke(const FRay& WorldRay)
 
 	UpdateBrushPosition(WorldRay);
 
-	if (SculptProperties->PrimaryBrushType == EMeshVertexSculptBrushType::Plane ||
-		SculptProperties->PrimaryBrushType == EMeshVertexSculptBrushType::PlaneViewAligned)
+	TUniquePtr<FMeshSculptBrushOp>& UseBrushOp = GetActiveBrushOp();
+	FMeshSculptBrushOp::EReferencePlaneType ReferencePlaneType = UseBrushOp->GetReferencePlaneType();
+	if (ReferencePlaneType == FMeshSculptBrushOp::EReferencePlaneType::InitialROI ||
+		ReferencePlaneType == FMeshSculptBrushOp::EReferencePlaneType::InitialROI_ViewAligned)
 	{
-		UpdateROI(GetBrushFrameLocal().Origin);
+		UpdateROI(GetBrushFrameLocal());
 		UpdateStrokeReferencePlaneForROI(GetBrushFrameLocal(), TriangleROIArray,
-			(SculptProperties->PrimaryBrushType == EMeshVertexSculptBrushType::PlaneViewAligned));
+			ReferencePlaneType == FMeshSculptBrushOp::EReferencePlaneType::InitialROI_ViewAligned);
 	}
-	else if (SculptProperties->PrimaryBrushType == EMeshVertexSculptBrushType::FixedPlane)
+	else if (ReferencePlaneType == FMeshSculptBrushOp::EReferencePlaneType::WorkPlane)
 	{
 		UpdateStrokeReferencePlaneFromWorkPlane();
 	}
@@ -476,6 +491,8 @@ void UMeshVertexSculptTool::OnBeginStroke(const FRay& WorldRay)
 	LastStamp.Power = GetActivePressure() * GetCurrentBrushStrength();
 	LastStamp.TimeStamp = FDateTime::Now();
 
+	PreviousRayDirection = FVector3d::ZeroVector;
+
 	// If applying symmetry, make sure the stamp is on the "positive" side. 
 	if (bApplySymmetry)
 	{
@@ -491,7 +508,6 @@ void UMeshVertexSculptTool::OnBeginStroke(const FRay& WorldRay)
 	//SculptOptions.bPreserveUVFlow = false; // SculptProperties->bPreserveUVFlow;
 	SculptOptions.ConstantReferencePlane = GetCurrentStrokeReferencePlane();
 
-	TUniquePtr<FMeshSculptBrushOp>& UseBrushOp = GetActiveBrushOp();
 	UseBrushOp->ConfigureOptions(SculptOptions);
 	UseBrushOp->BeginStroke(GetSculptMesh(), LastStamp, VertexROI);
 
@@ -521,23 +537,100 @@ void UMeshVertexSculptTool::OnCancelStroke()
 	ActiveVertexChange = nullptr;
 }
 
-
-
-void UMeshVertexSculptTool::UpdateROI(const FVector3d& BrushPos)
+// The first part of UpdateROI, which updates TriangleROIArray to be triangles
+//  in our region of interest, and VertexROI to be vertices in our region of
+//  interest.
+void UMeshVertexSculptTool::UpdateRangeQueryTriBuffer(const FFrame3d& LocalFrame)
 {
-	TRACE_CPUPROFILER_EVENT_SCOPE(VtxSculptTool_UpdateROI);
+	FDynamicMesh3* Mesh = GetSculptMesh();
+	FVector3d BrushPos = LocalFrame.Origin;
 
+	// By default, our brush is a sphere, and we affect vertices inside it
 	float RadiusSqr = GetCurrentBrushRadius() * GetCurrentBrushRadius();
-	FAxisAlignedBox3d BrushBox(
-		BrushPos - GetCurrentBrushRadius() * FVector3d::One(),
-		BrushPos + GetCurrentBrushRadius() * FVector3d::One());
+	// This function gets called when first gathering the triangles that might intersect our brush
+	//  from the octree's cells
+	TFunction<void()> GatherOverlappingCells = [this, &BrushPos]()
+	{
+		FAxisAlignedBox3d BrushBox(
+			BrushPos - GetCurrentBrushRadius() * FVector3d::One(),
+			BrushPos + GetCurrentBrushRadius() * FVector3d::One());
+		Octree.ParallelRangeQuery(BrushBox, RangeQueryTriBuffer);
+	};
+	// This is used to filter the gathered verts for ones that are actually in the brush.
+	TFunction<bool(int32 Vid)> IsVertInBrush = [Mesh, &BrushPos, RadiusSqr](int32 Vid)
+	{
+		return DistanceSquared(BrushPos, Mesh->GetVertexRef(Vid)) < RadiusSqr;
+	};
 
-	// do a parallel range quer
+	// Some brush types want their brush to be a cylinder, so we need to change how we evaluate
+	//  cells/vertices that are within reach
+	TUniquePtr<FMeshSculptBrushOp>& CurrentBrush = GetActiveBrushOp();
+	if (CurrentBrush &&
+		(CurrentBrush->GetBrushRegionType() == FMeshSculptBrushOp::EBrushRegionType::InfiniteCylinder
+		|| CurrentBrush->GetBrushRegionType() == FMeshSculptBrushOp::EBrushRegionType::CylinderOnSphere))
+	{
+		double CylinderRadius = GetCurrentBrushRadius();
+		double CylinderHeight = TNumericLimits<double>::Max();
+
+		FVector3d CylinderCenter, CylinderAxis;
+		if (CurrentBrush->GetBrushRegionType() == FMeshSculptBrushOp::EBrushRegionType::InfiniteCylinder)
+		{
+			CylinderCenter = BrushPos;
+			CylinderAxis = LocalFrame.Z();
+
+			// Since cylinder is infinite, just have to check distance from line for the actual vert containment function
+			FLine3d CylinderLine(CylinderCenter, CylinderAxis);
+			IsVertInBrush = [Mesh, CylinderLine, RadiusSqr](int32 Vid)
+			{
+				return CylinderLine.DistanceSquared(Mesh->GetVertexRef(Vid)) < RadiusSqr;
+			};
+		}
+		else // if cylinder on sphere
+		{
+			FVector3d SphereCenter = GizmoProperties ? CurTargetTransform.InverseTransformPosition(GizmoProperties->Position) 
+				: FVector3d::ZeroVector;
+			CylinderAxis = BrushPos - SphereCenter;
+			if (!CylinderAxis.Normalize())
+			{
+				CylinderAxis = FVector3d::UnitZ();
+			}
+			// We want the bottom of our cylinder to be at the sphere center, and the top to go infinitely up, 
+			//  but we need a non-infinite position for the center. So lets pick our height to be based on mesh bounds,
+			//  with some arbitrary minimum instead.
+			CylinderHeight = FMath::Max(InitialBoundsMaxDim, 1000);
+			CylinderCenter = SphereCenter + CylinderAxis * (CylinderHeight / 2);
+
+			IsVertInBrush = [Mesh, CylinderCenter, CylinderAxis, CylinderHeight, CylinderRadius](int32 Vid)
+			{
+				return UE::Geometry::DoesCylinderContainPoint(CylinderCenter, CylinderAxis, 
+					CylinderRadius, CylinderHeight, Mesh->GetVertexRef(Vid));
+			};
+		}
+
+		auto DoesCellIntersectBrush = [CylinderCenter, CylinderAxis, CylinderRadius, CylinderHeight](const FAxisAlignedBox3d& CellBounds)
+		{
+			return UE::Geometry::DoesCylinderIntersectBox(CellBounds,
+				CylinderCenter, CylinderAxis, CylinderRadius, CylinderHeight);
+		};
+		
+		FAxisAlignedBox3d ConservativeCylinderBounds;
+		ConservativeCylinderBounds.Contain(CylinderCenter + CylinderAxis * (CylinderHeight / 2));
+		ConservativeCylinderBounds.Contain(CylinderCenter - CylinderAxis * (CylinderHeight / 2));
+		ConservativeCylinderBounds.Expand(CylinderRadius);
+
+		GatherOverlappingCells = [this, ConservativeCylinderBounds, BoundsOverlapFn = MoveTemp(DoesCellIntersectBrush)]()
+		{
+			Octree.ParallelRangeQuery(ConservativeCylinderBounds, BoundsOverlapFn, RangeQueryTriBuffer);
+		};
+	}
+
+	// Do a parallel range query to find those triangles that may intersect with
+	//  our brush bounds. This grabs all triangles of intersecting cells, so we
+	//  will need to do additional filtering afterward.
 	RangeQueryTriBuffer.Reset();
-	FDynamicMesh3* Mesh = GetSculptMesh();
 	{
 		TRACE_CPUPROFILER_EVENT_SCOPE(VtxSculptTool_UpdateROI_RangeQuery);
-		Octree.ParallelRangeQuery(BrushBox, RangeQueryTriBuffer);
+		GatherOverlappingCells();
 	}
 
 	int32 ActiveComponentID = -1;
@@ -583,9 +676,9 @@ void UMeshVertexSculptTool::UpdateROI(const FVector3d& BrushPos)
 			}
 
 			const FIndex3i& TriV = Mesh->GetTriangleRef(tid);
-			TriangleROIInBuf[k].A = (DistanceSquared(BrushPos, Mesh->GetVertexRef(TriV.A)) < RadiusSqr) ? 1 : 0;
-			TriangleROIInBuf[k].B = (DistanceSquared(BrushPos, Mesh->GetVertexRef(TriV.B)) < RadiusSqr) ? 1 : 0;
-			TriangleROIInBuf[k].C = (DistanceSquared(BrushPos, Mesh->GetVertexRef(TriV.C)) < RadiusSqr) ? 1 : 0;
+			TriangleROIInBuf[k].A = IsVertInBrush(TriV.A) ? 1 : 0;
+			TriangleROIInBuf[k].B = IsVertInBrush(TriV.B) ? 1 : 0;
+			TriangleROIInBuf[k].C = IsVertInBrush(TriV.C) ? 1 : 0;
 			if (TriangleROIInBuf[k].A + TriangleROIInBuf[k].B + TriangleROIInBuf[k].C == 0)
 			{
 				RangeQueryTriBuffer[k] = -1;
@@ -620,7 +713,47 @@ void UMeshVertexSculptTool::UpdateROI(const FVector3d& BrushPos)
 			}
 		}
 
-		VertexROIBuilder.SwapValuesWith(VertexROI);
+		// See if we need to filter our vertices based on connectivity to hit location (used to avoid affecting
+		//  hidden regions of a mesh that might be in the volume of the brush)
+		if (RequireConnectivityToHitPointInStamp()
+			&& ensure(LastBrushTriangleID != IndexConstants::InvalidID))
+		{
+			FIndex3i HitTriVids = Mesh->GetTriangle(LastBrushTriangleID);
+			TArray<int32> SeedVids;
+			for (int i = 0; i < 3; ++i)
+			{
+				if (VertexROIBuilder.Contains(HitTriVids[i]))
+				{
+					SeedVids.Add(HitTriVids[i]);
+				}
+			}
+			
+			TSet<int32> ConnectedROIVids;
+				FMeshConnectedComponents Components(Mesh);
+				Components.GrowToConnectedVertices(*Mesh, SeedVids, ConnectedROIVids, nullptr,
+					[this](int32 Vid, int32 Tid) { return VertexROIBuilder.Contains(Vid); });
+
+			// We'll need to update TriangleROIBuilder based on the vertices too
+			TArray<int32> TidsToFilter = TriangleROIBuilder.TakeValues();
+			TriangleROIBuilder.Initialize(Mesh->MaxTriangleID());
+			for (int32 Tid : TidsToFilter)
+			{
+				FIndex3i TriVids = Mesh->GetTriangle(Tid);
+				for (int i = 0; i < 3; ++i)
+				{
+					if (ConnectedROIVids.Contains(TriVids[i]))
+					{
+						TriangleROIBuilder.Add(Tid);
+						break; // continue to next Tid
+					}
+				}
+			}
+			VertexROI = ConnectedROIVids.Array();
+		}
+		else
+		{
+			VertexROIBuilder.SwapValuesWith(VertexROI);
+		}
 
 		if (bApplySymmetry)
 		{
@@ -647,10 +780,10 @@ void UMeshVertexSculptTool::UpdateROI(const FVector3d& BrushPos)
 	// In this path we combine everything into one loop. Does fewer distance checks
 	// but nothing can be done in parallel (would change if ROIBuilders had atomic-try-add)
 
-	check(bApplySymmetry == false);		// have not implemented this path yet...
-
 	// TODO would need to support these, this branch is likely dead though
 	ensure(SculptProperties->BrushFilter == EMeshVertexSculptBrushFilterType::None);
+	ensure(!bApplySymmetry);
+	ensure(!RequireConnectivityToHitPointInStamp());
 
 	// collect set of vertices and triangles inside brush sphere, from range query result
 	{
@@ -667,7 +800,7 @@ void UMeshVertexSculptTool::UpdateROI(const FVector3d& BrushPos)
 				{
 					InsideCount++;
 				} 
-				else if (BrushPos.DistanceSquared(Mesh->GetVertexRef(TriV[j])) < RadiusSqr)
+				else if (IsVertInBrush(TriV[j]))
 				{
 					VertexROIBuilder.Add(TriV[j]);
 					InsideCount++;
@@ -682,33 +815,55 @@ void UMeshVertexSculptTool::UpdateROI(const FVector3d& BrushPos)
 		TriangleROIBuilder.SwapValuesWith(TriangleROIArray);
 	}
 #endif
+}
 
+// Second part of UpdateROI, which fills out ROIPrevPositionBuffer, prepares ROIPositionBuffer,
+//  and prepares the symmetry buffers if relevant.
+void UMeshVertexSculptTool::PrepROIVertPositionBuffers()
+{
+	FDynamicMesh3* Mesh = GetSculptMesh();
+	// set up and populate position buffers for Vertex ROI
+	TRACE_CPUPROFILER_EVENT_SCOPE(DynamicMeshSculptTool_UpdateROI_4ROI);
+	int32 ROISize = VertexROI.Num();
+	ROIPositionBuffer.SetNum(ROISize, EAllowShrinking::No);
+	ROIPrevPositionBuffer.SetNum(ROISize, EAllowShrinking::No);
+	ParallelFor(ROISize, [&](int i)
+	{
+		ROIPrevPositionBuffer[i] = Mesh->GetVertexRef(VertexROI[i]);
+	});
+	// do the same for the Symmetric Vertex ROI
+	if (bApplySymmetry)
 	{
-		// set up and populate position buffers for Vertex ROI
-		TRACE_CPUPROFILER_EVENT_SCOPE(DynamicMeshSculptTool_UpdateROI_4ROI);
-		int32 ROISize = VertexROI.Num();
-		ROIPositionBuffer.SetNum(ROISize, EAllowShrinking::No);
-		ROIPrevPositionBuffer.SetNum(ROISize, EAllowShrinking::No);
+		SymmetricROIPositionBuffer.SetNum(ROISize, EAllowShrinking::No);
+		SymmetricROIPrevPositionBuffer.SetNum(ROISize, EAllowShrinking::No);
 		ParallelFor(ROISize, [&](int i)
 		{
-			ROIPrevPositionBuffer[i] = Mesh->GetVertexRef(VertexROI[i]);
-		});
-		// do the same for the Symmetric Vertex ROI
-		if (bApplySymmetry)
-		{
-			SymmetricROIPositionBuffer.SetNum(ROISize, EAllowShrinking::No);
-			SymmetricROIPrevPositionBuffer.SetNum(ROISize, EAllowShrinking::No);
-			ParallelFor(ROISize, [&](int i)
+			if ( Mesh->IsVertex(SymmetricVertexROI[i]) )
 			{
-				if ( Mesh->IsVertex(SymmetricVertexROI[i]) )
-				{
-					SymmetricROIPrevPositionBuffer[i] = Mesh->GetVertexRef(SymmetricVertexROI[i]);
-				}
-			});
-		}
+				SymmetricROIPrevPositionBuffer[i] = Mesh->GetVertexRef(SymmetricVertexROI[i]);
+			}
+		});
 	}
 }
 
+void UMeshVertexSculptTool::UpdateROI(const FVector3d& BrushPos)
+{
+	UpdateROI(FFrame3d(BrushPos));
+}
+
+void UMeshVertexSculptTool::UpdateROI(const FFrame3d& LocalFrame)
+{
+	TRACE_CPUPROFILER_EVENT_SCOPE(VtxSculptTool_UpdateROI);
+
+	UpdateRangeQueryTriBuffer(LocalFrame);
+	PrepROIVertPositionBuffers();
+}
+
+/*
+ * Updates CurrentStamp, LastStamp, bMouseMoved, and LastMovedStamp (if bMouseMoved is true)
+ * 
+ * @return false if this ray can be ignored because it did not move and brush ignores zero movement.
+ */
 bool UMeshVertexSculptTool::UpdateStampPosition(const FRay& WorldRay)
 {
 	TRACE_CPUPROFILER_EVENT_SCOPE(VtxSculptTool_UpdateStampPosition);
@@ -731,10 +886,8 @@ bool UMeshVertexSculptTool::UpdateStampPosition(const FRay& WorldRay)
 		break;
 	}
 
-	if (UseBrushOp->GetAlignStampToView())
-	{
-		AlignBrushToView();
-	}
+	// Adjust stamp alignment if needed
+	RealignBrush(UseBrushOp->GetStampAlignmentType());
 
 	CurrentStamp = LastStamp;
 	//CurrentStamp.DeltaTime = FMathd::Min((FDateTime::Now() - LastStamp.TimeStamp).GetTotalSeconds(), 1.0);
@@ -766,15 +919,46 @@ bool UMeshVertexSculptTool::UpdateStampPosition(const FRay& WorldRay)
 	CurrentStamp.PrevLocalFrame = LastStamp.LocalFrame;
 	CurrentStamp.PrevWorldFrame = LastStamp.WorldFrame;
 
-	FVector3d MoveDelta = CurrentStamp.LocalFrame.Origin - CurrentStamp.PrevLocalFrame.Origin;
-	if (UseBrushOp->IgnoreZeroMovements() && MoveDelta.SquaredLength() < FMathd::ZeroTolerance)
+	bMouseMoved = (PreviousRayDirection - WorldRay.Direction).SquaredLength() > FMathd::ZeroTolerance;
+	if (bMouseMoved)
 	{
-		return false;
+		LastMovedStamp = CurrentStamp;
+		PreviousRayDirection = WorldRay.Direction;
 	}
-
-	return true;
+	return bMouseMoved || !UseBrushOp->IgnoreZeroMovements();
 }
 
+// Adjusts brush alignment (assumes that currently brush is aligned to hit normal)
+void UMeshVertexSculptTool::RealignBrush(FMeshSculptBrushOp::EStampAlignmentType AlignmentType)
+{
+	switch (AlignmentType)
+	{
+	case FMeshSculptBrushOp::EStampAlignmentType::HitNormal:
+		// Assume this is already aligned
+		break;
+	case FMeshSculptBrushOp::EStampAlignmentType::Camera:
+		AlignBrushToView();
+		break;
+	case FMeshSculptBrushOp::EStampAlignmentType::ReferencePlane:
+		// Note for this and reference sphere: GetCurrentStrokeReferencePlane is not (necessarily)
+		//  what we want because we may not have done UpdateStrokeReferencePlaneFromWorkPlane.
+		UpdateBrushFrameWorld(GetBrushFrameWorld().Origin, 
+			GizmoProperties ? GizmoProperties->Rotation.GetAxisZ() : FVector3d::UnitZ());
+		break;
+	case FMeshSculptBrushOp::EStampAlignmentType::ReferenceSphere:
+	{
+		FVector3d BrushLocation = GetBrushFrameWorld().Origin;
+		FVector3d SphereCenter = GizmoProperties ? GizmoProperties->Position : FVector3d::ZeroVector;
+		FVector3d NormalToUse = BrushLocation - SphereCenter;
+		if (!NormalToUse.Normalize())
+		{
+			NormalToUse = FVector3d::UnitZ();
+		}
+		UpdateBrushFrameWorld(BrushLocation, NormalToUse);
+		break;
+	}
+	}
+}
 
 TFuture<void> UMeshVertexSculptTool::ApplyStamp()
 {
@@ -782,16 +966,28 @@ TFuture<void> UMeshVertexSculptTool::ApplyStamp()
 
 	TUniquePtr<FMeshSculptBrushOp>& UseBrushOp = GetActiveBrushOp();
 
+	FSculptBrushStamp* StampToUse = &CurrentStamp;
+	// If we haven't moved our stamp, we might want to consider it to be at the same location
+	//  (depending on the brush we're using).
+	// TODO: It would be nice to have the brush visualization reflect this (currently it is still
+	//  centered on the ray hit), but requires more changes to the base tool. Also this sometimes
+	//  doesn't work super well with RequireConnectivityToHitPointInStamp because the hit location
+	//  may be outside the stamp, but it's not terrible.
+	if (!bMouseMoved && UseBrushOp->UseLastStampFrameOnZeroMovement())
+	{
+		StampToUse = &LastMovedStamp;
+	}
+
 	// compute region plane if necessary. This may currently be expensive?
 	if (UseBrushOp->WantsStampRegionPlane())
 	{
-		CurrentStamp.RegionPlane = ComputeStampRegionPlane(CurrentStamp.LocalFrame, TriangleROIArray, true, false, false);
+		StampToUse->RegionPlane = ComputeStampRegionPlane(StampToUse->LocalFrame, TriangleROIArray, true, false, false);
 	}
 
 	// set up alpha function if we have one
 	if (bHaveBrushAlpha)
 	{
-		CurrentStamp.StampAlphaFunc = [this](const FSculptBrushStamp& Stamp, const FVector3d& Position)
+		StampToUse->StampAlphaFunc = [this](const FSculptBrushStamp& Stamp, const FVector3d& Position)
 		{
 			return this->SampleBrushAlpha(Stamp, Position);
 		};
@@ -801,11 +997,11 @@ TFuture<void> UMeshVertexSculptTool::ApplyStamp()
 	FDynamicMesh3* Mesh = GetSculptMesh();
 	{
 		TRACE_CPUPROFILER_EVENT_SCOPE(VtxSculptTool_ApplyStamp_Apply);
-		UseBrushOp->ApplyStamp(Mesh, CurrentStamp, VertexROI, ROIPositionBuffer);
+		UseBrushOp->ApplyStamp(Mesh, *StampToUse, VertexROI, ROIPositionBuffer);
 	}
 
 	// can discard alpha now
-	CurrentStamp.StampAlphaFunc = nullptr;
+	StampToUse->StampAlphaFunc = nullptr;
 
 	// if we are applying symmetry, we need to update the on-plane positions as they
 	// will not be in the SymmetricVertexROI
@@ -893,7 +1089,7 @@ TFuture<void> UMeshVertexSculptTool::ApplyStamp()
 		Mesh->UpdateChangeStamps(true, false);
 	}
 
-	LastStamp = CurrentStamp;
+	LastStamp = *StampToUse;
 	LastStamp.TimeStamp = FDateTime::Now();
 
 	// let caller wait for this to finish
@@ -966,9 +1162,9 @@ bool UMeshVertexSculptTool::UpdateBrushPosition(const FRay& WorldRay)
 		break;
 	}
 
-	if (bHit && UseBrushOp->GetAlignStampToView())
+	if (bHit)
 	{
-		AlignBrushToView();
+		RealignBrush(UseBrushOp->GetStampAlignmentType());
 	}
 
 	return bHit;
@@ -992,7 +1188,7 @@ bool UMeshVertexSculptTool::OnUpdateHover(const FInputDeviceRay& DevicePos)
 	// 4.26 HOTFIX: update LastWorldRay position so that we have it for updating WorkPlane position
 	UMeshSurfacePointTool::LastWorldRay = DevicePos.WorldRay;
 
-	PendingStampType = SculptProperties->PrimaryBrushType;
+	PendingStampBrushID = SculptProperties->PrimaryBrushID;
 	if(ensure(InStroke() == false))
 	{
 		UpdateBrushPosition(DevicePos.WorldRay);
@@ -1112,7 +1308,7 @@ void UMeshVertexSculptTool::OnTick(float DeltaTime)
 		WaitForPendingStampUpdate();
 
 		// update sculpt ROI
-		UpdateROI(CurrentStamp.LocalFrame.Origin);
+		UpdateROI(CurrentStamp.LocalFrame);
 
 		// Append updated ROI to modified region (async). For some reason this is very expensive,
 		// maybe because of TSet? but we have a lot of time to do it.
@@ -1209,7 +1405,7 @@ void UMeshVertexSculptTool::WaitForPendingStampUpdate()
 	if (bStampUpdatePending)
 	{
 		StampUpdateOctreeFuture.Wait();
-		bStampUpdatePending = true;
+		bStampUpdatePending = false;
 	}
 }
 
@@ -1459,31 +1655,39 @@ void UMeshVertexSculptTool::OnDynamicMeshComponentChanged(UDynamicMeshComponent*
 
 void UMeshVertexSculptTool::UpdateBrushType(EMeshVertexSculptBrushType BrushType)
 {
+	UpdateBrushType((int32)BrushType);
+}
+
+void UMeshVertexSculptTool::UpdateBrushType(int32 BrushID)
+{	
 	static const FText BaseMessage = LOCTEXT("OnStartSculptTool", "Hold Shift to Smooth, Ctrl to Invert (where applicable). [/] and S/D change Size (+Shift to small-step), W/E changes Strength.");
 	FTextBuilder Builder;
 	Builder.AppendLine(BaseMessage);
 
-	SetActivePrimaryBrushType((int32)BrushType);
+	SetActivePrimaryBrushType(BrushID);
+
+	if (ensure(SculptProperties))
+	{
+		SculptProperties->bCanFreezeTarget = 
+			BrushID == (int32)EMeshVertexSculptBrushType::Offset
+			|| BrushID == (int32)EMeshVertexSculptBrushType::SculptMax || BrushID == (int32)EMeshVertexSculptBrushType::SculptView
+			|| BrushID == (int32)EMeshVertexSculptBrushType::Pinch;
+	}
 
 	SetToolPropertySourceEnabled(GizmoProperties, false);
-	if (BrushType == EMeshVertexSculptBrushType::FixedPlane)
+	if (PrimaryBrushOp && 
+		(PrimaryBrushOp->GetReferencePlaneType() == FMeshSculptBrushOp::EReferencePlaneType::WorkPlane
+		|| PrimaryBrushOp->GetStampAlignmentType() == FMeshSculptBrushOp::EStampAlignmentType::ReferencePlane
+		|| PrimaryBrushOp->GetStampAlignmentType() == FMeshSculptBrushOp::EStampAlignmentType::ReferenceSphere))
 	{
 		Builder.AppendLine(LOCTEXT("FixedPlaneTip", "Use T to reposition Work Plane at cursor, Shift+T to align to Normal, Ctrl+Shift+T to align to View"));
 		SetToolPropertySourceEnabled(GizmoProperties, true);
 	}
 
-
-	bool bEnableAlpha = false;
-	switch (BrushType)
-	{
-	case EMeshVertexSculptBrushType::Offset:
-	case EMeshVertexSculptBrushType::SculptView:
-	case EMeshVertexSculptBrushType::SculptMax:
-		bEnableAlpha = true; break;
-	}
+	TUniquePtr<FMeshSculptBrushOp>& UseBrushOp = GetActiveBrushOp();
+	bool bEnableAlpha = UseBrushOp && UseBrushOp->UsesAlpha();
 	SetToolPropertySourceEnabled(AlphaProperties, bEnableAlpha);
 
-
 	GetToolManager()->DisplayMessage(Builder.ToText(), EToolMessageLevel::UserNotification);
 }
 
diff --git a/Engine/Plugins/Experimental/MeshModelingToolsetExp/Source/MeshModelingToolsExp/Public/MeshVertexSculptTool.h b/Engine/Plugins/Experimental/MeshModelingToolsetExp/Source/MeshModelingToolsExp/Public/MeshVertexSculptTool.h
--- a/Engine/Plugins/Experimental/MeshModelingToolsetExp/Source/MeshModelingToolsExp/Public/MeshVertexSculptTool.h
+++ b/Engine/Plugins/Experimental/MeshModelingToolsetExp/Source/MeshModelingToolsExp/Public/MeshVertexSculptTool.h
@@ -123,9 +123,14 @@ class MESHMODELINGTOOLSEXP_API UVertexBrushSculptProperties : public UInteractiv
 	GENERATED_BODY()
 
 public:
+	
+	UE_DEPRECATED(5.6, "Vertex sculpt now uses an integer for a brush ID to allow custom brush registration")
+	UPROPERTY()
+	EMeshVertexSculptBrushType PrimaryBrushType = EMeshVertexSculptBrushType::Offset;
+
 	/** Primary Brush Mode */
 	UPROPERTY(EditAnywhere, Category = Sculpting, meta = (DisplayName = "Brush"))
-	EMeshVertexSculptBrushType PrimaryBrushType = EMeshVertexSculptBrushType::Offset;
+	int32 PrimaryBrushID = 0;
 
 	/** Primary Brush Falloff Type, multiplied by Alpha Mask where applicable */
 	UPROPERTY(EditAnywhere, Category = Sculpting, meta = (DisplayName = "Falloff"))
@@ -136,9 +141,12 @@ public:
 	EMeshVertexSculptBrushFilterType BrushFilter = EMeshVertexSculptBrushFilterType::None;
 
 	/** When Freeze Target is toggled on, the Brush Target Surface will be Frozen in its current state, until toggled off. Brush strokes will be applied relative to the Target Surface, for applicable Brushes */
-	UPROPERTY(EditAnywhere, Category = Sculpting, meta = (EditCondition = "PrimaryBrushType == EMeshVertexSculptBrushType::Offset || PrimaryBrushType == EMeshVertexSculptBrushType::SculptMax || PrimaryBrushType == EMeshVertexSculptBrushType::SculptView || PrimaryBrushType == EMeshVertexSculptBrushType::Pinch || PrimaryBrushType == EMeshVertexSculptBrushType::Resample" ))
+	UPROPERTY(EditAnywhere, Category = Sculpting, meta = (EditCondition = "bCanFreezeTarget", HideEditConditionToggle))
 	bool bFreezeTarget = false;
 
+	UPROPERTY()
+	bool bCanFreezeTarget = false;
+
 	// parent ref required for details customization
 	UPROPERTY(meta = (TransientToolProperty))
 	TWeakObjectPtr<UMeshVertexSculptTool> Tool;
@@ -259,6 +267,12 @@ public:
 	OnInteractiveToolDetailsPanelRequestRebuild OnDetailsPanelRequestRebuild;
 
 protected:
+	// Overriden by subclasses to choose the brushes they use
+	virtual void RegisterBrushes();
+	// Should be overriden by subclasses so that settings are not shared across
+	//  this tool and subclasses.
+	virtual FString GetPropertyCacheIdentifier() const;
+
 	// UMeshSculptToolBase API
 	virtual void InitializeIndicator() override;
 	virtual UPreviewMesh* MakeBrushIndicatorMesh(UObject* Parent, UWorld* World) override;
@@ -290,7 +304,11 @@ protected:
 	void OnDynamicMeshComponentChanged(UDynamicMeshComponent* Component, const FMeshVertexChange* Change, bool bRevert);
 	FDelegateHandle OnDynamicMeshComponentChangedHandle;
 
+	UE_DEPRECATED(5.6, "Vertex sculpt now uses an integer for a brush ID to allow custom brush registration")
 	void UpdateBrushType(EMeshVertexSculptBrushType BrushType);
+private:
+	void UpdateBrushType(int32 BrushID);
+protected:
 
 	TUniquePtr<UE::Geometry::FPolygroupSet> ActiveGroupSet;
 	TArray<int32> TriangleComponentIDs;
@@ -317,7 +335,14 @@ protected:
 	TArray<UE::Geometry::FIndex3i> TriangleROIInBuf;
 	TArray<int> VertexROI;
 	TArray<int> TriangleROIArray;
+	void UpdateROI(const FFrame3d& LocalFrame);
+	UE_DEPRECATED(5.6, "Use the FFrame3d overload instead.")
 	void UpdateROI(const FVector3d& BrushPos);
+private:
+	void UpdateRangeQueryTriBuffer(const FFrame3d& LocalFrame);
+	void PrepROIVertPositionBuffers();
+protected:
+	virtual bool RequireConnectivityToHitPointInStamp() const { return false; }
 
 	UE::Geometry::FUniqueIndexSet NormalsROIBuilder;
 	TArray<std::atomic<bool>> NormalsFlags;		// set of per-vertex or per-element-id flags that indicate
@@ -326,13 +351,27 @@ protected:
 
 	bool bTargetDirty;
 
+	UE_DEPRECATED(5.6, "Vertex sculpt now uses an integer for a brush ID to allow custom brush registration")
 	EMeshVertexSculptBrushType PendingStampType = EMeshVertexSculptBrushType::Smooth;
+private:
+	int32 PendingStampBrushID = (int32)EMeshVertexSculptBrushType::Smooth;
+protected:
 
 	bool UpdateStampPosition(const FRay& WorldRay);
 	TFuture<void> ApplyStamp();
+	
+	FVector3d PreviousRayDirection;
+	bool bMouseMoved = false;
+	// The stamp used last time bMouseMoved was true
+	FSculptBrushStamp LastMovedStamp;
 
 	FRandomStream StampRandomStream;
 
+	// The base mesh is a second copy of our mesh that we can intentionally not update during
+	//  some sculpt strokes so that we can base certain hittesting or processing off that mesh
+	//  instead of the latest result, e.g. when using brushes that offset some max
+	//  amount, so that we only offset relative to the mesh as it was before the start of the
+	//  stroke.
 	FDynamicMesh3 BaseMesh;
 	UE::Geometry::FDynamicMeshOctree3 BaseMeshSpatial;
 	TArray<int32> BaseMeshIndexBuffer;
@@ -372,7 +411,13 @@ protected:
 
 
 protected:
-	virtual bool ShowWorkPlane() const override { return SculptProperties->PrimaryBrushType == EMeshVertexSculptBrushType::FixedPlane; }
+	virtual bool ShowWorkPlane() const override { return SculptProperties->PrimaryBrushID == (int32)EMeshVertexSculptBrushType::FixedPlane; }
+
+private:
+	// Assumes that brush is currently aligned to hit normal
+	void RealignBrush(FMeshSculptBrushOp::EStampAlignmentType AlignmentType);
+
+	double InitialBoundsMaxDim = 0;
 };
 
 
diff --git a/Engine/Plugins/Experimental/MeshModelingToolsetExp/Source/MeshModelingToolsExp/Public/Sculpting/MeshBrushOpBase.h b/Engine/Plugins/Experimental/MeshModelingToolsetExp/Source/MeshModelingToolsExp/Public/Sculpting/MeshBrushOpBase.h
--- a/Engine/Plugins/Experimental/MeshModelingToolsetExp/Source/MeshModelingToolsExp/Public/Sculpting/MeshBrushOpBase.h
+++ b/Engine/Plugins/Experimental/MeshModelingToolsetExp/Source/MeshModelingToolsExp/Public/Sculpting/MeshBrushOpBase.h
@@ -100,6 +100,48 @@ public:
 class MESHMODELINGTOOLSEXP_API FMeshSculptBrushOp
 {
 public:
+	// Determines what region of a mesh a brush wants to affect
+	enum class EBrushRegionType
+	{
+		// Affect a sphere around the current brush location
+		LocalSphere,
+		// Affect a cylinder centered at current brush location and extending infinitely upwards and downwards
+		//  (based on local frame)
+		InfiniteCylinder,
+		// Affect a cylinder whose axis starts at a reference sphere center, passes through the brush position,
+		//  and extends further away from the sphere.
+		CylinderOnSphere,
+	};
+
+	// Determines how the brush expects its local frame to be aligned as the mouse is moved
+	enum class EStampAlignmentType
+	{
+		// Align brush to the hit normal
+		HitNormal,
+		// Align brush to have its local plane face the camera
+		Camera,
+		// Align brush such that its local plane is parallel to a reference plane
+		ReferencePlane,
+		// Align brush such that its local plane normal points away from a reference sphere center
+		ReferenceSphere,
+	};
+
+	// Determines what kind of reference plane the brush wants to use (which will be accessed 
+	//  via CurrentOptions.ConstantReferencePlane)
+	enum class EReferencePlaneType
+	{
+		// Reference plane is not used
+		None,
+		// The reference plane is expected to have been caluclated using triangles in the region
+		//  of interest affected by the brush at its initial application (at start of stroke)
+		InitialROI,
+		// Like InitialROI, but the plane normal is aligned to camera direction, so only the
+		//  centroid is calculated from the region of interest.
+		InitialROI_ViewAligned,
+		// The brush expects custom "work" plane (typically set via gizmo)
+		WorkPlane,
+	};
+
 	virtual ~FMeshSculptBrushOp() {}
 
 	TWeakObjectPtr<UMeshSculptBrushOpProps> PropertySet;
@@ -107,8 +149,19 @@ public:
 	template<typename PropType> 
 	PropType* GetPropertySetAs()
 	{
-		check(PropertySet.IsValid());
-		return CastChecked<PropType>(PropertySet.Get());
+		ensure(PropertySet.IsValid());
+		PropType* CastResult = Cast<PropType>(PropertySet.Get());
+		ensure(CastResult);
+		return CastResult;
+	}
+
+	template<typename PropType>
+	const PropType* GetPropertySetAs() const
+	{
+		ensure(PropertySet.IsValid());
+		PropType* CastResult = Cast<PropType>(PropertySet.Get());
+		ensure(CastResult);
+		return CastResult;
 	}
 
 
@@ -138,6 +191,23 @@ public:
 		return ESculptBrushOpTargetType::SculptMesh;
 	}
 
+	// Determines what region of a mesh a brush wants to affect
+	virtual EBrushRegionType GetBrushRegionType() const
+	{ 
+		return EBrushRegionType::LocalSphere;
+	}
+
+	// Determines how the brush expects its local frame to be aligned as the mouse is moved
+	virtual EStampAlignmentType GetStampAlignmentType() const
+	{
+		if (GetAlignStampToView())
+		{
+			return EStampAlignmentType::Camera;
+		}
+		return EStampAlignmentType::HitNormal;
+	}
+
+	//~ TODO: Deprecate this in favor of GetStampAlignmentType
 	virtual bool GetAlignStampToView() const
 	{
 		return false;
@@ -148,12 +218,35 @@ public:
 		return false;
 	}
 
+	// If this is true and there is no mouse movement, the stamp will be applied with the same
+	//  Local/World frames as the previous application. Irrelevant if IgnoreZeroMovements()
+	//  returns true, since in that case the stamp requires mouse movement to be applied.
+	// Useful for height brushes which change the mesh (and therefore the hit location) but
+	//  want to continue being applied in the same vertical region while the mouse is not moved.
+	virtual bool UseLastStampFrameOnZeroMovement() const
+	{
+		return false;
+	}
 
+	// Determines what kind of plane (if any) the brush wants to be stored at the start of
+	//  the stroke in CurrentOptions.ConstantReferencePlane.
+	virtual EReferencePlaneType GetReferencePlaneType() const
+	{
+		return EReferencePlaneType::None;
+	}
+
+	// Whether the brush wants an average plane to be computed at each stamp application out 
+	//  of the affected triangles (accessed through Stamp.RegionPlane)
 	virtual bool WantsStampRegionPlane() const
 	{
 		return false;
 	}
 
+	virtual bool UsesAlpha() const
+	{
+		return false;
+	}
+
 	virtual bool SupportsVariableSpacing() const
 	{
 		return false;
diff --git a/Engine/Plugins/Experimental/MeshModelingToolsetExp/Source/MeshModelingToolsExp/Public/Sculpting/MeshPlaneBrushOps.h b/Engine/Plugins/Experimental/MeshModelingToolsetExp/Source/MeshModelingToolsExp/Public/Sculpting/MeshPlaneBrushOps.h
--- a/Engine/Plugins/Experimental/MeshModelingToolsetExp/Source/MeshModelingToolsExp/Public/Sculpting/MeshPlaneBrushOps.h
+++ b/Engine/Plugins/Experimental/MeshModelingToolsetExp/Source/MeshModelingToolsExp/Public/Sculpting/MeshPlaneBrushOps.h
@@ -114,6 +114,31 @@ public:
 	double BrushSpeedTuning = 6.0;
 
 	FFrame3d StrokePlane;
+	
+	virtual EReferencePlaneType GetReferencePlaneType() const
+	{
+		// FPlaneBrushOp doesn't have one type of reference plane it expects. Instead, it can
+		//  be used as multiple different brushes depending on what reference plane it gets 
+		//  (for instance, a flatten brush, or bring to work plane brush, etc.). 
+		// Ideally we would have some enum on the brush instance that will change which 
+		//  reference plane type it requests, so that each brush instance can be configured
+		//  to act as a particular brush type. What we do instead (for now) is rely on the
+		//  fact that we configure each of these different plane brush types with a different
+		//  property set, so by looking at our property set, we know what brush we are.
+
+		if (ExactCast<UViewAlignedPlaneBrushOpProps>(PropertySet.Get()))
+		{
+			return EReferencePlaneType::InitialROI_ViewAligned;
+		}
+		else if (ExactCast<UFixedPlaneBrushOpProps>(PropertySet.Get()))
+		{
+			return EReferencePlaneType::WorkPlane;
+		}
+		else // UPlaneBrushOpProps
+		{
+			return EReferencePlaneType::InitialROI;
+		}
+	}
 
 	virtual void BeginStroke(const FDynamicMesh3* Mesh, const FSculptBrushStamp& Stamp, const TArray<int32>& InitialVertices) override
 	{
diff --git a/Engine/Plugins/Experimental/MeshModelingToolsetExp/Source/MeshModelingToolsExp/Public/Sculpting/MeshSculptBrushOps.h b/Engine/Plugins/Experimental/MeshModelingToolsetExp/Source/MeshModelingToolsExp/Public/Sculpting/MeshSculptBrushOps.h
--- a/Engine/Plugins/Experimental/MeshModelingToolsetExp/Source/MeshModelingToolsExp/Public/Sculpting/MeshSculptBrushOps.h
+++ b/Engine/Plugins/Experimental/MeshModelingToolsetExp/Source/MeshModelingToolsExp/Public/Sculpting/MeshSculptBrushOps.h
@@ -50,6 +50,10 @@ public:
 	{
 		return true;
 	}
+	virtual bool UsesAlpha() const override
+	{
+		return true;
+	}
 
 	virtual void ApplyStamp(const FDynamicMesh3* Mesh, const FSculptBrushStamp& Stamp, const TArray<int32>& Vertices, TArray<FVector3d>& NewPositionsOut) override
 	{
@@ -133,6 +137,10 @@ public:
 	{
 		return true;
 	}
+	virtual bool UsesAlpha() const override
+	{
+		return true;
+	}
 
 	virtual void ApplyStamp(const FDynamicMesh3* Mesh, const FSculptBrushStamp& Stamp, const TArray<int32>& Vertices, TArray<FVector3d>& NewPositionsOut) override
 	{
@@ -226,6 +234,10 @@ public:
 	{
 		return true;
 	}
+	virtual bool UsesAlpha() const override
+	{
+		return true;
+	}
 
 	virtual void ApplyStamp(const FDynamicMesh3* Mesh, const FSculptBrushStamp& Stamp, const TArray<int32>& Vertices, TArray<FVector3d>& NewPositionsOut) override
 	{
