ModelingTools: Added option to BakeVertex tool and geometry script node to reuse existing vertex color topology rather than always regenerating it. This is necessary for cases where we need to preserve existing color data.

diff --git a/Engine/Plugins/Experimental/MeshModelingToolsetExp/Source/MeshModelingToolsExp/Private/BakeMeshAttributeVertexTool.cpp b/Engine/Plugins/Experimental/MeshModelingToolsetExp/Source/MeshModelingToolsExp/Private/BakeMeshAttributeVertexTool.cpp
--- a/Engine/Plugins/Experimental/MeshModelingToolsetExp/Source/MeshModelingToolsExp/Private/BakeMeshAttributeVertexTool.cpp
+++ b/Engine/Plugins/Experimental/MeshModelingToolsetExp/Source/MeshModelingToolsExp/Private/BakeMeshAttributeVertexTool.cpp
@@ -317,6 +317,7 @@ void UBakeMeshAttributeVertexTool::Setup()
 	Settings->WatchProperty(Settings->OutputTypeB, [this](int32) { OpState |= EBakeOpState::Evaluate; UpdateOnModeChange(); });
 	Settings->WatchProperty(Settings->OutputTypeA, [this](int32) { OpState |= EBakeOpState::Evaluate; UpdateOnModeChange(); });
 	Settings->WatchProperty(Settings->PreviewMode, [this](EBakeVertexChannel) { UpdateVisualization(); });
+	Settings->WatchProperty(Settings->TopologyMode, [this](EBakeVertexTopology) { bColorTopologyValid = false; OpState |= EBakeOpState::Evaluate; });
 	Settings->WatchProperty(Settings->bSplitAtNormalSeams, [this](bool) { bColorTopologyValid = false; OpState |= EBakeOpState::Evaluate; });
 	Settings->WatchProperty(Settings->bSplitAtUVSeams, [this](bool) { bColorTopologyValid = false; OpState |= EBakeOpState::Evaluate; });
 
@@ -602,34 +603,52 @@ void UBakeMeshAttributeVertexTool::UpdateVisualization()
 void UBakeMeshAttributeVertexTool::UpdateColorTopology()
 {
 	// Update PreviewMesh color topology
-	PreviewMesh->EditMesh([this](FDynamicMesh3& Mesh)
+	if (Settings->TopologyMode == EBakeVertexTopology::CreateNew)
 	{
-		Mesh.EnableAttributes();
-		Mesh.Attributes()->EnablePrimaryColors();
-		Mesh.Attributes()->PrimaryColors()->ClearElements();
-
-		FDynamicMeshNormalOverlay* NormalOverlay = Mesh.Attributes()->PrimaryNormals();
-		FDynamicMeshUVOverlay* UVOverlay = Mesh.Attributes()->PrimaryUV();
-		Mesh.Attributes()->PrimaryColors()->CreateFromPredicate(
-			[&](int /*ParentVID*/, int TriIDA, int TriIDB) -> bool
-			{
-				auto OverlayCanShare = [TriIDA, TriIDB] (auto Overlay) -> bool
+		PreviewMesh->EditMesh([this](FDynamicMesh3& Mesh)
+		{
+			Mesh.EnableAttributes();
+			Mesh.Attributes()->EnablePrimaryColors();
+			Mesh.Attributes()->PrimaryColors()->ClearElements();
+
+			FDynamicMeshNormalOverlay* NormalOverlay = Mesh.Attributes()->PrimaryNormals();
+			FDynamicMeshUVOverlay* UVOverlay = Mesh.Attributes()->PrimaryUV();
+			Mesh.Attributes()->PrimaryColors()->CreateFromPredicate(
+				[this, NormalOverlay, UVOverlay](int /*ParentVID*/, int TriIDA, int TriIDB) -> bool
 				{
-					return Overlay ? Overlay->AreTrianglesConnected(TriIDA, TriIDB) : true;
-				};
+					auto OverlayCanShare = [TriIDA, TriIDB] (auto Overlay) -> bool
+					{
+						return Overlay ? Overlay->AreTrianglesConnected(TriIDA, TriIDB) : true;
+					};
 				
-				bool bCanShare = true;
-				if (Settings->bSplitAtNormalSeams)
-				{
-					bCanShare = bCanShare && OverlayCanShare(NormalOverlay);
-				}
-				if (Settings->bSplitAtUVSeams)
-				{
-					bCanShare = bCanShare && OverlayCanShare(UVOverlay);
-				}
-				return bCanShare;
-			}, 0.0f);
-	});
+					bool bCanShare = true;
+					if (Settings->bSplitAtNormalSeams)
+					{
+						bCanShare = bCanShare && OverlayCanShare(NormalOverlay);
+					}
+					if (Settings->bSplitAtUVSeams)
+					{
+						bCanShare = bCanShare && OverlayCanShare(UVOverlay);
+					}
+					return bCanShare;
+				}, 0.0f);
+		});
+	}
+	else
+	{
+		PreviewMesh->EditMesh([this](FDynamicMesh3& Mesh)
+		{
+			// Copy original color overlay from TargetMesh
+			Mesh.EnableAttributes();
+			Mesh.Attributes()->EnablePrimaryColors();
+			Mesh.Attributes()->PrimaryColors()->ClearElements();
+			
+			if (TargetMesh->Attributes() && TargetMesh->Attributes()->PrimaryColors())
+			{
+				Mesh.Attributes()->PrimaryColors()->Copy(*TargetMesh->Attributes()->PrimaryColors());
+			}
+		});
+	}
 
 	// Copy source vertex colors onto new color overlay topology.
 	UpdateSourceVertexColors();
@@ -689,6 +708,7 @@ void UBakeMeshAttributeVertexTool::UpdateResult()
 	BakeSettings.OutputTypePerChannel[1] = static_cast<EBakeMapType>(Settings->OutputTypeG);
 	BakeSettings.OutputTypePerChannel[2] = static_cast<EBakeMapType>(Settings->OutputTypeB);
 	BakeSettings.OutputTypePerChannel[3] = static_cast<EBakeMapType>(Settings->OutputTypeA);
+	BakeSettings.TopologyMode = Settings->TopologyMode;
 	BakeSettings.bSplitAtNormalSeams = Settings->bSplitAtNormalSeams;
 	BakeSettings.bSplitAtUVSeams = Settings->bSplitAtUVSeams;
 	BakeSettings.bProjectionInWorldSpace = InputMeshSettings->bProjectionInWorldSpace;
@@ -702,6 +722,19 @@ void UBakeMeshAttributeVertexTool::UpdateResult()
 	OpState &= ~EBakeOpState::Invalid;
 
 	// Validate bake inputs
+	if (NumColorElements <= 0)
+	{
+		if (Settings->TopologyMode == EBakeVertexTopology::UseExisting)
+		{
+			GetToolManager()->DisplayMessage(LOCTEXT("InvalidExistingVertexColorTopology", "Topology Mode is set to UseExisting, but no valid existing vertex color topology was found."), EToolMessageLevel::UserWarning);
+		}
+		else
+		{
+			GetToolManager()->DisplayMessage(LOCTEXT("InvalidNewVertexColorTopology", "Invalid vertex color topology."), EToolMessageLevel::UserWarning);
+		}
+		OpState |= EBakeOpState::Invalid;
+	}
+	
 	const FImageDimensions Dimensions(NumColorElements, 1);
 	if (CachedBakeSettings.OutputMode == EBakeVertexOutput::RGBA)
 	{
diff --git a/Engine/Plugins/Experimental/MeshModelingToolsetExp/Source/MeshModelingToolsExp/Public/BakeMeshAttributeVertexTool.h b/Engine/Plugins/Experimental/MeshModelingToolsetExp/Source/MeshModelingToolsExp/Public/BakeMeshAttributeVertexTool.h
--- a/Engine/Plugins/Experimental/MeshModelingToolsetExp/Source/MeshModelingToolsExp/Public/BakeMeshAttributeVertexTool.h
+++ b/Engine/Plugins/Experimental/MeshModelingToolsetExp/Source/MeshModelingToolsExp/Public/BakeMeshAttributeVertexTool.h
@@ -32,6 +32,16 @@ public:
 };
 
 
+UENUM()
+enum class EBakeVertexTopology
+{
+	/* Generate new vertex color topology, optionally using normal and UV seams to determine splits */
+	CreateNew,
+	/* Use existing vertex color topology on the Target Mesh */
+	UseExisting
+};
+
+
 UENUM()
 enum class EBakeVertexOutput
 {
@@ -97,12 +107,16 @@ public:
 	UPROPERTY(EditAnywhere, Category = BakeOutput, meta = (TransientToolProperty))
 	EBakeVertexChannel PreviewMode = EBakeVertexChannel::RGBA;
 
+	/** The vertex color topology to use for the bake */
+	UPROPERTY(EditAnywhere, Category = BakeOutput, meta = (TransientToolProperty))
+	EBakeVertexTopology TopologyMode = EBakeVertexTopology::CreateNew;
+
 	/** If true, compute a separate vertex color for each unique normal on a vertex */
-	UPROPERTY(EditAnywhere, AdvancedDisplay, Category = BakeOutput)
+	UPROPERTY(EditAnywhere, AdvancedDisplay, Category = BakeOutput, meta=(EditCondition="TopologyMode == EBakeVertexTopology::CreateNew"))
 	bool bSplitAtNormalSeams = false;
 
 	/** If true, Compute a separate vertex color for each unique UV on a vertex. */
-	UPROPERTY(EditAnywhere, AdvancedDisplay, Category = BakeOutput, meta=(DisplayName = "Split at UV Seams"))
+	UPROPERTY(EditAnywhere, AdvancedDisplay, Category = BakeOutput, meta=(DisplayName = "Split at UV Seams", EditCondition="TopologyMode == EBakeVertexTopology::CreateNew"))
 	bool bSplitAtUVSeams = false;
 };
 
@@ -176,6 +190,7 @@ protected:
 		EBakeMapType OutputType = EBakeMapType::TangentSpaceNormal;
 		EBakeMapType OutputTypePerChannel[4] = { EBakeMapType::None, EBakeMapType::None, EBakeMapType::None, EBakeMapType::None };
 		EBakeVertexChannel PreviewMode = EBakeVertexChannel::RGBA;
+		EBakeVertexTopology TopologyMode = EBakeVertexTopology::CreateNew;
 		float ProjectionDistance = 3.0;
 		bool bProjectionInWorldSpace = false;
 		bool bSplitAtNormalSeams = false;
@@ -189,6 +204,7 @@ protected:
 				OutputTypePerChannel[1] == Other.OutputTypePerChannel[1] &&
 				OutputTypePerChannel[2] == Other.OutputTypePerChannel[2] &&
 				OutputTypePerChannel[3] == Other.OutputTypePerChannel[3] &&
+				TopologyMode == Other.TopologyMode &&
 				bProjectionInWorldSpace == Other.bProjectionInWorldSpace &&
 				ProjectionDistance == Other.ProjectionDistance &&
 				bSplitAtNormalSeams == Other.bSplitAtNormalSeams &&
diff --git a/Engine/Plugins/Runtime/GeometryScripting/Source/GeometryScriptingCore/Private/MeshBakeFunctions.cpp b/Engine/Plugins/Runtime/GeometryScripting/Source/GeometryScriptingCore/Private/MeshBakeFunctions.cpp
--- a/Engine/Plugins/Runtime/GeometryScripting/Source/GeometryScriptingCore/Private/MeshBakeFunctions.cpp
+++ b/Engine/Plugins/Runtime/GeometryScripting/Source/GeometryScriptingCore/Private/MeshBakeFunctions.cpp
@@ -722,31 +722,43 @@ namespace GeometryScriptBakeLocals
 
 		// Initialize the color overlay on the TargetMesh
 		FDynamicMesh3& TargetMeshRef = TargetMesh->GetMeshRef();
-		TargetMeshRef.EnableAttributes();
-		TargetMeshRef.Attributes()->DisablePrimaryColors();
-		TargetMeshRef.Attributes()->EnablePrimaryColors();
-
-		FDynamicMeshNormalOverlay* NormalOverlay = TargetMeshRef.Attributes()->PrimaryNormals();
-		FDynamicMeshUVOverlay* UVOverlay = TargetMeshRef.Attributes()->PrimaryUV();
-		TargetMeshRef.Attributes()->PrimaryColors()->CreateFromPredicate(
-			[&BakeOptions, NormalOverlay, UVOverlay](int ParentVID, int TriIDA, int TriIDB) -> bool
-			{
-				auto OverlayCanShare = [&] (auto Overlay) -> bool
+		if (BakeOptions.TopologyMode == EGeometryScriptBakeVertexTopology::CreateNew)
+		{
+			TargetMeshRef.EnableAttributes();
+			TargetMeshRef.Attributes()->EnablePrimaryColors();
+			TargetMeshRef.Attributes()->PrimaryColors()->ClearElements();
+
+			FDynamicMeshNormalOverlay* NormalOverlay = TargetMeshRef.Attributes()->PrimaryNormals();
+			FDynamicMeshUVOverlay* UVOverlay = TargetMeshRef.Attributes()->PrimaryUV();
+			TargetMeshRef.Attributes()->PrimaryColors()->CreateFromPredicate(
+				[&BakeOptions, NormalOverlay, UVOverlay](int ParentVID, int TriIDA, int TriIDB) -> bool
 				{
-					return Overlay ? Overlay->AreTrianglesConnected(TriIDA, TriIDB) : true;
-				};
+					auto OverlayCanShare = [TriIDA, TriIDB] (auto Overlay) -> bool
+					{
+						return Overlay ? Overlay->AreTrianglesConnected(TriIDA, TriIDB) : true;
+					};
 				
-				bool bCanShare = true;
-				if (BakeOptions.bSplitAtNormalSeams)
-				{
-					bCanShare = bCanShare && OverlayCanShare(NormalOverlay);
-				}
-				if (BakeOptions.bSplitAtUVSeams)
-				{
-					bCanShare = bCanShare && OverlayCanShare(UVOverlay);
-				}
-				return bCanShare;
-			}, 0.0f);
+					bool bCanShare = true;
+					if (BakeOptions.bSplitAtNormalSeams)
+					{
+						bCanShare = bCanShare && OverlayCanShare(NormalOverlay);
+					}
+					if (BakeOptions.bSplitAtUVSeams)
+					{
+						bCanShare = bCanShare && OverlayCanShare(UVOverlay);
+					}
+					return bCanShare;
+				}, 0.0f);
+		}
+		else
+		{
+			if (!TargetMeshRef.Attributes() ||
+				!TargetMeshRef.Attributes()->PrimaryColors())
+			{
+				UE::Geometry::AppendError(Debug, EGeometryScriptErrorType::InvalidInputs, LOCTEXT("BakeVertex_InvalidVertexColorTopology", "BakeVertex: Topology Mode was set to UseExisting, but existing colors were not found."));
+				return nullptr;
+			}
+		}
 
 		if (bIsBakeToSelf || BakeTypes.OutputMode == EGeometryScriptBakeOutputMode::PerChannel)
 		{
diff --git a/Engine/Plugins/Runtime/GeometryScripting/Source/GeometryScriptingCore/Public/GeometryScript/MeshBakeFunctions.h b/Engine/Plugins/Runtime/GeometryScripting/Source/GeometryScriptingCore/Public/GeometryScript/MeshBakeFunctions.h
--- a/Engine/Plugins/Runtime/GeometryScripting/Source/GeometryScriptingCore/Public/GeometryScript/MeshBakeFunctions.h
+++ b/Engine/Plugins/Runtime/GeometryScripting/Source/GeometryScriptingCore/Public/GeometryScript/MeshBakeFunctions.h
@@ -96,6 +96,15 @@ enum class EGeometryScriptBakeNormalSpace : uint8
 	Object
 };
 
+UENUM(BlueprintType)
+enum class EGeometryScriptBakeVertexTopology : uint8
+{
+	/* Create new vertex color topology */
+	CreateNew,
+	/* Use existing vertex color topology */
+	UseExisting
+};
+
 struct GEOMETRYSCRIPTINGCORE_API FGeometryScriptBakeTypes
 {
 };
@@ -269,6 +278,10 @@ struct GEOMETRYSCRIPTINGCORE_API FGeometryScriptBakeVertexOptions
 {
 	GENERATED_BODY()
 
+	/** Option to generate new vertex color topology or reuse existing vertex color topology. */
+	UPROPERTY(EditAnywhere, BlueprintReadWrite, Category = Options)
+	EGeometryScriptBakeVertexTopology TopologyMode = EGeometryScriptBakeVertexTopology::CreateNew;
+
 	/** If true, compute a separate vertex color for each unique normal on a vertex */
 	UPROPERTY(EditAnywhere, BlueprintReadWrite, Category = Options)
 	bool bSplitAtNormalSeams = false;
