fix unknown brush type warnings in dynamic sculpt tool when using kelvin brushes

diff --git a/Engine/Plugins/Experimental/MeshModelingToolsetExp/Source/MeshModelingToolsExp/Private/DynamicMeshSculptTool.cpp b/Engine/Plugins/Experimental/MeshModelingToolsetExp/Source/MeshModelingToolsExp/Private/DynamicMeshSculptTool.cpp
--- a/Engine/Plugins/Experimental/MeshModelingToolsetExp/Source/MeshModelingToolsExp/Private/DynamicMeshSculptTool.cpp
+++ b/Engine/Plugins/Experimental/MeshModelingToolsetExp/Source/MeshModelingToolsExp/Private/DynamicMeshSculptTool.cpp
@@ -1421,6 +1421,10 @@ bool UDynamicMeshSculptTool::UpdateBrushPosition(const FRay& WorldRay)
 	case EDynamicMeshSculptBrushType::Flatten:
 	case EDynamicMeshSculptBrushType::Plane:
 	case EDynamicMeshSculptBrushType::FixedPlane:
+	case EDynamicMeshSculptBrushType::PullSharpKelvin:
+	case EDynamicMeshSculptBrushType::PullKelvin:
+	case EDynamicMeshSculptBrushType::TwistKelvin:
+	case EDynamicMeshSculptBrushType::ScaleKelvin:
 		bHit = UpdateBrushPositionOnSculptMesh(WorldRay, false);
 		break;
 
