Modeling Mode: add ability to launch vertex sculpt tool with a desired brush selected

diff --git a/Engine/Plugins/Experimental/MeshModelingToolsetExp/Source/MeshModelingToolsExp/Private/MeshVertexSculptTool.cpp b/Engine/Plugins/Experimental/MeshModelingToolsetExp/Source/MeshModelingToolsExp/Private/MeshVertexSculptTool.cpp
--- a/Engine/Plugins/Experimental/MeshModelingToolsetExp/Source/MeshModelingToolsExp/Private/MeshVertexSculptTool.cpp
+++ b/Engine/Plugins/Experimental/MeshModelingToolsetExp/Source/MeshModelingToolsExp/Private/MeshVertexSculptTool.cpp
@@ -54,6 +54,13 @@ namespace
 #else
 	static EAsyncExecution VertexSculptToolAsyncExecTarget = EAsyncExecution::ThreadPool;
 #endif
+
+static FToolTargetTypeRequirements TypeRequirements({
+	UMaterialProvider::StaticClass(),
+	UDynamicMeshProvider::StaticClass(),
+	UDynamicMeshCommitter::StaticClass(),
+	UPrimitiveComponentBackedTarget::StaticClass()
+});
 }
 
 
@@ -69,12 +76,19 @@ UMeshSurfacePointTool* UMeshVertexSculptToolBuilder::CreateNewTool(const FToolBu
 
 const FToolTargetTypeRequirements& UMeshVertexSculptToolBuilder::GetTargetRequirements() const
 {
-	static FToolTargetTypeRequirements TypeRequirements({
-		UMaterialProvider::StaticClass(),
-		UDynamicMeshProvider::StaticClass(),
-		UDynamicMeshCommitter::StaticClass(),
-		UPrimitiveComponentBackedTarget::StaticClass()
-	});
+	return TypeRequirements;
+}
+
+UMeshSurfacePointTool* UMeshVertexSculptBrushToolBuilder::CreateNewTool(const FToolBuilderState& SceneState) const
+{
+	UMeshVertexSculptTool* SculptTool = NewObject<UMeshVertexSculptTool>(SceneState.ToolManager);
+	SculptTool->SetWorld(SceneState.World);
+	SculptTool->SetDefaultPrimaryBrushID(DefaultPrimaryBrushID);
+	return SculptTool;
+}
+
+const FToolTargetTypeRequirements& UMeshVertexSculptBrushToolBuilder::GetTargetRequirements() const
+{
 	return TypeRequirements;
 }
 
@@ -215,6 +229,10 @@ void UMeshVertexSculptTool::Setup()
 	UMeshSculptToolBase::BrushProperties->bShowPerBrushProps = false;
 	UMeshSculptToolBase::BrushProperties->bShowFalloff = false;
 	SculptProperties->RestoreProperties(this, GetPropertyCacheIdentifier());
+	if (DefaultPrimaryBrushID >= 0)
+	{
+		SculptProperties->PrimaryBrushID = DefaultPrimaryBrushID;
+	}
 	AddToolPropertySource(SculptProperties);
 	CalculateBrushRadius();
 
@@ -235,6 +253,8 @@ void UMeshVertexSculptTool::Setup()
 
 	RegisterBrushes();
 
+	ensure(SculptProperties->PrimaryBrushID < RegisteredPrimaryBrushTypes.Num());
+
 	// falloffs
 	RegisterStandardFalloffTypes();
 
@@ -1598,6 +1618,11 @@ void UMeshVertexSculptTool::EndChange()
 	ActiveVertexChange = nullptr;
 }
 
+void UMeshVertexSculptTool::SetDefaultPrimaryBrushID(const int32 InPrimaryBrushID)
+{
+	DefaultPrimaryBrushID = InPrimaryBrushID;
+}
+
 
 void UMeshVertexSculptTool::WaitForPendingUndoRedo()
 {
diff --git a/Engine/Plugins/Experimental/MeshModelingToolsetExp/Source/MeshModelingToolsExp/Public/MeshVertexSculptTool.h b/Engine/Plugins/Experimental/MeshModelingToolsetExp/Source/MeshModelingToolsExp/Public/MeshVertexSculptTool.h
--- a/Engine/Plugins/Experimental/MeshModelingToolsetExp/Source/MeshModelingToolsExp/Public/MeshVertexSculptTool.h
+++ b/Engine/Plugins/Experimental/MeshModelingToolsetExp/Source/MeshModelingToolsExp/Public/MeshVertexSculptTool.h
@@ -101,6 +101,17 @@ enum class EMeshVertexSculptBrushType : uint8
 
 };
 
+UCLASS()
+class MESHMODELINGTOOLSEXP_API UMeshVertexSculptBrushToolBuilder : public UMeshSurfacePointMeshEditingToolBuilder
+{
+	GENERATED_BODY()
+public:
+	virtual UMeshSurfacePointTool* CreateNewTool(const FToolBuilderState& SceneState) const override;
+	
+	int32 DefaultPrimaryBrushID = -1;
+protected:
+	virtual const FToolTargetTypeRequirements& GetTargetRequirements() const override;
+};
 
 
 /** Brush Triangle Filter Type */
@@ -255,6 +266,8 @@ public:
 	virtual void SetActiveFalloffType(int32 Identifier);
 	virtual void SetRegionFilterType(int32 Identifier);
 
+	void SetDefaultPrimaryBrushID(int32 InPrimaryBrushID);
+
 	/**
 	* OnDetailsPanelRequestRebuild is broadcast by the tool when it detects it needs to have it's details panel rebuilt outside
 	* of normal rebuilding triggers, such as changing property set objects. This is useful in rare circumstances, such as when
@@ -409,6 +422,8 @@ protected:
 	void BeginChange();
 	void EndChange();
 
+	int32 DefaultPrimaryBrushID = -1;
+
 
 protected:
 	virtual bool ShowWorkPlane() const override { return SculptProperties->PrimaryBrushID == (int32)EMeshVertexSculptBrushType::FixedPlane; }
