Modeling Mode: fix TypeRequirements conflict in VSculpt Tool

diff --git a/Engine/Plugins/Experimental/MeshModelingToolsetExp/Source/MeshModelingToolsExp/Private/MeshVertexSculptTool.cpp b/Engine/Plugins/Experimental/MeshModelingToolsetExp/Source/MeshModelingToolsExp/Private/MeshVertexSculptTool.cpp
--- a/Engine/Plugins/Experimental/MeshModelingToolsetExp/Source/MeshModelingToolsExp/Private/MeshVertexSculptTool.cpp
+++ b/Engine/Plugins/Experimental/MeshModelingToolsetExp/Source/MeshModelingToolsExp/Private/MeshVertexSculptTool.cpp
@@ -54,13 +54,6 @@ namespace
 #else
 	static EAsyncExecution VertexSculptToolAsyncExecTarget = EAsyncExecution::ThreadPool;
 #endif
-
-static FToolTargetTypeRequirements TypeRequirements({
-	UMaterialProvider::StaticClass(),
-	UDynamicMeshProvider::StaticClass(),
-	UDynamicMeshCommitter::StaticClass(),
-	UPrimitiveComponentBackedTarget::StaticClass()
-});
 }
 
 
@@ -74,9 +67,16 @@ UMeshSurfacePointTool* UMeshVertexSculptToolBuilder::CreateNewTool(const FToolBu
 	return SculptTool;
 }
 
+FToolTargetTypeRequirements UMeshVertexSculptToolBuilder::VSculptTypeRequirements({
+	UMaterialProvider::StaticClass(),
+	UDynamicMeshProvider::StaticClass(),
+	UDynamicMeshCommitter::StaticClass(),
+	UPrimitiveComponentBackedTarget::StaticClass()
+});
+
 const FToolTargetTypeRequirements& UMeshVertexSculptToolBuilder::GetTargetRequirements() const
 {
-	return TypeRequirements;
+	return VSculptTypeRequirements;
 }
 
 UMeshSurfacePointTool* UMeshVertexSculptBrushToolBuilder::CreateNewTool(const FToolBuilderState& SceneState) const
@@ -87,12 +87,6 @@ UMeshSurfacePointTool* UMeshVertexSculptBrushToolBuilder::CreateNewTool(const FT
 	return SculptTool;
 }
 
-const FToolTargetTypeRequirements& UMeshVertexSculptBrushToolBuilder::GetTargetRequirements() const
-{
-	return TypeRequirements;
-}
-
-
 /*
  * internal Change classes
  */
diff --git a/Engine/Plugins/Experimental/MeshModelingToolsetExp/Source/MeshModelingToolsExp/Public/MeshVertexSculptTool.h b/Engine/Plugins/Experimental/MeshModelingToolsetExp/Source/MeshModelingToolsExp/Public/MeshVertexSculptTool.h
--- a/Engine/Plugins/Experimental/MeshModelingToolsetExp/Source/MeshModelingToolsExp/Public/MeshVertexSculptTool.h
+++ b/Engine/Plugins/Experimental/MeshModelingToolsetExp/Source/MeshModelingToolsExp/Public/MeshVertexSculptTool.h
@@ -39,6 +39,8 @@ public:
 	virtual UMeshSurfacePointTool* CreateNewTool(const FToolBuilderState& SceneState) const override;
 protected:
 	virtual const FToolTargetTypeRequirements& GetTargetRequirements() const override;
+private:
+	static FToolTargetTypeRequirements VSculptTypeRequirements;
 };
 
 
@@ -102,15 +104,13 @@ enum class EMeshVertexSculptBrushType : uint8
 };
 
 UCLASS()
-class MESHMODELINGTOOLSEXP_API UMeshVertexSculptBrushToolBuilder : public UMeshSurfacePointMeshEditingToolBuilder
+class MESHMODELINGTOOLSEXP_API UMeshVertexSculptBrushToolBuilder : public UMeshVertexSculptToolBuilder
 {
 	GENERATED_BODY()
 public:
 	virtual UMeshSurfacePointTool* CreateNewTool(const FToolBuilderState& SceneState) const override;
 	
 	int32 DefaultPrimaryBrushID = -1;
-protected:
-	virtual const FToolTargetTypeRequirements& GetTargetRequirements() const override;
 };
 
 
