[GeometryProcessing] CombineMeshInstances - Fixing some of the behavior of the new CombineMeshInstances API, some bugs but also making sure the behavior of the existing CombineMeshInstances() function is still possible within the new one.

- Fixed some bugs regarding UVs not being preserved.
- Now exposing the "LOD Demotion" as an option rather than assuming higher LODs should trickle down to max the budget.
- Allowing triangle budget to be applied on Copied LODs.
- Renamed EDecorationHandling to EDecorativePartHandling for consistency.
- Fixing source LOD for approximated LOD only looking at Copied LODs and not considering simplified LODs.
- Removing adding missing tri-count condition when generating approximated LODs.

diff --git a/Engine/Plugins/Experimental/MeshModelingToolsetExp/Source/GeometryProcessingAdapters/Private/CombineMeshInstancesImpl.cpp b/Engine/Plugins/Experimental/MeshModelingToolsetExp/Source/GeometryProcessingAdapters/Private/CombineMeshInstancesImpl.cpp
--- a/Engine/Plugins/Experimental/MeshModelingToolsetExp/Source/GeometryProcessingAdapters/Private/CombineMeshInstancesImpl.cpp
+++ b/Engine/Plugins/Experimental/MeshModelingToolsetExp/Source/GeometryProcessingAdapters/Private/CombineMeshInstancesImpl.cpp
@@ -3896,7 +3896,7 @@ void BuildCombinedMeshFromPrecomputedMeshes(
 				{
 					PromotedAvgError = PartMeshes.SimplifiedMeshErrors.IsEmpty() ? PartMeshes.ApproximatedMeshErrors[0].AverageError : PartMeshes.SimplifiedMeshErrors[0].AverageError;
 				}
-				return 0;
+				return ECombinedLODType::Copied;
 			}
 			SourceIdx -= PartMeshes.Source.Num();
 			if (SourceIdx < PartMeshes.Simplified.Num())
@@ -3911,7 +3911,7 @@ void BuildCombinedMeshFromPrecomputedMeshes(
 				{
 					PromotedAvgError = PartMeshes.ApproximatedMeshErrors[0].AverageError;
 				}
-				return 1;
+				return ECombinedLODType::Simplified;
 			}
 			SourceIdx -= PartMeshes.Simplified.Num();
 			WithinCategoryIdx = SourceIdx;
@@ -3920,7 +3920,7 @@ void BuildCombinedMeshFromPrecomputedMeshes(
 			{
 				PromotedAvgError = PartMeshes.ApproximatedMeshErrors[SourceIdx + 1].AverageError;
 			}
-			return 2;
+			return ECombinedLODType::Approximated;
 		};
 	auto SourceIndexToMesh = [&Assembly](int32 PartIdx, int32 SourceIdx)
 		{
@@ -3958,11 +3958,11 @@ void BuildCombinedMeshFromPrecomputedMeshes(
 			bool bIsDecorative = Assembly.Parts[PartIdx]->Instances[InstIdx].DetailLevel == EMeshDetailLevel::Decorative;
 			if (bIsDecorative)
 			{
-				if (PerLODOptions[LODLevel].Decorations == IGeometryProcessing_CombineMeshInstances::EDecorationHandling::Remove)
+				if (PerLODOptions[LODLevel].DecorativePartHandling == IGeometryProcessing_CombineMeshInstances::EDecorativePartHandling::Remove)
 				{
 					return nullptr;
 				}
-				else if (PerLODOptions[LODLevel].Decorations == IGeometryProcessing_CombineMeshInstances::EDecorationHandling::Approximate)
+				else if (PerLODOptions[LODLevel].DecorativePartHandling == IGeometryProcessing_CombineMeshInstances::EDecorativePartHandling::Approximate)
 				{
 					SourceIdx = MAX_int32; // approximate decorations with the coarsest available approximation
 				}
@@ -4003,6 +4003,11 @@ void BuildCombinedMeshFromPrecomputedMeshes(
 		int32 PrevLODLevel = FMath::Max(0, LODLevel - 1);
 		ECombinedLODType LODType = LODTypes[LODLevel];
 		const FLODOpts& Opts = PerLODOptions[LODLevel];
+
+		const bool bApplyLODBudget = LODType != ECombinedLODType::VoxWrapped 
+			&& Opts.bEnableBudgetStrategy_PartLODPromotion
+			&& Opts.TriangleBudgetOptions.Method != IGeometryProcessing_CombineMeshInstances::ETriangleBudgetMethod::NoRestriction;
+
 		if (LODType == ECombinedLODType::Copied)
 		{
 			for (int32 PartIdx = 0; PartIdx < NumParts; ++PartIdx)
@@ -4014,25 +4019,51 @@ void BuildCombinedMeshFromPrecomputedMeshes(
 		}
 		else if (LODType == ECombinedLODType::Simplified || LODType == ECombinedLODType::Approximated)
 		{
-			// first pass: copy forward from the previous source indices, + make all part indices appropriate for the requested LOD type
 			for (int32 PartIdx = 0; PartIdx < NumParts; ++PartIdx)
 			{
 				const IGeometryProcessing_CombineMeshInstances::FSinglePartMeshSet& PartMeshes = *Assembly.Parts[PartIdx]->PrecomputedMeshes;
-				// Start from the previous LOD level's source index
-				int32 UseSourceIdx = PartSourcesPerLOD[PrevLODLevel * NumParts + PartIdx];
-				// If approximated meshes are required, enforce that
-				if (LODType == ECombinedLODType::Approximated)
+				int32 UseSourceIdx = LODLevel;
+
+				// If LOD Demotion is enabled, start from the previous LOD level's source index instead, try to keep the higher LOD geometry for as long as possible
+				if (bApplyLODBudget && Opts.bEnableBudgetStrategy_PartLODDemotion)
 				{
-					UseSourceIdx = FMath::Max(UseSourceIdx, PartMeshes.Source.Num() + PartMeshes.Simplified.Num());
+					UseSourceIdx = PartSourcesPerLOD[PrevLODLevel * NumParts + PartIdx];
 				}
+
 				PartSourcesPerLOD[LODLevel * NumParts + PartIdx] = UseSourceIdx;
 			}
+		}
+		else if (LODType == ECombinedLODType::VoxWrapped)
+		{
+			int32 TargetTriCount = Opts.TriangleBudgetOptions.TriangleBudget;
+			if (Opts.TriangleBudgetOptions.Method == IGeometryProcessing_CombineMeshInstances::ETriangleBudgetMethod::UsePercentageOfPreviousLOD
+				|| Opts.TriangleBudgetOptions.Method == IGeometryProcessing_CombineMeshInstances::ETriangleBudgetMethod::NoRestriction) // for CoarseLODs, treat 'no restriction' as keeping the budget constant
+			{
+				// Note: after the first coarse LOD we set the tri budget with respect to the previous budget, not the previous actual tri count,
+				// because we haven't yet computed the coarse meshes so do not have the tri counts
+				int32 PrevBudget = TriBudget[PrevLODLevel];
+				if (LODLevel == FirstVoxWrappedIndex || PrevBudget == -1)
+				{
+					PrevBudget = GetLODTriCount(PrevLODLevel);
+				}
+				TargetTriCount = PrevBudget;
+				if (Opts.TriangleBudgetOptions.Method == IGeometryProcessing_CombineMeshInstances::ETriangleBudgetMethod::UsePercentageOfPreviousLOD)
+				{
+					TargetTriCount *= Opts.TriangleBudgetOptions.LODReductionPercentage;
+				}
+			}
+			TriBudget[LODLevel] = TargetTriCount;
 
-			if (!Opts.bEnableBudgetStrategy_PartLODPromotion || Opts.TriangleBudgetOptions.Method == IGeometryProcessing_CombineMeshInstances::ETriangleBudgetMethod::NoRestriction)
+			// For voxwrapped LOD levels, reference the last source LOD
+			for (int32 PartIdx = 0; PartIdx < NumParts; ++PartIdx)
 			{
-				continue;
+				const IGeometryProcessing_CombineMeshInstances::FSinglePartMeshSet& PartMeshes = *Assembly.Parts[PartIdx]->PrecomputedMeshes;
+				PartSourcesPerLOD[LODLevel * NumParts + PartIdx] = FMath::Clamp(Opts.PreferredLOD, 0, PartMeshes.Source.Num() - 1);
 			}
+		}
 
+		if (bApplyLODBudget)
+		{
 			int32 TargetTriCount = Opts.TriangleBudgetOptions.TriangleBudget;
 			if (Opts.TriangleBudgetOptions.Method == IGeometryProcessing_CombineMeshInstances::ETriangleBudgetMethod::UsePercentageOfPreviousLOD)
 			{
@@ -4062,7 +4093,7 @@ void BuildCombinedMeshFromPrecomputedMeshes(
 					bool bCanPromoteWithinCategory, bCanPromoteAtAll;
 					int32 WithinCatIdx;
 					double PromotedAvgError = 0;
-					int32 Category = SourceIndexToMeshCategory(PartIdx, SourceIdx, WithinCatIdx, bCanPromoteWithinCategory, bCanPromoteAtAll, PromotedAvgError);
+					int32 Category = static_cast<int32>(SourceIndexToMeshCategory(PartIdx, SourceIdx, WithinCatIdx, bCanPromoteWithinCategory, bCanPromoteAtAll, PromotedAvgError));
 					if (!bCanPromoteAtAll)
 					{
 						continue;
@@ -4078,11 +4109,12 @@ void BuildCombinedMeshFromPrecomputedMeshes(
 					}
 					if (bCanPromoteWithinCategory)
 					{
-						// For source meshes (category 0), promote by LOD index primarily (note that source meshes don't have average error stats)
+						// For source meshes, promote by LOD index primarily (note that source meshes don't have average error stats)
 						// For simplified and approximated meshes, promote primarily by average error
-						if (Category < BestPromotableCategory || 
-						   (Category == 0 && WithinCatIdx < BestPromotableWithinCatIdx) ||
-						   (Category == BestPromotableCategory && Category > 0 && PromotedAvgError < BestWithinCatAvgError))
+						constexpr int32 CopiedLODCategory = static_cast<int32>(ECombinedLODType::Copied);
+						if (Category < BestPromotableCategory ||
+							(Category == CopiedLODCategory && WithinCatIdx < BestPromotableWithinCatIdx) ||
+							(Category == BestPromotableCategory && Category > CopiedLODCategory && PromotedAvgError < BestWithinCatAvgError))
 						{
 							BestPromotableCategory = Category;
 							BestPromotableWithinCatIdx = WithinCatIdx;
@@ -4106,34 +4138,6 @@ void BuildCombinedMeshFromPrecomputedMeshes(
 				CurTriCount = GetLODTriCount(LODLevel);
 			}
 		}
-		else if (LODType == ECombinedLODType::VoxWrapped)
-		{
-			int32 TargetTriCount = Opts.TriangleBudgetOptions.TriangleBudget;
-			if (Opts.TriangleBudgetOptions.Method == IGeometryProcessing_CombineMeshInstances::ETriangleBudgetMethod::UsePercentageOfPreviousLOD
-				|| Opts.TriangleBudgetOptions.Method == IGeometryProcessing_CombineMeshInstances::ETriangleBudgetMethod::NoRestriction) // for CoarseLODs, treat 'no restriction' as keeping the budget constant
-			{
-				// Note: after the first coarse LOD we set the tri budget with respect to the previous budget, not the previous actual tri count,
-				// because we haven't yet computed the coarse meshes so do not have the tri counts
-				int32 PrevBudget = TriBudget[PrevLODLevel];
-				if (LODLevel == FirstVoxWrappedIndex || PrevBudget == -1)
-				{
-					PrevBudget = GetLODTriCount(PrevLODLevel);
-				}
-				TargetTriCount = PrevBudget;
-				if (Opts.TriangleBudgetOptions.Method == IGeometryProcessing_CombineMeshInstances::ETriangleBudgetMethod::UsePercentageOfPreviousLOD)
-				{
-					TargetTriCount *= Opts.TriangleBudgetOptions.LODReductionPercentage;
-				}
-			}
-			TriBudget[LODLevel] = TargetTriCount;
-
-			// For voxwrapped LOD levels, reference the last source LOD
-			for (int32 PartIdx = 0; PartIdx < NumParts; ++PartIdx)
-			{
-				const IGeometryProcessing_CombineMeshInstances::FSinglePartMeshSet& PartMeshes = *Assembly.Parts[PartIdx]->PrecomputedMeshes;
-				PartSourcesPerLOD[LODLevel * NumParts + PartIdx] = PartMeshes.Source.Num() - 1;
-			}
-		}
 	}
 
 
@@ -4150,11 +4154,6 @@ void BuildCombinedMeshFromPrecomputedMeshes(
 
 		for (int32 LODLevel = 0; LODLevel < NumLODs; ++LODLevel)
 		{
-			int32 SourceIndex = PartSourcesPerLOD[NumParts * LODLevel + SetIndex];
-			const FDynamicMesh3* UseAppendMesh = SourceIndexToMesh(SetIndex, SourceIndex);
-
-			ECombinedLODType LevelLODType = LODTypes[LODLevel];
-
 			FCombinedMeshLOD& CombinedMeshLODData = MeshLODs[LODLevel];
 
 			for (int32 InstIdx = 0; InstIdx < Part->Instances.Num(); ++InstIdx)
@@ -4631,11 +4630,11 @@ static void ComputeSinglePartMeshSet_Helper(
 
 	int32 NumSourceLODs = Result.Source.Num();
 
-	const FDynamicMesh3* SimplificationSourceMesh = Options.SimplificationSourceLOD < Result.Source.Num() ?
-		&Result.Source[Options.SimplificationSourceLOD] : &Result.Source.Last();
+	const FDynamicMesh3* SimplificationSourceMesh = &Result.Source[FMath::Clamp(Options.SimplificationSourceLOD, 0, Result.Source.Num() - 1)];
 	FDynamicMeshAABBTree3 SimplificationSourceMeshSpatial(SimplificationSourceMesh, true);
 	Result.Simplified.Reserve(NumSimplifiedLODs);
 	double UseSimplifyTolerance = Options.SimplifyBaseTolerance;
+	int32 LowestSimpliedTriCount = Result.Source.Last().TriangleCount();
 	for (int32 SimplifyIdx = 0; SimplifyIdx < NumSimplifiedLODs; ++SimplifyIdx)
 	{
 		TUniquePtr<FDynamicMesh3> ToAdd = MakeUnique<FDynamicMesh3>(*SimplificationSourceMesh);
@@ -4647,37 +4646,51 @@ static void ComputeSinglePartMeshSet_Helper(
 		// Add the simplification if it's simpler than the lowest-tri source LOD
 		if (ToAdd->TriangleCount() < Result.Source.Last().TriangleCount())
 		{
+			LowestSimpliedTriCount = FMath::Min(LowestSimpliedTriCount, ToAdd->TriangleCount());
 			Result.Simplified.Add(ToAdd.Release());
 		}
 	}
-		
-	const FDynamicMesh3* ApproximationSourceMesh = Options.ApproximationSourceLOD < Result.Source.Num() ?
-		&Result.Source[Options.ApproximationSourceLOD] : &Result.Simplified[FMath::Clamp(Options.ApproximationSourceLOD - Result.Source.Num(), 0, Result.Simplified.Num())];
+	
+	const int32 ClampedApproximatedSourceLOD = FMath::Clamp(Options.ApproximationSourceLOD, 0, Options.ApproximationSourceLOD);
+	const FDynamicMesh3* ApproximationSourceMesh = &Result.Source.Last();
+	if (Result.Source.IsValidIndex(ClampedApproximatedSourceLOD))
+	{
+		ApproximationSourceMesh = &Result.Source[ClampedApproximatedSourceLOD];
+	}
+	else if (Result.Simplified.IsValidIndex(ClampedApproximatedSourceLOD - Result.Source.Num()))
+	{
+		ApproximationSourceMesh = &Result.Simplified[ClampedApproximatedSourceLOD - Result.Source.Num()];
+	}
 
 	Result.Approximated.Reserve(NumApproxLODs);
 	TArray<EApproximatePartMethod> SelectedMethodID; SelectedMethodID.Reserve(NumApproxLODs);		// useful for debugging
 	for (int32 ApproxIdx = 0; ApproxIdx < NumApproxLODs; ++ApproxIdx)
 	{
 		double UseTriCost = Options.ApproximationTriCosts[ApproxIdx];
-		Result.Approximated.Add(new FDynamicMesh3(*ApproximationSourceMesh));
-		FDynamicMesh3& ToApprox = Result.Approximated.Last();
+		TUniquePtr<FDynamicMesh3> ToApprox = MakeUnique<FDynamicMesh3>(*ApproximationSourceMesh);
 		EApproximatePartMethod UsedMethod;
 		SelectBestFittingMeshApproximation(*ApproximationSourceMesh, SimplificationSourceMeshSpatial,
-			PartOptions.ApproximationConstraint, ToApprox, UsedMethod,
+			PartOptions.ApproximationConstraint, *ToApprox, UsedMethod,
 			Options.SimplifyBaseTolerance, UseTriCost, Options.MaxAllowableApproximationDeviation);
 			
 		// update enabled attribs (is this good?)
-		ToApprox.EnableMatchingAttributes(*ApproximationSourceMesh);
+		ToApprox->EnableMatchingAttributes(*ApproximationSourceMesh);
 
 		// recompute normals
-		FMeshNormals::InitializeOverlayTopologyFromOpeningAngle(&ToApprox, ToApprox.Attributes()->PrimaryNormals(), AngleThresholdDeg);
-		FMeshNormals::QuickRecomputeOverlayNormals(ToApprox);
+		FMeshNormals::InitializeOverlayTopologyFromOpeningAngle(ToApprox.Get(), ToApprox->Attributes()->PrimaryNormals(), AngleThresholdDeg);
+		FMeshNormals::QuickRecomputeOverlayNormals(*ToApprox);
 
 		// stop making approximations if we are already down to a box
-		if (ToApprox.TriangleCount() <= 12)
+		if (ToApprox->TriangleCount() <= 12)
 		{
 			break;
 		}
+
+		// Add the approximation if it's simpler than the lowest tri simplified LOD.
+		if (ToApprox->TriangleCount() < LowestSimpliedTriCount)
+		{
+			Result.Approximated.Add(ToApprox.Release());
+		}
 	}
 
 	// Compute error metrics for simplified and approximated meshes
@@ -4938,11 +4951,13 @@ void FCombineMeshInstancesImpl::ComputePartMeshSets(
 		{
 			TSharedPtr<FSinglePartMeshSet>& NewMeshSet = ResultMeshSets.Add_GetRef(MakeShared<FSinglePartMeshSet>());
 			EApproximationType ApproximationConstraint = EApproximationType::NoConstraint;
+			bool bPreservePartUVs = false;
 			if (SourceInstanceList.InstanceGroupDatas.IsValidIndex(SourceMeshInstance.GroupDataIndex))
 			{
 				ApproximationConstraint = SourceInstanceList.InstanceGroupDatas[SourceMeshInstance.GroupDataIndex].ApproximationConstraint;
+				bPreservePartUVs = SourceInstanceList.InstanceGroupDatas[SourceMeshInstance.GroupDataIndex].bPreserveUVs;
 			}
-			IGeometryProcessing_CombineMeshInstances::FComputePartMeshesSinglePartOptions PartOptions(ApproximationConstraint);
+			IGeometryProcessing_CombineMeshInstances::FComputePartMeshesSinglePartOptions PartOptions(ApproximationConstraint, bPreservePartUVs);
 			ComputeSinglePartMeshSet(MeshLODSet->ReferencedMeshLODs, Options, PartOptions, *NewMeshSet);
 			MeshLODSetToPartMap.Add(MeshLODSet, NewMeshSet);
 
diff --git a/Engine/Source/Developer/GeometryProcessingInterfaces/Public/GeometryProcessingInterfaces/CombineMeshInstances.h b/Engine/Source/Developer/GeometryProcessingInterfaces/Public/GeometryProcessingInterfaces/CombineMeshInstances.h
--- a/Engine/Source/Developer/GeometryProcessingInterfaces/Public/GeometryProcessingInterfaces/CombineMeshInstances.h
+++ b/Engine/Source/Developer/GeometryProcessingInterfaces/Public/GeometryProcessingInterfaces/CombineMeshInstances.h
@@ -374,7 +374,7 @@ public:
 		float LODReductionPercentage = 0.75f;
 	};
 
-	enum class EDecorationHandling : uint8
+	enum class EDecorativePartHandling : uint8
 	{
 		Keep,
 		Approximate,
@@ -391,7 +391,7 @@ public:
 		EMeshOptimizationMethod OptimizationMethod = EMeshOptimizationMethod::None;
 
 		/** How decoration should be treated for this LOD (e.g., kept, approximated, or removed). */
-		EDecorationHandling Decorations = EDecorationHandling::Keep;
+		EDecorativePartHandling DecorativePartHandling = EDecorativePartHandling::Keep;
 
 		/** Tolerance to use for simplification of the LOD (if simplification is applied) */
 		double SimplificationTolerance = 1.0;
@@ -416,6 +416,9 @@ public:
 
 		// Whether a part promotion strategy should be used to attempt to hit triangle budgets
 		bool bEnableBudgetStrategy_PartLODPromotion = false;
+		// Whether a part demotion strategy should be used to attempt to max out the triangle budget.
+		// In other words, when enabled, we try to keep the higher LOD geometry for as long as possible.
+		bool bEnableBudgetStrategy_PartLODDemotion = false;
 		// If promoting parts to worse approximations to fit a triangle budget, scale the triangle budget by this amount (to account for potential reductions from e.g. hidden triangle removal, that will happen after part meshes are chosen and merged)
 		double PartLODPromotionBudgetMultiplier = 2.0;
 
@@ -515,14 +518,14 @@ public:
 				Options.RemoveHiddenSamplingDensity = InOptions.RemoveHiddenSamplingDensity;
 				Options.bMergeCoplanarFaces = InOptions.bMergeCoplanarFaces;
 				Options.bPlanarPolygonRetriangulation = InOptions.bMergeCoplanarFaces && (InOptions.PlanarPolygonRetriangulationStartLOD > -1 && LODIndex >= InOptions.PlanarPolygonRetriangulationStartLOD);
-				Options.Decorations = EDecorationHandling::Keep;
+				Options.DecorativePartHandling = EDecorativePartHandling::Keep;
 				if (LODIndex >= InOptions.FilterDecorativePartsLODLevel - InOptions.ApproximateDecorativePartLODs)
 				{
-					Options.Decorations = EDecorationHandling::Approximate;
+					Options.DecorativePartHandling = EDecorativePartHandling::Approximate;
 				}
 				else if (LODIndex >= InOptions.FilterDecorativePartsLODLevel)
 				{
-					Options.Decorations = EDecorationHandling::Remove;
+					Options.DecorativePartHandling = EDecorativePartHandling::Remove;
 				}
 				Options.SimplificationTolerance = SimplificationTolerance;
 				if (Options.OptimizationMethod != EMeshOptimizationMethod::None)
