[GeometryProcessing] Moving the CombineMeshInstances geometry utility into the internal codebase.

This Geometry processing feature was developped for internal use and it's current location prevent us from making any changes to it in 5.5. The move will allow us to properly make maintain this feature.

diff --git a/Engine/Plugins/Experimental/MeshModelingToolsetExp/Source/GeometryProcessingAdapters/Private/CombineMeshInstancesImpl.cpp b/Engine/Plugins/Experimental/MeshModelingToolsetExp/Source/GeometryProcessingAdapters/Private/CombineMeshInstancesImpl.cpp
deleted file mode 100644
diff --git a/Engine/Plugins/Experimental/MeshModelingToolsetExp/Source/GeometryProcessingAdapters/Private/GeometryProcessingAdaptersModule.cpp b/Engine/Plugins/Experimental/MeshModelingToolsetExp/Source/GeometryProcessingAdapters/Private/GeometryProcessingAdaptersModule.cpp
--- a/Engine/Plugins/Experimental/MeshModelingToolsetExp/Source/GeometryProcessingAdapters/Private/GeometryProcessingAdaptersModule.cpp
+++ b/Engine/Plugins/Experimental/MeshModelingToolsetExp/Source/GeometryProcessingAdapters/Private/GeometryProcessingAdaptersModule.cpp
@@ -4,7 +4,6 @@
 #include "Features/IModularFeatures.h"
 
 #include "GeometryProcessing/ApproximateActorsImpl.h"
-#include "GeometryProcessing/CombineMeshInstancesImpl.h"
 #include "GeometryProcessing/MeshAutoUVImpl.h"
 
 #define LOCTEXT_NAMESPACE "FGeometryProcessingAdaptersModule"
@@ -20,12 +19,6 @@ void FGeometryProcessingAdaptersModule::StartupModule()
 	{
 		IModularFeatures::Get().RegisterModularFeature(IGeometryProcessing_ApproximateActors::GetModularFeatureName(), ApproximateActors.Get());
 	}
-	
-	CombineMeshInstances = MakeShared<FCombineMeshInstancesImpl>();
-	if (CombineMeshInstances.IsValid())
-	{
-		IModularFeatures::Get().RegisterModularFeature(IGeometryProcessing_CombineMeshInstances::GetModularFeatureName(), CombineMeshInstances.Get());
-	}
 
 	MeshAutoUV = MakeShared<FMeshAutoUVImpl>();
 	if (MeshAutoUV.IsValid())
@@ -44,12 +37,6 @@ void FGeometryProcessingAdaptersModule::ShutdownModule()
 		IModularFeatures::Get().UnregisterModularFeature(IGeometryProcessing_ApproximateActors::GetModularFeatureName(), ApproximateActors.Get());
 		ApproximateActors = nullptr;
 	}
-	
-	if ( CombineMeshInstances.IsValid() )
-	{
-		IModularFeatures::Get().UnregisterModularFeature(IGeometryProcessing_CombineMeshInstances::GetModularFeatureName(), CombineMeshInstances.Get());
-		CombineMeshInstances = nullptr;
-	}
 
 	if (MeshAutoUV.IsValid())
 	{
diff --git a/Engine/Plugins/Experimental/MeshModelingToolsetExp/Source/GeometryProcessingAdapters/Public/GeometryProcessing/CombineMeshInstancesImpl.h b/Engine/Plugins/Experimental/MeshModelingToolsetExp/Source/GeometryProcessingAdapters/Public/GeometryProcessing/CombineMeshInstancesImpl.h
deleted file mode 100644
diff --git a/Engine/Plugins/Experimental/MeshModelingToolsetExp/Source/GeometryProcessingAdapters/Public/GeometryProcessingAdaptersModule.h b/Engine/Plugins/Experimental/MeshModelingToolsetExp/Source/GeometryProcessingAdapters/Public/GeometryProcessingAdaptersModule.h
--- a/Engine/Plugins/Experimental/MeshModelingToolsetExp/Source/GeometryProcessingAdapters/Public/GeometryProcessingAdaptersModule.h
+++ b/Engine/Plugins/Experimental/MeshModelingToolsetExp/Source/GeometryProcessingAdapters/Public/GeometryProcessingAdaptersModule.h
@@ -8,7 +8,6 @@
 namespace UE::Geometry
 {
 	class FApproximateActorsImpl;
-	class FCombineMeshInstancesImpl;
 	class FMeshAutoUVImpl;
 }
 
@@ -23,6 +22,5 @@ public:
 	
 protected:
 	TSharedPtr<UE::Geometry::FApproximateActorsImpl> ApproximateActors;
-	TSharedPtr<UE::Geometry::FCombineMeshInstancesImpl> CombineMeshInstances;
 	TSharedPtr<UE::Geometry::FMeshAutoUVImpl> MeshAutoUV;
 };
diff --git a/Engine/Source/Developer/GeometryProcessingInterfaces/Private/GeometryProcessingInterfacesModule.cpp b/Engine/Source/Developer/GeometryProcessingInterfaces/Private/GeometryProcessingInterfacesModule.cpp
--- a/Engine/Source/Developer/GeometryProcessingInterfaces/Private/GeometryProcessingInterfacesModule.cpp
+++ b/Engine/Source/Developer/GeometryProcessingInterfaces/Private/GeometryProcessingInterfacesModule.cpp
@@ -6,7 +6,6 @@
 #include "Modules/ModuleManager.h"
 
 #include "GeometryProcessingInterfaces/ApproximateActors.h"
-#include "GeometryProcessingInterfaces/CombineMeshInstances.h"
 #include "GeometryProcessingInterfaces/MeshAutoUV.h"
 
 
@@ -45,14 +44,6 @@ IGeometryProcessing_ApproximateActors* FGeometryProcessingInterfacesModule::GetA
 	}
 	return ApproximateActors;
 }
-IGeometryProcessing_CombineMeshInstances* FGeometryProcessingInterfacesModule::GetCombineMeshInstancesImplementation()
-{
-	if (CombineMeshInstances == nullptr)
-	{
-		CombineMeshInstances = GetModularFeatureImplementation<IGeometryProcessing_CombineMeshInstances>();
-	}
-	return CombineMeshInstances;
-}
 
 IGeometryProcessing_MeshAutoUV* FGeometryProcessingInterfacesModule::GetMeshAutoUVImplementation()
 {
diff --git a/Engine/Source/Developer/GeometryProcessingInterfaces/Private/GeometryProcessingInterfacesModule.h b/Engine/Source/Developer/GeometryProcessingInterfaces/Private/GeometryProcessingInterfacesModule.h
--- a/Engine/Source/Developer/GeometryProcessingInterfaces/Private/GeometryProcessingInterfacesModule.h
+++ b/Engine/Source/Developer/GeometryProcessingInterfaces/Private/GeometryProcessingInterfacesModule.h
@@ -27,11 +27,6 @@ public:
 	 */
 	virtual IGeometryProcessing_ApproximateActors* GetApproximateActorsImplementation() override;
 
-	/**
-	 * @return implementation of IGeometryProcessing_CombineMeshInstances, if available, or nullptr (result is cached internally)
-	 */
-	virtual IGeometryProcessing_CombineMeshInstances* GetCombineMeshInstancesImplementation() override;
-
 	/**
 	 * @return implementation of IGeometryProcessing_MeshAutoUV, if available, or nullptr (result is cached internally)
 	 */
diff --git a/Engine/Source/Developer/GeometryProcessingInterfaces/Public/GeometryProcessingInterfaces/CombineMeshInstances.h b/Engine/Source/Developer/GeometryProcessingInterfaces/Public/GeometryProcessingInterfaces/CombineMeshInstances.h
deleted file mode 100644
diff --git a/Engine/Source/Developer/GeometryProcessingInterfaces/Public/IGeometryProcessingInterfacesModule.h b/Engine/Source/Developer/GeometryProcessingInterfaces/Public/IGeometryProcessingInterfacesModule.h
--- a/Engine/Source/Developer/GeometryProcessingInterfaces/Public/IGeometryProcessingInterfacesModule.h
+++ b/Engine/Source/Developer/GeometryProcessingInterfaces/Public/IGeometryProcessingInterfacesModule.h
@@ -30,11 +30,6 @@ public:
 	 */
 	virtual IGeometryProcessing_ApproximateActors* GetApproximateActorsImplementation() = 0;
 
-	/**
-	 * @return implementation of IGeometryProcessing_CombineMeshInstances, if available, or nullptr (result is cached internally)
-	 */
-	virtual IGeometryProcessing_CombineMeshInstances* GetCombineMeshInstancesImplementation() = 0;
-
 	/**
 	 * @return implementation of IGeometryProcessing_MeshAutoUV, if available, or nullptr (result is cached internally)
 	 */
