MeshModelingTools: Require a hit triangle if requiring connectivity in sculpt tool.

diff --git a/Engine/Plugins/Experimental/MeshModelingToolsetExp/Source/MeshModelingToolsExp/Private/MeshVertexSculptTool.cpp b/Engine/Plugins/Experimental/MeshModelingToolsetExp/Source/MeshModelingToolsExp/Private/MeshVertexSculptTool.cpp
--- a/Engine/Plugins/Experimental/MeshModelingToolsetExp/Source/MeshModelingToolsExp/Private/MeshVertexSculptTool.cpp
+++ b/Engine/Plugins/Experimental/MeshModelingToolsetExp/Source/MeshModelingToolsExp/Private/MeshVertexSculptTool.cpp
@@ -556,6 +556,19 @@ void UMeshVertexSculptTool::OnCancelStroke()
 //  interest.
 void UMeshVertexSculptTool::UpdateRangeQueryTriBuffer(const FFrame3d& LocalFrame)
 {
+	if (RequireConnectivityToHitPointInStamp()
+		// It's possible for LastBrushTriangleID to be null if we started a stroke and brushed off the
+		//  edge of the mesh.
+		&& LastBrushTriangleID == IndexConstants::InvalidID)
+	{
+		// If we're requiring connectivity, and we didn't hit a triangle to start with, then we shouldn't
+		//  move any triangles.
+		VertexROI.Reset();
+		TriangleROIArray.Reset();
+		SymmetricVertexROI.Reset();
+		return;
+	}
+
 	FDynamicMesh3* Mesh = GetSculptMesh();
 	FVector3d BrushPos = LocalFrame.Origin;
 
