ModelingTools: Refactor the VertexSculptBrushToolBuilder back into the main VertexSculptToolBuilder.

diff --git a/Engine/Plugins/Experimental/MeshModelingToolsetExp/Source/MeshModelingToolsExp/Private/MeshVertexSculptTool.cpp b/Engine/Plugins/Experimental/MeshModelingToolsetExp/Source/MeshModelingToolsExp/Private/MeshVertexSculptTool.cpp
--- a/Engine/Plugins/Experimental/MeshModelingToolsetExp/Source/MeshModelingToolsExp/Private/MeshVertexSculptTool.cpp
+++ b/Engine/Plugins/Experimental/MeshModelingToolsetExp/Source/MeshModelingToolsExp/Private/MeshVertexSculptTool.cpp
@@ -64,6 +64,7 @@ UMeshSurfacePointTool* UMeshVertexSculptToolBuilder::CreateNewTool(const FToolBu
 {
 	UMeshVertexSculptTool* SculptTool = NewObject<UMeshVertexSculptTool>(SceneState.ToolManager);
 	SculptTool->SetWorld(SceneState.World);
+	SculptTool->SetDefaultPrimaryBrushID(DefaultPrimaryBrushID);
 	return SculptTool;
 }
 
@@ -79,14 +80,6 @@ const FToolTargetTypeRequirements& UMeshVertexSculptToolBuilder::GetTargetRequir
 	return VSculptTypeRequirements;
 }
 
-UMeshSurfacePointTool* UMeshVertexSculptBrushToolBuilder::CreateNewTool(const FToolBuilderState& SceneState) const
-{
-	UMeshVertexSculptTool* SculptTool = NewObject<UMeshVertexSculptTool>(SceneState.ToolManager);
-	SculptTool->SetWorld(SceneState.World);
-	SculptTool->SetDefaultPrimaryBrushID(DefaultPrimaryBrushID);
-	return SculptTool;
-}
-
 /*
  * internal Change classes
  */
@@ -223,10 +216,7 @@ void UMeshVertexSculptTool::Setup()
 	UMeshSculptToolBase::BrushProperties->bShowPerBrushProps = false;
 	UMeshSculptToolBase::BrushProperties->bShowFalloff = false;
 	SculptProperties->RestoreProperties(this, GetPropertyCacheIdentifier());
-	if (DefaultPrimaryBrushID >= 0)
-	{
-		SculptProperties->PrimaryBrushID = DefaultPrimaryBrushID;
-	}
+	
 	AddToolPropertySource(SculptProperties);
 	CalculateBrushRadius();
 
@@ -247,7 +237,10 @@ void UMeshVertexSculptTool::Setup()
 
 	RegisterBrushes();
 
-	ensure(SculptProperties->PrimaryBrushID < RegisteredPrimaryBrushTypes.Num());
+	if (DefaultPrimaryBrushID >= 0 && ensure(SculptProperties->PrimaryBrushID < RegisteredPrimaryBrushTypes.Num()))
+	{
+		SculptProperties->PrimaryBrushID = DefaultPrimaryBrushID;
+	}
 
 	// falloffs
 	RegisterStandardFalloffTypes();
diff --git a/Engine/Plugins/Experimental/MeshModelingToolsetExp/Source/MeshModelingToolsExp/Public/MeshVertexSculptTool.h b/Engine/Plugins/Experimental/MeshModelingToolsetExp/Source/MeshModelingToolsExp/Public/MeshVertexSculptTool.h
--- a/Engine/Plugins/Experimental/MeshModelingToolsetExp/Source/MeshModelingToolsExp/Public/MeshVertexSculptTool.h
+++ b/Engine/Plugins/Experimental/MeshModelingToolsetExp/Source/MeshModelingToolsExp/Public/MeshVertexSculptTool.h
@@ -37,6 +37,8 @@ class MESHMODELINGTOOLSEXP_API UMeshVertexSculptToolBuilder : public UMeshSurfac
 
 public:
 	virtual UMeshSurfacePointTool* CreateNewTool(const FToolBuilderState& SceneState) const override;
+
+	int32 DefaultPrimaryBrushID = -1;
 protected:
 	virtual const FToolTargetTypeRequirements& GetTargetRequirements() const override;
 private:
@@ -103,16 +105,6 @@ enum class EMeshVertexSculptBrushType : uint8
 
 };
 
-UCLASS()
-class MESHMODELINGTOOLSEXP_API UMeshVertexSculptBrushToolBuilder : public UMeshVertexSculptToolBuilder
-{
-	GENERATED_BODY()
-public:
-	virtual UMeshSurfacePointTool* CreateNewTool(const FToolBuilderState& SceneState) const override;
-	
-	int32 DefaultPrimaryBrushID = -1;
-};
-
 
 /** Brush Triangle Filter Type */
 UENUM()
