MeshModelingTools: Fix Draw Spline tool not working correctly if Spline Tangent Scale setting is set to 0. Add tooltips to Arrive Tangent and Leave Tangent in spline details panel.

diff --git a/Engine/Plugins/Experimental/MeshModelingToolsetExp/Source/MeshModelingToolsEditorOnlyExp/Private/DrawSplineTool.cpp b/Engine/Plugins/Experimental/MeshModelingToolsetExp/Source/MeshModelingToolsEditorOnlyExp/Private/DrawSplineTool.cpp
--- a/Engine/Plugins/Experimental/MeshModelingToolsetExp/Source/MeshModelingToolsEditorOnlyExp/Private/DrawSplineTool.cpp
+++ b/Engine/Plugins/Experimental/MeshModelingToolsetExp/Source/MeshModelingToolsEditorOnlyExp/Private/DrawSplineTool.cpp
@@ -165,9 +165,14 @@ namespace DrawSplineToolLocals
 	};
 
 	// Gives the scale used for tangent visualization (and which therefore needs to be used in raycasting the handles)
-	float GetTangentScale()
+	float GetNonzeroTangentScale()
 	{
-		return GetDefault<ULevelEditorViewportSettings>()->SplineTangentScale;
+		// We'll use the setting that gets used when editing splines in the viewport, but if it is set to 0 (because
+		//  the user presumably doesn't want to see tangent handles in viewport), we'll ignore it and use the default
+		//  value of 0.5, because we need to be able to draw a tangent for our "tangent drag" approach.
+		// The abs here isn't actually necessary since the setting is clamped in editor, but it looks more correct.
+		float Scale = FMath::Abs(GetDefault<ULevelEditorViewportSettings>()->SplineTangentScale);
+		return Scale > 0.0f ? Scale : 0.5f;
 	}
 
 	// Might be useful to have in SplineUtil, but uncertain what the API should be (should it be part of
@@ -182,7 +187,7 @@ namespace DrawSplineToolLocals
 
 		FPrimitiveDrawInterface* PDI = RenderAPI.GetPrimitiveDrawInterface();
 
-		const float TangentScale = GetDefault<ULevelEditorViewportSettings>()->SplineTangentScale;
+		const float TangentScale = GetNonzeroTangentScale();
 		const float TangentHandleSize = 8.0f + GetDefault<ULevelEditorViewportSettings>()->SplineTangentHandleSizeAdjustment;
 
 		const FVector Location = SplineComp.GetLocationAtSplinePoint(PointIndex, ESplineCoordinateSpace::World);
@@ -940,7 +945,7 @@ void UDrawSplineTool::OnClickDrag(const FInputDeviceRay& DragPos)
 		{
 			// Set the tangent
 			FVector3d LastPoint = WorkingSpline->GetLocationAtSplinePoint(NumSplinePoints - 1, ESplineCoordinateSpace::World);
-			FVector3d Tangent = (HitLocation - LastPoint) / GetTangentScale();
+			FVector3d Tangent = (HitLocation - LastPoint) / GetNonzeroTangentScale();
 			WorkingSpline->SetTangentAtSplinePoint(NumSplinePoints - 1, Tangent, ESplineCoordinateSpace::World, true);
 			bDrawTangentForLastPoint = true;
 			break;
diff --git a/Engine/Source/Editor/DetailCustomizations/Private/SplineComponentDetails.cpp b/Engine/Source/Editor/DetailCustomizations/Private/SplineComponentDetails.cpp
--- a/Engine/Source/Editor/DetailCustomizations/Private/SplineComponentDetails.cpp
+++ b/Engine/Source/Editor/DetailCustomizations/Private/SplineComponentDetails.cpp
@@ -676,6 +676,9 @@ void FSplinePointDetails::GenerateChildContent(IDetailChildrenBuilder& ChildrenB
 				SNew(STextBlock)
 				.Text(LOCTEXT("ArriveTangent", "Arrive Tangent"))
 				.Font(IDetailLayoutBuilder::GetDetailFont())
+				.ToolTipText(LOCTEXT("ArriveTangent_Tooltip", "Incoming tangent. Note that the size shown in viewport "
+					"is controlled by Spline Tangent Scale in editor preferences (and hidden if 0). Only allowed to "
+					"differ from Leave Tangent if Allow Discontinuous Spline is true."))
 			]
 		.ValueContent()
 			.MinDesiredWidth(375.0f)
@@ -711,6 +714,8 @@ void FSplinePointDetails::GenerateChildContent(IDetailChildrenBuilder& ChildrenB
 				SNew(STextBlock)
 				.Text(LOCTEXT("LeaveTangent", "Leave Tangent"))
 				.Font(IDetailLayoutBuilder::GetDetailFont())
+				.ToolTipText(LOCTEXT("LeaveTangent_Tooltip", "Outgoing tangent. Note that the size shown in viewport "
+					"is controlled by Spline Tangent Scale in editor preferences (and hidden if 0)."))
 			]
 		.ValueContent()
 			.MinDesiredWidth(375.0f)
