ModelingTools: Add quick edit tag to various tool properties.

diff --git a/Engine/Plugins/Experimental/MeshModelingToolsetExp/Source/MeshModelingToolsExp/Private/Sculpting/KelvinletBrushOp.h b/Engine/Plugins/Experimental/MeshModelingToolsetExp/Source/MeshModelingToolsExp/Private/Sculpting/KelvinletBrushOp.h
--- a/Engine/Plugins/Experimental/MeshModelingToolsetExp/Source/MeshModelingToolsExp/Private/Sculpting/KelvinletBrushOp.h
+++ b/Engine/Plugins/Experimental/MeshModelingToolsetExp/Source/MeshModelingToolsExp/Private/Sculpting/KelvinletBrushOp.h
@@ -194,7 +194,7 @@ class MESHMODELINGTOOLSEXP_API UScaleKelvinletBrushOpProps : public   UBaseKelvi
 
 public:
 	/** Strength of the Brush */
-	UPROPERTY(EditAnywhere, Category = KelvinScaleBrush, meta = (DisplayName = "Strength", UIMin = "0.0", UIMax = "10.", ClampMin = "0.0", ClampMax = "10."))
+	UPROPERTY(EditAnywhere, Category = KelvinScaleBrush, meta = (DisplayName = "Strength", UIMin = "0.0", UIMax = "10.", ClampMin = "0.0", ClampMax = "10.", ModelingQuickEdit))
 	float Strength = 0.5;
 
 	/** Amount of falloff to apply */
diff --git a/Engine/Plugins/Experimental/MeshModelingToolsetExp/Source/MeshModelingToolsExp/Public/DynamicMeshSculptTool.h b/Engine/Plugins/Experimental/MeshModelingToolsetExp/Source/MeshModelingToolsExp/Public/DynamicMeshSculptTool.h
--- a/Engine/Plugins/Experimental/MeshModelingToolsetExp/Source/MeshModelingToolsExp/Public/DynamicMeshSculptTool.h
+++ b/Engine/Plugins/Experimental/MeshModelingToolsetExp/Source/MeshModelingToolsExp/Public/DynamicMeshSculptTool.h
@@ -157,7 +157,7 @@ public:
 	EDynamicMeshSculptBrushType PrimaryBrushType = EDynamicMeshSculptBrushType::Move;
 
 	/** Strength of the Primary Brush */
-	UPROPERTY(EditAnywhere, Category = Sculpting, meta = (DisplayName = "Strength", UIMin = "0.0", UIMax = "1.0", ClampMin = "0.0", ClampMax = "1.0", EditCondition = "PrimaryBrushType != EDynamicMeshSculptBrushType::Pull"))
+	UPROPERTY(EditAnywhere, Category = Sculpting, meta = (DisplayName = "Strength", UIMin = "0.0", UIMax = "1.0", ClampMin = "0.0", ClampMax = "1.0", EditCondition = "PrimaryBrushType != EDynamicMeshSculptBrushType::Pull", ModelingQuickEdit))
 	float PrimaryBrushSpeed = 0.5;
 
 	/** If true, try to preserve the shape of the UV/3D mapping. This will limit Smoothing and Remeshing in some cases. */
diff --git a/Engine/Plugins/Experimental/MeshModelingToolsetExp/Source/MeshModelingToolsExp/Public/Polymodeling/ExtrudeMeshSelectionTool.h b/Engine/Plugins/Experimental/MeshModelingToolsetExp/Source/MeshModelingToolsExp/Public/Polymodeling/ExtrudeMeshSelectionTool.h
--- a/Engine/Plugins/Experimental/MeshModelingToolsetExp/Source/MeshModelingToolsExp/Public/Polymodeling/ExtrudeMeshSelectionTool.h
+++ b/Engine/Plugins/Experimental/MeshModelingToolsetExp/Source/MeshModelingToolsExp/Public/Polymodeling/ExtrudeMeshSelectionTool.h
@@ -55,11 +55,11 @@ class MESHMODELINGTOOLSEXP_API UExtrudeMeshSelectionToolProperties : public UInt
 
 public:
 	/** Control how the Extruded Area should be Transformed */
-	UPROPERTY(EditAnywhere, Category = "Extrude")
+	UPROPERTY(EditAnywhere, Category = "Extrude", meta = (ModelingQuickEdit))
 	EExtrudeMeshSelectionInteractionMode InputMode = EExtrudeMeshSelectionInteractionMode::Interactive;
 
 	/** The Extrusion Distance used in Fixed Input Mode*/
-	UPROPERTY(EditAnywhere, Category = "Extrude", meta = (DisplayName="Fixed Distance", UIMin = -250, UIMax = 250, EditCondition = "InputMode == EExtrudeMeshSelectionInteractionMode::Fixed"))
+	UPROPERTY(EditAnywhere, Category = "Extrude", meta = (DisplayName="Fixed Distance", UIMin = -250, UIMax = 250, EditCondition = "InputMode == EExtrudeMeshSelectionInteractionMode::Fixed", ModelingQuickEdit))
 	double ExtrudeDistance = 10.0;
 
 	/** Control how the Extruded Area should be deformed as part of the Extrusion */
diff --git a/Engine/Plugins/Experimental/MeshModelingToolsetExp/Source/MeshModelingToolsExp/Public/Polymodeling/OffsetMeshSelectionTool.h b/Engine/Plugins/Experimental/MeshModelingToolsetExp/Source/MeshModelingToolsExp/Public/Polymodeling/OffsetMeshSelectionTool.h
--- a/Engine/Plugins/Experimental/MeshModelingToolsetExp/Source/MeshModelingToolsExp/Public/Polymodeling/OffsetMeshSelectionTool.h
+++ b/Engine/Plugins/Experimental/MeshModelingToolsetExp/Source/MeshModelingToolsExp/Public/Polymodeling/OffsetMeshSelectionTool.h
@@ -56,7 +56,7 @@ public:
 	//EOffsetMeshSelectionInteractionMode InputMode = EOffsetMeshSelectionInteractionMode::Interactive;
 
 	/** The Extrusion Distance used in Fixed Input Mode*/
-	UPROPERTY(EditAnywhere, Category = "Offset", meta = (DisplayName="Distance", UIMin = -250, UIMax = 250))
+	UPROPERTY(EditAnywhere, Category = "Offset", meta = (DisplayName="Distance", UIMin = -250, UIMax = 250, ModelingQuickEdit))
 	double OffsetDistance = 10.0;
 
 	/** Control how the Offset Area should be displaced */
diff --git a/Engine/Plugins/Experimental/MeshModelingToolsetExp/Source/MeshModelingToolsExp/Public/Sculpting/MeshInflateBrushOps.h b/Engine/Plugins/Experimental/MeshModelingToolsetExp/Source/MeshModelingToolsExp/Public/Sculpting/MeshInflateBrushOps.h
--- a/Engine/Plugins/Experimental/MeshModelingToolsetExp/Source/MeshModelingToolsExp/Public/Sculpting/MeshInflateBrushOps.h
+++ b/Engine/Plugins/Experimental/MeshModelingToolsetExp/Source/MeshModelingToolsExp/Public/Sculpting/MeshInflateBrushOps.h
@@ -16,7 +16,7 @@ class MESHMODELINGTOOLSEXP_API UInflateBrushOpProps : public UMeshSculptBrushOpP
 	GENERATED_BODY()
 public:
 	/** Strength of the Brush */
-	UPROPERTY(EditAnywhere, Category = InflateBrush, meta = (DisplayName = "Strength", UIMin = "0.0", UIMax = "1.0", ClampMin = "0.0", ClampMax = "1.0"))
+	UPROPERTY(EditAnywhere, Category = InflateBrush, meta = (DisplayName = "Strength", UIMin = "0.0", UIMax = "1.0", ClampMin = "0.0", ClampMax = "1.0", ModelingQuickEdit))
 	float Strength = 0.5;
 
 	/** Amount of falloff to apply */
diff --git a/Engine/Plugins/Experimental/MeshModelingToolsetExp/Source/MeshModelingToolsExp/Public/Sculpting/MeshMoveBrushOps.h b/Engine/Plugins/Experimental/MeshModelingToolsetExp/Source/MeshModelingToolsExp/Public/Sculpting/MeshMoveBrushOps.h
--- a/Engine/Plugins/Experimental/MeshModelingToolsetExp/Source/MeshModelingToolsExp/Public/Sculpting/MeshMoveBrushOps.h
+++ b/Engine/Plugins/Experimental/MeshModelingToolsetExp/Source/MeshModelingToolsExp/Public/Sculpting/MeshMoveBrushOps.h
@@ -18,7 +18,7 @@ class MESHMODELINGTOOLSEXP_API UMoveBrushOpProps : public UMeshSculptBrushOpProp
 	GENERATED_BODY()
 public:
 	/** Strength of the Brush */
-	UPROPERTY(EditAnywhere, Category = MoveBrush, meta = (DisplayName = "Strength", UIMin = "0.0", UIMax = "1.0", ClampMin = "0.0", ClampMax = "1.0"))
+	UPROPERTY(EditAnywhere, Category = MoveBrush, meta = (DisplayName = "Strength", UIMin = "0.0", UIMax = "1.0", ClampMin = "0.0", ClampMax = "1.0", ModelingQuickEdit))
 	float Strength = 1.0;
 
 	/** Amount of falloff to apply */
diff --git a/Engine/Plugins/Experimental/MeshModelingToolsetExp/Source/MeshModelingToolsExp/Public/Sculpting/MeshPinchBrushOps.h b/Engine/Plugins/Experimental/MeshModelingToolsetExp/Source/MeshModelingToolsExp/Public/Sculpting/MeshPinchBrushOps.h
--- a/Engine/Plugins/Experimental/MeshModelingToolsetExp/Source/MeshModelingToolsExp/Public/Sculpting/MeshPinchBrushOps.h
+++ b/Engine/Plugins/Experimental/MeshModelingToolsetExp/Source/MeshModelingToolsExp/Public/Sculpting/MeshPinchBrushOps.h
@@ -15,7 +15,7 @@ class MESHMODELINGTOOLSEXP_API UPinchBrushOpProps : public UMeshSculptBrushOpPro
 	GENERATED_BODY()
 public:
 	/** Strength of the Brush */
-	UPROPERTY(EditAnywhere, Category = PinchBrush, meta = (DisplayName = "Strength", UIMin = "0.0", UIMax = "1.0", ClampMin = "0.0", ClampMax = "1.0"))
+	UPROPERTY(EditAnywhere, Category = PinchBrush, meta = (DisplayName = "Strength", UIMin = "0.0", UIMax = "1.0", ClampMin = "0.0", ClampMax = "1.0", ModelingQuickEdit))
 	float Strength = 0.5;
 
 	/** Amount of falloff to apply */
diff --git a/Engine/Plugins/Experimental/MeshModelingToolsetExp/Source/MeshModelingToolsExp/Public/Sculpting/MeshPlaneBrushOps.h b/Engine/Plugins/Experimental/MeshModelingToolsetExp/Source/MeshModelingToolsExp/Public/Sculpting/MeshPlaneBrushOps.h
--- a/Engine/Plugins/Experimental/MeshModelingToolsetExp/Source/MeshModelingToolsExp/Public/Sculpting/MeshPlaneBrushOps.h
+++ b/Engine/Plugins/Experimental/MeshModelingToolsetExp/Source/MeshModelingToolsExp/Public/Sculpting/MeshPlaneBrushOps.h
@@ -25,7 +25,7 @@ class MESHMODELINGTOOLSEXP_API UPlaneBrushOpProps : public UBasePlaneBrushOpProp
 	GENERATED_BODY()
 public:
 	/** Strength of the Brush */
-	UPROPERTY(EditAnywhere, Category = PlaneBrush, meta = (DisplayName = "Strength", UIMin = "0.0", UIMax = "1.0", ClampMin = "0.0", ClampMax = "1.0"))
+	UPROPERTY(EditAnywhere, Category = PlaneBrush, meta = (DisplayName = "Strength", UIMin = "0.0", UIMax = "1.0", ClampMin = "0.0", ClampMax = "1.0", ModelingQuickEdit))
 	float Strength = 0.5;
 
 	/** Amount of falloff to apply */
@@ -54,7 +54,7 @@ class MESHMODELINGTOOLSEXP_API UViewAlignedPlaneBrushOpProps : public UBasePlane
 	GENERATED_BODY()
 public:
 	/** Strength of the Brush */
-	UPROPERTY(EditAnywhere, Category = ViewPlaneBrush, meta = (DisplayName = "Strength", UIMin = "0.0", UIMax = "1.0", ClampMin = "0.0", ClampMax = "1.0"))
+	UPROPERTY(EditAnywhere, Category = ViewPlaneBrush, meta = (DisplayName = "Strength", UIMin = "0.0", UIMax = "1.0", ClampMin = "0.0", ClampMax = "1.0", ModelingQuickEdit))
 	float Strength = 0.5;
 
 	/** Amount of falloff to apply */
@@ -83,7 +83,7 @@ class MESHMODELINGTOOLSEXP_API UFixedPlaneBrushOpProps : public UBasePlaneBrushO
 	GENERATED_BODY()
 public:
 	/** Strength of the Brush */
-	UPROPERTY(EditAnywhere, Category = FixedPlaneBrush, meta = (DisplayName = "Strength", UIMin = "0.0", UIMax = "1.0", ClampMin = "0.0", ClampMax = "1.0"))
+	UPROPERTY(EditAnywhere, Category = FixedPlaneBrush, meta = (DisplayName = "Strength", UIMin = "0.0", UIMax = "1.0", ClampMin = "0.0", ClampMax = "1.0", ModelingQuickEdit))
 	float Strength = 0.5;
 
 	/** Amount of falloff to apply */
diff --git a/Engine/Plugins/Experimental/MeshModelingToolsetExp/Source/MeshModelingToolsExp/Public/Sculpting/MeshSculptBrushOps.h b/Engine/Plugins/Experimental/MeshModelingToolsetExp/Source/MeshModelingToolsExp/Public/Sculpting/MeshSculptBrushOps.h
--- a/Engine/Plugins/Experimental/MeshModelingToolsetExp/Source/MeshModelingToolsExp/Public/Sculpting/MeshSculptBrushOps.h
+++ b/Engine/Plugins/Experimental/MeshModelingToolsetExp/Source/MeshModelingToolsExp/Public/Sculpting/MeshSculptBrushOps.h
@@ -16,7 +16,7 @@ class MESHMODELINGTOOLSEXP_API UStandardSculptBrushOpProps : public UMeshSculptB
 	GENERATED_BODY()
 public:
 	/** Strength of the Brush */
-	UPROPERTY(EditAnywhere, Category = SculptBrush, meta = (DisplayName = "Strength", UIMin = "0.0", UIMax = "1.0", ClampMin = "0.0", ClampMax = "1.0"))
+	UPROPERTY(EditAnywhere, Category = SculptBrush, meta = (DisplayName = "Strength", UIMin = "0.0", UIMax = "1.0", ClampMin = "0.0", ClampMax = "1.0", ModelingQuickEdit))
 	float Strength = 0.5;
 
 	/** Amount of falloff to apply */
@@ -97,7 +97,7 @@ class MESHMODELINGTOOLSEXP_API UViewAlignedSculptBrushOpProps : public UMeshScul
 	GENERATED_BODY()
 public:
 	/** Strength of the Brush */
-	UPROPERTY(EditAnywhere, Category = SculptToViewBrush, meta = (DisplayName = "Strength", UIMin = "0.0", UIMax = "1.0", ClampMin = "0.0", ClampMax = "1.0"))
+	UPROPERTY(EditAnywhere, Category = SculptToViewBrush, meta = (DisplayName = "Strength", UIMin = "0.0", UIMax = "1.0", ClampMin = "0.0", ClampMax = "1.0", ModelingQuickEdit))
 	float Strength = 0.5;
 
 	/** Amount of falloff to apply */
@@ -186,7 +186,7 @@ class MESHMODELINGTOOLSEXP_API USculptMaxBrushOpProps : public UMeshSculptBrushO
 	GENERATED_BODY()
 public:
 	/** Strength of the Brush */
-	UPROPERTY(EditAnywhere, Category = SculptMaxBrush, meta = (DisplayName = "Strength", UIMin = "0.0", UIMax = "1.0", ClampMin = "0.0", ClampMax = "1.0"))
+	UPROPERTY(EditAnywhere, Category = SculptMaxBrush, meta = (DisplayName = "Strength", UIMin = "0.0", UIMax = "1.0", ClampMin = "0.0", ClampMax = "1.0", ModelingQuickEdit))
 	float Strength = 0.5;
 
 	/** Amount of falloff to apply */
diff --git a/Engine/Plugins/Experimental/MeshModelingToolsetExp/Source/MeshModelingToolsExp/Public/Sculpting/MeshSculptToolBase.h b/Engine/Plugins/Experimental/MeshModelingToolsetExp/Source/MeshModelingToolsExp/Public/Sculpting/MeshSculptToolBase.h
--- a/Engine/Plugins/Experimental/MeshModelingToolsetExp/Source/MeshModelingToolsExp/Public/Sculpting/MeshSculptToolBase.h
+++ b/Engine/Plugins/Experimental/MeshModelingToolsetExp/Source/MeshModelingToolsExp/Public/Sculpting/MeshSculptToolBase.h
@@ -100,7 +100,7 @@ class MESHMODELINGTOOLSEXP_API USculptBrushProperties : public UInteractiveToolP
 	GENERATED_BODY()
 public:
 
-	UPROPERTY(EditAnywhere, Category = Brush, meta = (DisplayPriority = 1))
+	UPROPERTY(EditAnywhere, Category = Brush, meta = (DisplayPriority = 1, ModelingQuickEdit))
 	FBrushToolRadius BrushSize;
 
 	/** Amount of falloff to apply (0.0 - 1.0) */
diff --git a/Engine/Plugins/Experimental/MeshModelingToolsetExp/Source/MeshModelingToolsExp/Public/Sculpting/MeshSmoothingBrushOps.h b/Engine/Plugins/Experimental/MeshModelingToolsetExp/Source/MeshModelingToolsExp/Public/Sculpting/MeshSmoothingBrushOps.h
--- a/Engine/Plugins/Experimental/MeshModelingToolsetExp/Source/MeshModelingToolsExp/Public/Sculpting/MeshSmoothingBrushOps.h
+++ b/Engine/Plugins/Experimental/MeshModelingToolsetExp/Source/MeshModelingToolsExp/Public/Sculpting/MeshSmoothingBrushOps.h
@@ -25,7 +25,7 @@ class MESHMODELINGTOOLSEXP_API USmoothBrushOpProps : public UBaseSmoothBrushOpPr
 	GENERATED_BODY()
 public:
 	/** Strength of the Brush */
-	UPROPERTY(EditAnywhere, Category = SmoothBrush, meta = (DisplayName = "Strength", UIMin = "0.0", UIMax = "1.0", ClampMin = "0.0", ClampMax = "1.0"))
+	UPROPERTY(EditAnywhere, Category = SmoothBrush, meta = (DisplayName = "Strength", UIMin = "0.0", UIMax = "1.0", ClampMin = "0.0", ClampMax = "1.0", ModelingQuickEdit))
 	float Strength = 0.5;
 
 	/** Amount of falloff to apply (0.0 - 1.0) */
@@ -114,7 +114,7 @@ class MESHMODELINGTOOLSEXP_API USmoothFillBrushOpProps : public UBaseSmoothBrush
 	GENERATED_BODY()
 public:
 	/** Strength of the Brush */
-	UPROPERTY(EditAnywhere, Category = SmoothFillBrush, meta = (DisplayName = "Strength", UIMin = "0.0", UIMax = "1.0", ClampMin = "0.0", ClampMax = "1.0"))
+	UPROPERTY(EditAnywhere, Category = SmoothFillBrush, meta = (DisplayName = "Strength", UIMin = "0.0", UIMax = "1.0", ClampMin = "0.0", ClampMax = "1.0", ModelingQuickEdit))
 	float Strength = 0.5;
 
 	/** Amount of falloff to apply (0.0 - 1.0) */
@@ -187,7 +187,7 @@ class MESHMODELINGTOOLSEXP_API UFlattenBrushOpProps : public UMeshSculptBrushOpP
 	GENERATED_BODY()
 public:
 	/** Strength of the Brush */
-	UPROPERTY(EditAnywhere, Category = FlattenBrush, meta = (DisplayName = "Strength", UIMin = "0.0", UIMax = "1.0", ClampMin = "0.0", ClampMax = "1.0"))
+	UPROPERTY(EditAnywhere, Category = FlattenBrush, meta = (DisplayName = "Strength", UIMin = "0.0", UIMax = "1.0", ClampMin = "0.0", ClampMax = "1.0", ModelingQuickEdit))
 	float Strength = 0.5;
 
 	/** Amount of falloff to apply (0.0 - 1.0) */
@@ -267,7 +267,7 @@ class MESHMODELINGTOOLSEXP_API UEraseBrushOpProps : public UMeshSculptBrushOpPro
 	GENERATED_BODY()
 public:
 	/** Strength of the Brush */
-	UPROPERTY(EditAnywhere, Category = EraseBrush, meta = (DisplayName = "Strength", UIMin = "0.0", UIMax = "1.0", ClampMin = "0.0", ClampMax = "1.0"))
+	UPROPERTY(EditAnywhere, Category = EraseBrush, meta = (DisplayName = "Strength", UIMin = "0.0", UIMax = "1.0", ClampMin = "0.0", ClampMax = "1.0", ModelingQuickEdit))
 	float Strength = 0.5;
 
 	/** Amount of falloff to apply (0.0 - 1.0) */
diff --git a/Engine/Plugins/Experimental/MeshModelingToolsetExp/Source/MeshModelingToolsExp/Public/Sculpting/MeshVertexPaintBrushOps.h b/Engine/Plugins/Experimental/MeshModelingToolsetExp/Source/MeshModelingToolsExp/Public/Sculpting/MeshVertexPaintBrushOps.h
--- a/Engine/Plugins/Experimental/MeshModelingToolsetExp/Source/MeshModelingToolsExp/Public/Sculpting/MeshVertexPaintBrushOps.h
+++ b/Engine/Plugins/Experimental/MeshModelingToolsetExp/Source/MeshModelingToolsExp/Public/Sculpting/MeshVertexPaintBrushOps.h
@@ -123,7 +123,7 @@ class MESHMODELINGTOOLSEXP_API UVertexColorBaseBrushOpProps : public UMeshSculpt
 public:
 
 	/** Strength of the Brush */
-	UPROPERTY(EditAnywhere, Category = PaintBrush, meta = (DisplayName = "Strength", UIMin = "0.0", UIMax = "1.0", ClampMin = "0.0", ClampMax = "1.0"))
+	UPROPERTY(EditAnywhere, Category = PaintBrush, meta = (DisplayName = "Strength", UIMin = "0.0", UIMax = "1.0", ClampMin = "0.0", ClampMax = "1.0", ModelingQuickEdit))
 	float Strength = 0.5;
 
 	/** Amount of falloff to apply */
diff --git a/Engine/Plugins/Runtime/MeshModelingToolset/Source/MeshModelingTools/Private/ToolActivities/PolyEditBevelEdgeActivity.h b/Engine/Plugins/Runtime/MeshModelingToolset/Source/MeshModelingTools/Private/ToolActivities/PolyEditBevelEdgeActivity.h
--- a/Engine/Plugins/Runtime/MeshModelingToolset/Source/MeshModelingTools/Private/ToolActivities/PolyEditBevelEdgeActivity.h
+++ b/Engine/Plugins/Runtime/MeshModelingToolset/Source/MeshModelingTools/Private/ToolActivities/PolyEditBevelEdgeActivity.h
@@ -26,7 +26,7 @@ class MESHMODELINGTOOLS_API UPolyEditBevelEdgeProperties : public UInteractiveTo
 
 public:
 	/** Distance that each beveled mesh edge is inset from its initial position */
-	UPROPERTY(EditAnywhere, Category = Bevel, meta = (UIMin = ".001", UIMax = "100", ClampMin = "0.0001", ClampMax = "10000", SliderExponent = "3"))
+	UPROPERTY(EditAnywhere, Category = Bevel, meta = (UIMin = ".001", UIMax = "100", ClampMin = "0.0001", ClampMax = "10000", SliderExponent = "3", ModelingQuickEdit))
 	double BevelDistance = 4.0;
 
 	/** Number of edge loops added along the bevel faces */
