ModelingTools: Disable symmetry by default in MeshVertexSculpt tool.

diff --git a/Engine/Plugins/Experimental/MeshModelingToolsetExp/Source/MeshModelingToolsExp/Public/MeshVertexSculptTool.h b/Engine/Plugins/Experimental/MeshModelingToolsetExp/Source/MeshModelingToolsExp/Public/MeshVertexSculptTool.h
--- a/Engine/Plugins/Experimental/MeshModelingToolsetExp/Source/MeshModelingToolsExp/Public/MeshVertexSculptTool.h
+++ b/Engine/Plugins/Experimental/MeshModelingToolsetExp/Source/MeshModelingToolsExp/Public/MeshVertexSculptTool.h
@@ -197,7 +197,7 @@ class MESHMODELINGTOOLSEXP_API UMeshSymmetryProperties : public UInteractiveTool
 public:
 	/** Enable/Disable symmetric sculpting. This option will not be available if symmetry cannot be detected, or a non-symmetric edit has been made */
 	UPROPERTY(EditAnywhere, Category = Symmetry, meta = (HideEditConditionToggle, EditCondition = bSymmetryCanBeEnabled))
-	bool bEnableSymmetry = true;
+	bool bEnableSymmetry = false;
 
 	// this flag is set/updated by the Tool to enable/disable the bEnableSymmetry toggle
 	UPROPERTY(meta = (TransientToolProperty))
