MeshModelingTools: Avoid a crash SculptUtil::RecalculateROINormals when normals are unset on a mesh.

diff --git a/Engine/Plugins/Experimental/MeshModelingToolsetExp/Source/MeshModelingToolsExp/Private/Sculpting/MeshSculptUtil.cpp b/Engine/Plugins/Experimental/MeshModelingToolsetExp/Source/MeshModelingToolsExp/Private/Sculpting/MeshSculptUtil.cpp
--- a/Engine/Plugins/Experimental/MeshModelingToolsetExp/Source/MeshModelingToolsExp/Private/Sculpting/MeshSculptUtil.cpp
+++ b/Engine/Plugins/Experimental/MeshModelingToolsetExp/Source/MeshModelingToolsExp/Private/Sculpting/MeshSculptUtil.cpp
@@ -359,10 +359,13 @@ void UE::SculptUtil::PrecalculateNormalsROI(const FDynamicMesh3* Mesh, const TAr
 		{
 			TRACE_CPUPROFILER_EVENT_SCOPE(PrecalculateNormalsROI_FIND);
 			ParallelFor(TriangleROI.Num(), [&](int i) {
-				FIndex3i TriElems = Normals->GetTriangle(TriangleROI[i]);
-				ROIFlags[TriElems.A] = true;
-				ROIFlags[TriElems.B] = true;
-				ROIFlags[TriElems.C] = true;
+				if (Normals->IsSetTriangle(TriangleROI[i]))
+				{
+					FIndex3i TriElems = Normals->GetTriangle(TriangleROI[i]);
+					ROIFlags[TriElems.A] = true;
+					ROIFlags[TriElems.B] = true;
+					ROIFlags[TriElems.C] = true;
+				}
 			});
 		}
 	}
