Skin Weight Paint Tool: make sure that the transfer transaction is using a cleaned mesh as the FSkinToolWeights data structure is based on a cleaned mesh description

#rnx

diff --git a/Engine/Plugins/Experimental/MeshModelingToolsetExp/Source/MeshModelingToolsEditorOnlyExp/Private/SkeletalMesh/SkinWeightsPaintTool.cpp b/Engine/Plugins/Experimental/MeshModelingToolsetExp/Source/MeshModelingToolsEditorOnlyExp/Private/SkeletalMesh/SkinWeightsPaintTool.cpp
--- a/Engine/Plugins/Experimental/MeshModelingToolsetExp/Source/MeshModelingToolsEditorOnlyExp/Private/SkeletalMesh/SkinWeightsPaintTool.cpp
+++ b/Engine/Plugins/Experimental/MeshModelingToolsetExp/Source/MeshModelingToolsEditorOnlyExp/Private/SkeletalMesh/SkinWeightsPaintTool.cpp
@@ -856,7 +856,7 @@ void UWeightToolTransferManager::TransferWeights()
 	const USkeletalMesh* TargetSkeletalMesh = GetSkeletalMeshComponent(WeightTool->GetTarget())->GetSkeletalMeshAsset();
 	const bool bSameMesh = SourceSkeletalMesh == TargetSkeletalMesh;
 	const bool bSameLOD = GetLODId(ToolProperties->ActiveLOD) == GetLODId(ToolProperties->SourceLOD);
-	const bool bSameProfile = ToolProperties->ActiveSkinWeightProfile == ToolProperties->SourceSkinWeightProfile;
+	const bool bSameProfile = ToolProperties->GetActiveSkinWeightProfile() == ToolProperties->SourceSkinWeightProfile;
 	const bool bVerticesSelectedOnTarget = ToolProperties->EditingMode == EWeightEditMode::Mesh && !WeightTool->GetMainMeshSelector()->GetSelectedVertices().IsEmpty();
 	const bool bVerticesSelectedOnSource = ToolProperties->EditingMode == EWeightEditMode::Mesh && !MeshSelector->GetSelectedVertices().IsEmpty();
 	const bool bHasAnySelectedVertices = bVerticesSelectedOnSource || bVerticesSelectedOnTarget;
@@ -1142,11 +1142,11 @@ void UWeightToolTransferManager::OnPropertyModified(const USkinWeightsPaintToolP
 }
 
 void UWeightToolTransferManager::ApplyTranferredWeightsAsTransaction(
-	const UE::Geometry::FDynamicMeshVertexSkinWeightsAttribute* TransferredSkinWeights,
-	const TArray<int32>& VertexSubset,
-	FDynamicMesh3& TargetMesh)
+	const UE::Geometry::FDynamicMeshVertexSkinWeightsAttribute* InTransferredSkinWeights,
+	const TArray<int32>& InVertexSubset,
+	const FDynamicMesh3& InTargetMesh)
 {
-	check(TransferredSkinWeights);
+	check(InTransferredSkinWeights);
 	
 	// weight edits for transaction
 	FMultiBoneWeightEdits WeightEdits;
@@ -1156,12 +1156,18 @@ void UWeightToolTransferManager::ApplyTranferredWeightsAsTransaction(
 	
 	// spin through all the transferred skin weights and record a weight edit to apply as a transaction
 	static constexpr float ZeroWeight = 0.f;
-	const bool bUseSubset = !VertexSubset.IsEmpty();
-	const int32 NumVertices = bUseSubset ? VertexSubset.Num() : TargetMesh.VertexCount();
+	const bool bUseSubset = !InVertexSubset.IsEmpty();
+
+	// make sure the target has been cleaned as FSkinToolWeights data is based on a cleaned mesh description
+	FMeshDescription DummyCleanedMeshDescription;
+	FCleanedEditMesh CleanedTarget(InTargetMesh, DummyCleanedMeshDescription);
+	const FDynamicMesh3& TargetMesh = CleanedTarget.GetEditableMesh();
+	
+	const int32 NumVertices = bUseSubset ? InVertexSubset.Num() : InTargetMesh.VertexCount();
 	const UE::Geometry::FNonManifoldMappingSupport NonManifoldMappingSupport(TargetMesh);
 	for (int32 VertexIndex = 0; VertexIndex < NumVertices; ++VertexIndex)
 	{
-		const int32 VertexID = bUseSubset ? VertexSubset[VertexIndex] : VertexIndex;
+		const int32 VertexID = bUseSubset ? InVertexSubset[VertexIndex] : VertexIndex;
 		const int32 SrcVertexID = NonManifoldMappingSupport.GetOriginalNonManifoldVertexID(VertexID);
 		
 		// remove all weight on vertex
@@ -1185,7 +1191,7 @@ void UWeightToolTransferManager::ApplyTranferredWeightsAsTransaction(
 
 		// update with new weight
 		UE::AnimationCore::FBoneWeights TransferedBoneWeights;
-		TransferredSkinWeights->GetValue(VertexID, TransferedBoneWeights);
+		InTransferredSkinWeights->GetValue(VertexID, TransferedBoneWeights);
 		for (const UE::AnimationCore::FBoneWeight& BoneWeight: TransferedBoneWeights)
 		{
 			const int32 BoneIndex = BoneWeight.GetBoneIndex();					
@@ -1199,7 +1205,7 @@ void UWeightToolTransferManager::ApplyTranferredWeightsAsTransaction(
 	const FText TransactionLabel = LOCTEXT("TransferWeightsChange", "Transfer skin weights.");
 	WeightTool->ApplyWeightEditsAsTransaction(WeightEdits, TransactionLabel);
 	
-	// put the mesh back in it's current pose
+	// put the mesh back in its current pose
 	Weights.Deformer.SetAllVerticesToBeUpdated();
 
 	// notify user that weights were transferred.
diff --git a/Engine/Plugins/Experimental/MeshModelingToolsetExp/Source/MeshModelingToolsEditorOnlyExp/Public/SkeletalMesh/SkinWeightsPaintTool.h b/Engine/Plugins/Experimental/MeshModelingToolsetExp/Source/MeshModelingToolsEditorOnlyExp/Public/SkeletalMesh/SkinWeightsPaintTool.h
--- a/Engine/Plugins/Experimental/MeshModelingToolsetExp/Source/MeshModelingToolsEditorOnlyExp/Public/SkeletalMesh/SkinWeightsPaintTool.h
+++ b/Engine/Plugins/Experimental/MeshModelingToolsetExp/Source/MeshModelingToolsEditorOnlyExp/Public/SkeletalMesh/SkinWeightsPaintTool.h
@@ -592,9 +592,9 @@ private:
 	void TransferWeightsFromSameMeshAndLOD();
 
 	void ApplyTranferredWeightsAsTransaction(
-		const UE::Geometry::FDynamicMeshVertexSkinWeightsAttribute* TransferredSkinWeights,
-		const TArray<int32>& VertexSubset,
-		FDynamicMesh3& TargetMesh);
+		const UE::Geometry::FDynamicMeshVertexSkinWeightsAttribute* InTransferredSkinWeights,
+		const TArray<int32>& InVertexSubset,
+		const FDynamicMesh3& InTargetMesh);
 	
 	UPROPERTY()
 	TObjectPtr<UPreviewMesh> SourcePreviewMesh = nullptr;
