ModelingTools: Added Height map bake support for BakeTextures tool.

diff --git a/Engine/Plugins/Editor/ModelingToolsEditorMode/Source/ModelingToolsEditorMode/Private/DetailsCustomizations/BakeMeshAttributeToolCustomizations.cpp b/Engine/Plugins/Editor/ModelingToolsEditorMode/Source/ModelingToolsEditorMode/Private/DetailsCustomizations/BakeMeshAttributeToolCustomizations.cpp
--- a/Engine/Plugins/Editor/ModelingToolsEditorMode/Source/ModelingToolsEditorMode/Private/DetailsCustomizations/BakeMeshAttributeToolCustomizations.cpp
+++ b/Engine/Plugins/Editor/ModelingToolsEditorMode/Source/ModelingToolsEditorMode/Private/DetailsCustomizations/BakeMeshAttributeToolCustomizations.cpp
@@ -49,7 +49,8 @@ namespace BakeCustomizationLocals
 				EBakeMapType::FaceNormal,
 				EBakeMapType::BentNormal,
 				EBakeMapType::Position,
-				EBakeMapType::Curvature
+				EBakeMapType::Curvature,
+				EBakeMapType::Height
 			},
 			// Lighting
 			{
diff --git a/Engine/Plugins/Experimental/MeshModelingToolsetExp/Source/MeshModelingToolsExp/Private/BakeMeshAttributeMapsTool.cpp b/Engine/Plugins/Experimental/MeshModelingToolsetExp/Source/MeshModelingToolsExp/Private/BakeMeshAttributeMapsTool.cpp
--- a/Engine/Plugins/Experimental/MeshModelingToolsetExp/Source/MeshModelingToolsExp/Private/BakeMeshAttributeMapsTool.cpp
+++ b/Engine/Plugins/Experimental/MeshModelingToolsetExp/Source/MeshModelingToolsExp/Private/BakeMeshAttributeMapsTool.cpp
@@ -15,6 +15,7 @@
 #include "Sampling/MeshPropertyMapEvaluator.h"
 #include "Sampling/MeshUVShellMapEvaluator.h"
 #include "Sampling/MeshResampleImageEvaluator.h"
+#include "Sampling/MeshHeightMapEvaluator.h"
 
 #include "ImageUtils.h"
 
@@ -116,6 +117,7 @@ public:
 	FOcclusionMapSettings OcclusionSettings;
 	FCurvatureMapSettings CurvatureSettings;
 	FMeshPropertyMapSettings PropertySettings;
+	FHeightMapSettings HeightSettings;
 	FUVShellMapSettings UVShellSettings;
 	FTexture2DSettings TextureSettings;
 	FTexture2DSettings MultiTextureSettings;
@@ -256,6 +258,22 @@ public:
 				Baker->AddEvaluator(PropertyEval);
 				break;
 			}
+			case EBakeMapType::Height:
+			{
+				TSharedPtr<FMeshHeightMapEvaluator> HeightEval = MakeShared<FMeshHeightMapEvaluator>();
+				HeightEval->RangeMode = (FMeshHeightMapEvaluator::EHeightRangeMode) HeightSettings.HeightRangeMode;
+				switch(HeightEval->RangeMode)
+				{
+				case FMeshHeightMapEvaluator::EHeightRangeMode::Absolute:
+					HeightEval->Range = FInterval1f::MakeFromUnordered(HeightSettings.InnerDistance, HeightSettings.OuterDistance);
+					break;
+				case FMeshHeightMapEvaluator::EHeightRangeMode::RelativeBounds:
+					HeightEval->Range = FInterval1f::MakeFromUnordered(HeightSettings.InnerBoundsDistance, HeightSettings.OuterBoundsDistance);
+					break;
+				}
+				Baker->AddEvaluator(HeightEval);
+				break;
+			}
 			case EBakeMapType::UVShell:
 			{
 				TSharedPtr<FMeshUVShellMapEvaluator> UVShellEval = MakeShared<FMeshUVShellMapEvaluator>();
@@ -417,6 +435,17 @@ void UBakeMeshAttributeMapsTool::Setup()
 	UVShellSettings->WatchProperty(UVShellSettings->BackgroundColor, [this](FLinearColor) { OpState |= EBakeOpState::Evaluate; });
 
 
+	HeightSettings = NewObject<UBakeHeightMapToolProperties>(this);
+	HeightSettings->RestoreProperties(this);
+	AddToolPropertySource(HeightSettings);
+	SetToolPropertySourceEnabled(HeightSettings, false);
+	HeightSettings->WatchProperty(HeightSettings->HeightRangeMode, [this](EBakeHeightRangeMode) { OpState |= EBakeOpState::Evaluate; });
+	HeightSettings->WatchProperty(HeightSettings->InnerDistance, [this](float) { OpState |= EBakeOpState::Evaluate; });
+	HeightSettings->WatchProperty(HeightSettings->OuterDistance, [this](float) { OpState |= EBakeOpState::Evaluate; });
+	HeightSettings->WatchProperty(HeightSettings->InnerBoundsDistance, [this](float) { OpState |= EBakeOpState::Evaluate; });
+	HeightSettings->WatchProperty(HeightSettings->OuterBoundsDistance, [this](float) { OpState |= EBakeOpState::Evaluate; });
+
+
 	TextureSettings = NewObject<UBakeTexture2DProperties>(this);
 	TextureSettings->RestoreProperties(this);
 	AddToolPropertySource(TextureSettings);
@@ -517,6 +546,11 @@ TUniquePtr<UE::Geometry::TGenericDataOperator<FMeshMapBaker>> UBakeMeshAttribute
 		Op->PropertySettings = CachedMeshPropertyMapSettings;
 	}
 
+	if ((bool)(CachedBakeSettings.BakeMapTypes & EBakeMapType::Height))
+	{
+		Op->HeightSettings = CachedHeightMapSettings;
+	}
+
 	if ((bool)(CachedBakeSettings.BakeMapTypes & EBakeMapType::UVShell))
 	{
 		Op->UVShellSettings = CachedUVShellMapSettings;
@@ -548,6 +582,7 @@ void UBakeMeshAttributeMapsTool::OnShutdown(EToolShutdownType ShutdownType)
 	InputMeshSettings->SaveProperties(this);
 	OcclusionSettings->SaveProperties(this);
 	CurvatureSettings->SaveProperties(this);
+	HeightSettings->SaveProperties(this);
 	TextureSettings->SaveProperties(this);
 	MultiTextureSettings->SaveProperties(this);
 
@@ -709,6 +744,10 @@ void UBakeMeshAttributeMapsTool::UpdateResult()
 	{
 		OpState |= UpdateResult_MeshProperty(CachedBakeSettings.Dimensions);
 	}
+	if ((bool)(CachedBakeSettings.BakeMapTypes & EBakeMapType::Height))
+	{
+		OpState |= UpdateResult_Height(CachedBakeSettings.Dimensions);
+	}
 	if ((bool)(CachedBakeSettings.BakeMapTypes & EBakeMapType::UVShell))
 	{
 		OpState |= UpdateResult_UVShellMap(CachedBakeSettings.Dimensions);
@@ -838,6 +877,7 @@ void UBakeMeshAttributeMapsTool::UpdateOnModeChange()
 	// Update tool property sets.
 	SetToolPropertySourceEnabled(OcclusionSettings, false);
 	SetToolPropertySourceEnabled(CurvatureSettings, false);
+	SetToolPropertySourceEnabled(HeightSettings, false);
 	SetToolPropertySourceEnabled(UVShellSettings, false);
 	SetToolPropertySourceEnabled(TextureSettings, false);
 	SetToolPropertySourceEnabled(MultiTextureSettings, false);
@@ -862,6 +902,9 @@ void UBakeMeshAttributeMapsTool::UpdateOnModeChange()
 		case EBakeMapType::PolyGroupID:
 		case EBakeMapType::VertexColor:
 			break;
+		case EBakeMapType::Height:
+			SetToolPropertySourceEnabled(HeightSettings, true);
+			break;
 		case EBakeMapType::UVShell:
 			SetToolPropertySourceEnabled(UVShellSettings, true);
 			break;
diff --git a/Engine/Plugins/Experimental/MeshModelingToolsetExp/Source/MeshModelingToolsExp/Private/BakeMeshAttributeMapsToolBase.cpp b/Engine/Plugins/Experimental/MeshModelingToolsetExp/Source/MeshModelingToolsExp/Private/BakeMeshAttributeMapsToolBase.cpp
--- a/Engine/Plugins/Experimental/MeshModelingToolsetExp/Source/MeshModelingToolsExp/Private/BakeMeshAttributeMapsToolBase.cpp
+++ b/Engine/Plugins/Experimental/MeshModelingToolsetExp/Source/MeshModelingToolsExp/Private/BakeMeshAttributeMapsToolBase.cpp
@@ -264,6 +264,7 @@ void UBakeMeshAttributeMapsToolBase::UpdatePreview(const EBakeMapType PreviewMap
 		case EBakeMapType::Texture:
 		case EBakeMapType::MultiTexture:
 		case EBakeMapType::VertexColor:
+		case EBakeMapType::Height:
 		case EBakeMapType::UVShell:
 			PreviewMaterial->SetTextureParameterValue(TEXT("NormalMap"), EmptyNormalMap);
 			PreviewMaterial->SetTextureParameterValue(TEXT("OcclusionMap"), EmptyColorMapWhite);
@@ -455,6 +456,12 @@ void UBakeMeshAttributeMapsToolBase::OnMapsUpdated(const TUniquePtr<UE::Geometry
 			UpdateCachedMap(MapType, Format, 0);
 			break;
 		}
+		case EMeshMapEvaluatorType::Height:
+		{
+			constexpr EBakeMapType MapType = EBakeMapType::Height;
+			UpdateCachedMap(MapType, Format, 0);
+			break;
+		}
 		case EMeshMapEvaluatorType::UVShell:
 		{
 			constexpr EBakeMapType MapType = EBakeMapType::UVShell;
@@ -517,6 +524,7 @@ FTexture2DBuilder::ETextureType UBakeMeshAttributeMapsToolBase::GetTextureType(c
 	case EBakeMapType::ObjectSpaceNormal:
 	case EBakeMapType::FaceNormal:
 	case EBakeMapType::Position:
+	case EBakeMapType::Height:
 		TexType = FTexture2DBuilder::ETextureType::ColorLinear;
 		break;
 	case EBakeMapType::MaterialID:
@@ -575,6 +583,9 @@ void UBakeMeshAttributeMapsToolBase::GetTextureName(const EBakeMapType MapType,
 	case EBakeMapType::Position:
 		TexName = FString::Printf(TEXT("%s_Position"), *BaseName);
 		break;
+	case EBakeMapType::Height:
+		TexName = FString::Printf(TEXT("&s_Height"), *BaseName);
+		break;
 	case EBakeMapType::UVShell:
 		TexName = FString::Printf(TEXT("%s_UVShell"), *BaseName);
 		break;
@@ -747,6 +758,9 @@ void UBakeMeshAttributeMapsToolBase::RecordAnalytics(const FBakeAnalytics& Data,
 	const bool bVertexColor = static_cast<bool>(Data.BakeSettings.SourceBakeMapTypes & EBakeMapType::VertexColor);
 	Attributes.Add(FAnalyticsEventAttribute(TEXT("Settings.VertexColor.Enabled"), bVertexColor));
 
+	const bool bHeight = static_cast<bool>(Data.BakeSettings.SourceBakeMapTypes & EBakeMapType::Height);
+	Attributes.Add(FAnalyticsEventAttribute(TEXT("Settings.Height.Enabled"), bHeight));
+
 	const bool bUVShell = static_cast<bool>(Data.BakeSettings.SourceBakeMapTypes & EBakeMapType::UVShell);
 	Attributes.Add(FAnalyticsEventAttribute(TEXT("Settings.UVShell.Enabled"), bUVShell));
 
diff --git a/Engine/Plugins/Experimental/MeshModelingToolsetExp/Source/MeshModelingToolsExp/Private/BakeMeshAttributeTool.cpp b/Engine/Plugins/Experimental/MeshModelingToolsetExp/Source/MeshModelingToolsExp/Private/BakeMeshAttributeTool.cpp
--- a/Engine/Plugins/Experimental/MeshModelingToolsetExp/Source/MeshModelingToolsExp/Private/BakeMeshAttributeTool.cpp
+++ b/Engine/Plugins/Experimental/MeshModelingToolsetExp/Source/MeshModelingToolsExp/Private/BakeMeshAttributeTool.cpp
@@ -7,6 +7,7 @@
 #include "AssetUtils/Texture2DUtil.h"
 
 #include "Sampling/MeshCurvatureMapEvaluator.h"
+#include "Sampling/MeshHeightMapEvaluator.h"
 
 #include UE_INLINE_GENERATED_CPP_BY_NAME(BakeMeshAttributeTool)
 
@@ -170,6 +171,36 @@ EBakeOpState UBakeMeshAttributeTool::UpdateResult_MeshProperty(const FImageDimen
 }
 
 
+EBakeOpState UBakeMeshAttributeTool::UpdateResult_Height(const FImageDimensions& Dimensions)
+{
+	EBakeOpState ResultState = EBakeOpState::Clean;
+
+	FHeightMapSettings HeightMapSettings;
+	HeightMapSettings.Dimensions = Dimensions;
+	HeightMapSettings.InnerDistance = HeightSettings->InnerDistance;
+	HeightMapSettings.OuterDistance = HeightSettings->OuterDistance;
+	HeightMapSettings.InnerBoundsDistance = HeightSettings->InnerBoundsDistance;
+	HeightMapSettings.OuterBoundsDistance = HeightSettings->OuterBoundsDistance;
+	
+	switch (HeightSettings->HeightRangeMode)
+	{
+	case EBakeHeightRangeMode::Absolute:
+		HeightMapSettings.HeightRangeMode = static_cast<int32>(FMeshHeightMapEvaluator::EHeightRangeMode::Absolute);
+		break;
+	case EBakeHeightRangeMode::RelativeBounds:
+		HeightMapSettings.HeightRangeMode = static_cast<int32>(FMeshHeightMapEvaluator::EHeightRangeMode::RelativeBounds);
+		break;
+	}
+
+	if (!(CachedHeightMapSettings == HeightMapSettings))
+	{
+		CachedHeightMapSettings = HeightMapSettings;
+		ResultState = EBakeOpState::Evaluate;
+	}
+	return ResultState;
+}
+
+
 EBakeOpState UBakeMeshAttributeTool::UpdateResult_Texture2DImage(const FImageDimensions& Dimensions, const FDynamicMesh3* DetailMesh)
 {
 	EBakeOpState ResultState = EBakeOpState::Clean;
diff --git a/Engine/Plugins/Experimental/MeshModelingToolsetExp/Source/MeshModelingToolsExp/Public/BakeMeshAttributeMapsTool.h b/Engine/Plugins/Experimental/MeshModelingToolsetExp/Source/MeshModelingToolsExp/Public/BakeMeshAttributeMapsTool.h
--- a/Engine/Plugins/Experimental/MeshModelingToolsetExp/Source/MeshModelingToolsExp/Public/BakeMeshAttributeMapsTool.h
+++ b/Engine/Plugins/Experimental/MeshModelingToolsetExp/Source/MeshModelingToolsExp/Public/BakeMeshAttributeMapsTool.h
@@ -54,7 +54,7 @@ class MESHMODELINGTOOLSEXP_API UBakeMeshAttributeMapsToolProperties : public UIn
 public:
 	/** The bake output types to generate */
 	UPROPERTY(EditAnywhere, Category = BakeOutput, meta = (DisplayName = "Output Types", Bitmask, BitmaskEnum = "/Script/MeshModelingToolsExp.EBakeMapType",
-		ValidEnumValues="TangentSpaceNormal, AmbientOcclusion, BentNormal, Curvature, Texture, ObjectSpaceNormal, FaceNormal, Position, MaterialID, PolyGroupID, MultiTexture, VertexColor, UVShell"))
+		ValidEnumValues="TangentSpaceNormal, AmbientOcclusion, BentNormal, Curvature, Texture, ObjectSpaceNormal, FaceNormal, Position, MaterialID, PolyGroupID, MultiTexture, VertexColor, UVShell, Height"))
 	int32 MapTypes = static_cast<int32>(EBakeMapType::None);
 
 	/** The baked output type used for preview in the viewport */
diff --git a/Engine/Plugins/Experimental/MeshModelingToolsetExp/Source/MeshModelingToolsExp/Public/BakeMeshAttributeTool.h b/Engine/Plugins/Experimental/MeshModelingToolsetExp/Source/MeshModelingToolsExp/Public/BakeMeshAttributeTool.h
--- a/Engine/Plugins/Experimental/MeshModelingToolsetExp/Source/MeshModelingToolsExp/Public/BakeMeshAttributeTool.h
+++ b/Engine/Plugins/Experimental/MeshModelingToolsetExp/Source/MeshModelingToolsExp/Public/BakeMeshAttributeTool.h
@@ -55,8 +55,10 @@ enum class EBakeMapType
 	Zero                   = 1 << 13,
 	/* UV shell */
 	UVShell                = 1 << 14 UMETA(DisplayName = "UV Shell"),
+	/* Height */
+	Height                 = 1 << 15,
 
-	All                    = 0x3FFF UMETA(Hidden)
+	All                    = 0x7FFF UMETA(Hidden)
 };
 ENUM_CLASS_FLAGS(EBakeMapType);
 
@@ -73,6 +75,7 @@ static constexpr EBakeMapType ENUM_EBAKEMAPTYPE_ALL[] =
 	EBakeMapType::BentNormal,
 	EBakeMapType::Position,
 	EBakeMapType::Curvature,
+	EBakeMapType::Height,
 	EBakeMapType::AmbientOcclusion,
 	EBakeMapType::Texture,
 	EBakeMapType::MultiTexture,
@@ -123,6 +126,9 @@ protected:
 
 	UPROPERTY()
 	TObjectPtr<UBakeMultiTexture2DProperties> MultiTextureSettings;
+
+	UPROPERTY()
+	TObjectPtr<UBakeHeightMapToolProperties> HeightSettings;
 	
 	
 	//
@@ -174,6 +180,9 @@ protected:
 
 	EBakeOpState UpdateResult_MeshProperty(const FImageDimensions& Dimensions);
 	FMeshPropertyMapSettings CachedMeshPropertyMapSettings;
+	
+	EBakeOpState UpdateResult_Height(const FImageDimensions& Dimensions);
+	FHeightMapSettings CachedHeightMapSettings;
 
 	EBakeOpState UpdateResult_Texture2DImage(const FImageDimensions& Dimensions, const FDynamicMesh3* DetailMesh);
 	FTexture2DSettings CachedTexture2DSettings;
diff --git a/Engine/Plugins/Experimental/MeshModelingToolsetExp/Source/MeshModelingToolsExp/Public/BakeMeshAttributeToolCommon.h b/Engine/Plugins/Experimental/MeshModelingToolsetExp/Source/MeshModelingToolsExp/Public/BakeMeshAttributeToolCommon.h
--- a/Engine/Plugins/Experimental/MeshModelingToolsetExp/Source/MeshModelingToolsExp/Public/BakeMeshAttributeToolCommon.h
+++ b/Engine/Plugins/Experimental/MeshModelingToolsetExp/Source/MeshModelingToolsExp/Public/BakeMeshAttributeToolCommon.h
@@ -261,6 +261,43 @@ public:
 };
 
 
+UENUM()
+enum class EBakeHeightRangeMode
+{
+	/** Height range defined in absolute units in object space */
+	Absolute,
+	/** Height range defined as a ratio of the maximal bounding box axis */
+	RelativeBounds
+};
+
+
+UCLASS()
+class MESHMODELINGTOOLSEXP_API UBakeHeightMapToolProperties : public UInteractiveToolPropertySet
+{
+	GENERATED_BODY()
+public:
+	/** The coordinate space to compute height */
+	UPROPERTY(EditAnywhere, Category = HeightOutput)
+	EBakeHeightRangeMode HeightRangeMode = EBakeHeightRangeMode::Absolute;
+
+	/** The inner height map distance from the low poly mesh, in object space units. This value maps to black. */
+	UPROPERTY(EditAnywhere, Category = HeightOutput, meta=(EditConditionHides,EditCondition="HeightRangeMode == EBakeHeightRangeMode::Absolute", UIMin=-100.0f, UIMax=100.0f))
+	float InnerDistance = -1.0f;
+
+	/** The outer height map distance from the low poly mesh, in object space units. This value maps to white. */
+	UPROPERTY(EditAnywhere, Category = HeightOutput, meta=(EditConditionHides,EditCondition="HeightRangeMode == EBakeHeightRangeMode::Absolute", UIMin=-100.0f, UIMax=100.0f))
+	float OuterDistance = 1.0f;
+
+	/** The inner height map distance from the low poly mesh, as a ratio of the maximum bounding box axis. This value maps to black. */
+	UPROPERTY(EditAnywhere, Category = HeightOutput, meta=(EditConditionHides,EditCondition="HeightRangeMode == EBakeHeightRangeMode::RelativeBounds", ClampMin=-1.0f, ClampMax=1.0f))
+	float InnerBoundsDistance = -0.1f;
+
+	/** The outer height map distance from the low poly mesh, as a ratio of the maximum bounding box axis. This value maps to white. */
+	UPROPERTY(EditAnywhere, Category = HeightOutput, meta=(EditConditionHides,EditCondition="HeightRangeMode == EBakeHeightRangeMode::RelativeBounds", ClampMin=-1.0f, ClampMax=1.0f))
+	float OuterBoundsDistance = 0.1f;
+};
+
+
 UCLASS()
 class MESHMODELINGTOOLSEXP_API UBakeTexture2DProperties : public UInteractiveToolPropertySet
 {
@@ -409,6 +446,26 @@ struct FMeshPropertyMapSettings
 	}
 };
 
+struct FHeightMapSettings
+{
+	FImageDimensions Dimensions;
+	int32 HeightRangeMode = 0;
+	float InnerDistance = -1.0f;
+	float OuterDistance = 1.0f;
+	float InnerBoundsDistance = -0.1f;
+	float OuterBoundsDistance = 0.1f;
+
+	bool operator==(const FHeightMapSettings& Other) const
+	{
+		return Dimensions == Other.Dimensions
+		&& HeightRangeMode == Other.HeightRangeMode
+		&& InnerDistance == Other.InnerDistance
+		&& OuterDistance == Other.OuterDistance
+		&& InnerBoundsDistance == Other.InnerBoundsDistance
+		&& OuterBoundsDistance == Other.OuterBoundsDistance;
+	}
+};
+
 struct FUVShellMapSettings
 {
 	FImageDimensions Dimensions;
