ModelingTools: Added height map bake to BakeVertex tool.

diff --git a/Engine/Plugins/Experimental/MeshModelingToolsetExp/Source/MeshModelingToolsExp/Private/BakeMeshAttributeVertexTool.cpp b/Engine/Plugins/Experimental/MeshModelingToolsetExp/Source/MeshModelingToolsExp/Private/BakeMeshAttributeVertexTool.cpp
--- a/Engine/Plugins/Experimental/MeshModelingToolsetExp/Source/MeshModelingToolsExp/Private/BakeMeshAttributeVertexTool.cpp
+++ b/Engine/Plugins/Experimental/MeshModelingToolsetExp/Source/MeshModelingToolsExp/Private/BakeMeshAttributeVertexTool.cpp
@@ -17,6 +17,7 @@
 #include "Sampling/MeshCurvatureMapEvaluator.h"
 #include "Sampling/MeshPropertyMapEvaluator.h"
 #include "Sampling/MeshResampleImageEvaluator.h"
+#include "Sampling/MeshHeightMapEvaluator.h"
 
 #include "TargetInterfaces/MaterialProvider.h"
 #include "TargetInterfaces/PrimitiveComponentBackedTarget.h"
@@ -67,6 +68,7 @@ public:
 	UBakeMeshAttributeVertexTool::FBakeSettings BakeSettings;
 	FOcclusionMapSettings OcclusionSettings;
 	FCurvatureMapSettings CurvatureSettings;
+	FHeightMapSettings HeightSettings;
 	FTexture2DSettings TextureSettings;
 	FTexture2DSettings MultiTextureSettings;
 
@@ -123,6 +125,20 @@ public:
 			CurvatureEval->UseClampMode = static_cast<FMeshCurvatureMapEvaluator::EClampMode>(CurvatureSettings.ClampMode);
 		};
 
+		auto InitHeightEvaluator = [this] (FMeshHeightMapEvaluator* HeightEval)
+		{
+			HeightEval->RangeMode = (FMeshHeightMapEvaluator::EHeightRangeMode) HeightSettings.HeightRangeMode;
+			switch (HeightEval->RangeMode)
+			{
+			case FMeshHeightMapEvaluator::EHeightRangeMode::Absolute:
+				HeightEval->Range = FInterval1f::MakeFromUnordered(HeightSettings.InnerDistance, HeightSettings.OuterDistance);
+				break;
+			case FMeshHeightMapEvaluator::EHeightRangeMode::RelativeBounds:
+				HeightEval->Range = FInterval1f::MakeFromUnordered(HeightSettings.InnerBoundsDistance, HeightSettings.OuterBoundsDistance);
+				break;
+			}
+		};
+
 		if (BakeSettings.OutputMode == EBakeVertexOutput::PerChannel)
 		{
 			for(int ChannelIdx = 0; ChannelIdx < 4; ++ChannelIdx)
@@ -143,6 +159,13 @@ public:
 					Baker->ChannelEvaluators[ChannelIdx] = CurvatureEval;
 					break;
 				}
+				case EBakeMapType::Height:
+				{
+					TSharedPtr<FMeshHeightMapEvaluator> HeightEval = MakeShared<FMeshHeightMapEvaluator>();
+					InitHeightEvaluator(HeightEval.Get());
+					Baker->ChannelEvaluators[ChannelIdx] = HeightEval;
+					break;
+				}
 				case EBakeMapType::One:
 				case EBakeMapType::Zero:
 				{
@@ -225,6 +248,13 @@ public:
 				Baker->ColorEvaluator = PropertyEval;
 				break;
 			}
+			case EBakeMapType::Height:
+			{
+				TSharedPtr<FMeshHeightMapEvaluator> HeightEval = MakeShared<FMeshHeightMapEvaluator>();
+				InitHeightEvaluator(HeightEval.Get());
+				Baker->ColorEvaluator = HeightEval;
+				break;
+			}
 			case EBakeMapType::Texture:
 			{
 				TSharedPtr<FMeshResampleImageEvaluator, ESPMode::ThreadSafe> TextureEval = MakeShared<FMeshResampleImageEvaluator, ESPMode::ThreadSafe>();
@@ -359,6 +389,16 @@ void UBakeMeshAttributeVertexTool::Setup()
 	CurvatureSettings->WatchProperty(CurvatureSettings->ColorMapping, [this](EBakeCurvatureColorMode) { OpState |= EBakeOpState::Evaluate; });
 	CurvatureSettings->WatchProperty(CurvatureSettings->Clamping, [this](EBakeCurvatureClampMode) { OpState |= EBakeOpState::Evaluate; });
 
+	HeightSettings = NewObject<UBakeHeightMapToolProperties>(this);
+	HeightSettings->RestoreProperties(this);
+	AddToolPropertySource(HeightSettings);
+	SetToolPropertySourceEnabled(HeightSettings, false);
+	HeightSettings->WatchProperty(HeightSettings->HeightRangeMode, [this](EBakeHeightRangeMode) { OpState |= EBakeOpState::Evaluate; });
+	HeightSettings->WatchProperty(HeightSettings->InnerDistance, [this](float) { OpState |= EBakeOpState::Evaluate; });
+	HeightSettings->WatchProperty(HeightSettings->OuterDistance, [this](float) { OpState |= EBakeOpState::Evaluate; });
+	HeightSettings->WatchProperty(HeightSettings->InnerBoundsDistance, [this](float) { OpState |= EBakeOpState::Evaluate; });
+	HeightSettings->WatchProperty(HeightSettings->OuterBoundsDistance, [this](float) { OpState |= EBakeOpState::Evaluate; });
+
 	TextureSettings = NewObject<UBakeTexture2DProperties>(this);
 	TextureSettings->RestoreProperties(this);
 	AddToolPropertySource(TextureSettings);
@@ -405,6 +445,7 @@ void UBakeMeshAttributeVertexTool::OnShutdown(EToolShutdownType ShutdownType)
 	InputMeshSettings->SaveProperties(this);
 	OcclusionSettings->SaveProperties(this);
 	CurvatureSettings->SaveProperties(this);
+	HeightSettings->SaveProperties(this);
 	TextureSettings->SaveProperties(this);
 	MultiTextureSettings->SaveProperties(this);
 
@@ -474,6 +515,7 @@ TUniquePtr<UE::Geometry::TGenericDataOperator<FMeshVertexBaker>> UBakeMeshAttrib
 	Op->BakeSettings = CachedBakeSettings;
 	Op->OcclusionSettings = CachedOcclusionMapSettings;
 	Op->CurvatureSettings = CachedCurvatureMapSettings;
+	Op->HeightSettings = CachedHeightMapSettings;
 	Op->TextureSettings = CachedTexture2DSettings;
 	Op->MultiTextureSettings = CachedMultiTexture2DSettings;
 	Op->bIsBakeToSelf = bIsBakeToSelf;
@@ -518,6 +560,7 @@ void UBakeMeshAttributeVertexTool::UpdateOnModeChange()
 {
 	SetToolPropertySourceEnabled(OcclusionSettings, false);
 	SetToolPropertySourceEnabled(CurvatureSettings, false);
+	SetToolPropertySourceEnabled(HeightSettings, false);
 	SetToolPropertySourceEnabled(TextureSettings, false);
 	SetToolPropertySourceEnabled(MultiTextureSettings, false);
 
@@ -532,6 +575,9 @@ void UBakeMeshAttributeVertexTool::UpdateOnModeChange()
 			case EBakeMapType::Curvature:
 				SetToolPropertySourceEnabled(CurvatureSettings, true);
 				break;
+			case EBakeMapType::Height:
+				SetToolPropertySourceEnabled(HeightSettings, true);
+				break;
 			case EBakeMapType::Texture:
 				SetToolPropertySourceEnabled(TextureSettings, true);
 				break;
@@ -561,6 +607,9 @@ void UBakeMeshAttributeVertexTool::UpdateOnModeChange()
 				case EBakeMapType::Curvature:
 					SetToolPropertySourceEnabled(CurvatureSettings, true);
 					break;
+				case EBakeMapType::Height:
+					SetToolPropertySourceEnabled(HeightSettings, true);
+					break;
 				case EBakeMapType::None:
 				default:
 					break;
@@ -759,6 +808,9 @@ void UBakeMeshAttributeVertexTool::UpdateResult()
 		case EBakeMapType::PolyGroupID:
 			OpState |= UpdateResult_MeshProperty(Dimensions);
 			break;
+		case EBakeMapType::Height:
+			OpState |= UpdateResult_Height(Dimensions);
+			break;
 		case EBakeMapType::VertexColor:
 			OpState |= UpdateResult_MeshProperty(Dimensions);
 			// Force copy the original vertex colors to our PreviewMesh so that
@@ -788,6 +840,10 @@ void UBakeMeshAttributeVertexTool::UpdateResult()
 		{
 			OpState |= UpdateResult_Curvature(Dimensions);
 		}
+		if (HeightSettings->IsPropertySetEnabled())
+		{
+			OpState |= UpdateResult_Height(Dimensions);
+		}
 		
 		// Always force copy the original vertex colors to our PreviewMesh for
 		// PerChannel bakes so that channels that are not targeted persist
diff --git a/Engine/Plugins/Experimental/MeshModelingToolsetExp/Source/MeshModelingToolsExp/Public/BakeMeshAttributeVertexTool.h b/Engine/Plugins/Experimental/MeshModelingToolsetExp/Source/MeshModelingToolsExp/Public/BakeMeshAttributeVertexTool.h
--- a/Engine/Plugins/Experimental/MeshModelingToolsetExp/Source/MeshModelingToolsExp/Public/BakeMeshAttributeVertexTool.h
+++ b/Engine/Plugins/Experimental/MeshModelingToolsetExp/Source/MeshModelingToolsExp/Public/BakeMeshAttributeVertexTool.h
@@ -75,31 +75,31 @@ public:
 
 	/** The bake output type to generate */
 	UPROPERTY(EditAnywhere, Category = BakeOutput, meta=(Bitmask, BitmaskEnum = "/Script/MeshModelingToolsExp.EBakeMapType",
-		ValidEnumValues="TangentSpaceNormal, AmbientOcclusion, BentNormal, Curvature, Texture, ObjectSpaceNormal, FaceNormal, Position, MaterialID, PolyGroupID, MultiTexture, VertexColor",
+		ValidEnumValues="TangentSpaceNormal, AmbientOcclusion, BentNormal, Curvature, Texture, ObjectSpaceNormal, FaceNormal, Position, MaterialID, PolyGroupID, MultiTexture, VertexColor, Height",
 		EditCondition="OutputMode == EBakeVertexOutput::RGBA", EditConditionHides))
 	int32 OutputType = static_cast<int32>(EBakeMapType::TangentSpaceNormal);
 
 	/** The bake output type to generate in the Red channel */
 	UPROPERTY(EditAnywhere, Category = BakeOutput, meta=(Bitmask, BitmaskEnum = "/Script/MeshModelingToolsExp.EBakeMapType",
-		ValidEnumValues="None, AmbientOcclusion, Curvature, One, Zero",
+		ValidEnumValues="None, AmbientOcclusion, Curvature, Height, One, Zero",
 		EditCondition="OutputMode == EBakeVertexOutput::PerChannel", EditConditionHides))
 	int32 OutputTypeR = static_cast<int32>(EBakeMapType::None);
 
 	/** The bake output type to generate in the Green channel */
 	UPROPERTY(EditAnywhere, Category = BakeOutput, meta=(Bitmask, BitmaskEnum = "/Script/MeshModelingToolsExp.EBakeMapType",
-		ValidEnumValues="None, AmbientOcclusion, Curvature, One, Zero",
+		ValidEnumValues="None, AmbientOcclusion, Curvature, Height, One, Zero",
 		EditCondition="OutputMode == EBakeVertexOutput::PerChannel", EditConditionHides))
 	int32 OutputTypeG = static_cast<int32>(EBakeMapType::None);
 
 	/** The bake output type to generate in the Blue channel */
 	UPROPERTY(EditAnywhere, Category = BakeOutput, meta=(Bitmask, BitmaskEnum = "/Script/MeshModelingToolsExp.EBakeMapType",
-		ValidEnumValues="None, AmbientOcclusion, Curvature, One, Zero",
+		ValidEnumValues="None, AmbientOcclusion, Curvature, Height, One, Zero",
 		EditCondition="OutputMode == EBakeVertexOutput::PerChannel", EditConditionHides))
 	int32 OutputTypeB = static_cast<int32>(EBakeMapType::None);
 
 	/** The bake output type to generate in the Alpha channel */
 	UPROPERTY(EditAnywhere, Category = BakeOutput, meta=(Bitmask, BitmaskEnum = "/Script/MeshModelingToolsExp.EBakeMapType",
-		ValidEnumValues="None, AmbientOcclusion, Curvature, One, Zero",
+		ValidEnumValues="None, AmbientOcclusion, Curvature, Height, One, Zero",
 		EditCondition="OutputMode == EBakeVertexOutput::PerChannel", EditConditionHides))
 	int32 OutputTypeA = static_cast<int32>(EBakeMapType::None);
 
