Modeling - Mesh Component Selection - Triangles tests

#rnx

#virtualized

diff --git a/Engine/Plugins/Editor/ModelingToolsEditorMode/Source/ModelingToolsEditorMode/Public/ModelingToolsEditorMode.h b/Engine/Plugins/Editor/ModelingToolsEditorMode/Source/ModelingToolsEditorMode/Public/ModelingToolsEditorMode.h
--- a/Engine/Plugins/Editor/ModelingToolsEditorMode/Source/ModelingToolsEditorMode/Public/ModelingToolsEditorMode.h
+++ b/Engine/Plugins/Editor/ModelingToolsEditorMode/Source/ModelingToolsEditorMode/Public/ModelingToolsEditorMode.h
@@ -27,7 +27,7 @@ class UGeometrySelectionManager;
 class UInteractiveCommand;
 class UBlueprint;
 
-UCLASS(Transient)
+UCLASS(Transient, MinimalAPI)
 class UModelingToolsEditorMode : public UBaseLegacyWidgetEdMode, public ILegacyEdModeSelectInterface
 {
 	GENERATED_BODY()
diff --git a/Engine/Plugins/Experimental/MeshModelingToolsetExp/Source/MeshModelingToolsExp/Private/Polymodeling/ExtrudeMeshSelectionTool.cpp b/Engine/Plugins/Experimental/MeshModelingToolsetExp/Source/MeshModelingToolsExp/Private/Polymodeling/ExtrudeMeshSelectionTool.cpp
--- a/Engine/Plugins/Experimental/MeshModelingToolsetExp/Source/MeshModelingToolsExp/Private/Polymodeling/ExtrudeMeshSelectionTool.cpp
+++ b/Engine/Plugins/Experimental/MeshModelingToolsetExp/Source/MeshModelingToolsExp/Private/Polymodeling/ExtrudeMeshSelectionTool.cpp
@@ -384,6 +384,25 @@ FBox UExtrudeMeshSelectionTool::GetWorldSpaceFocusBox()
 	return (FBox)CurBox;
 }
 
+FTransform UExtrudeMeshSelectionTool::GetTransform() const
+{
+	return TransformProxy->GetTransform();
+}
+
+void UExtrudeMeshSelectionTool::SetTransform(const FTransform& NewTransform)
+{
+	FTransformProxyChangeSource ChangeSource(TransformProxy);
+	ChangeSource.BeginChange();
+
+	TransformProxy->SetTransform(NewTransform);
+
+	TUniquePtr<FToolCommandChange> Change = ChangeSource.EndChange();
+	if (Change.IsValid())
+	{
+		GetToolManager()->EmitObjectChange(this, MoveTemp(Change), LOCTEXT("ExtrudeSelection", "Set Transform"));
+	}
+}
+
 void UExtrudeMeshSelectionTool::Initialize_GizmoMechanic()
 {
 	// Set up the gizmo.
diff --git a/Engine/Plugins/Experimental/MeshModelingToolsetExp/Source/MeshModelingToolsExp/Public/Polymodeling/ExtrudeMeshSelectionTool.h b/Engine/Plugins/Experimental/MeshModelingToolsetExp/Source/MeshModelingToolsExp/Public/Polymodeling/ExtrudeMeshSelectionTool.h
--- a/Engine/Plugins/Experimental/MeshModelingToolsetExp/Source/MeshModelingToolsExp/Public/Polymodeling/ExtrudeMeshSelectionTool.h
+++ b/Engine/Plugins/Experimental/MeshModelingToolsetExp/Source/MeshModelingToolsExp/Public/Polymodeling/ExtrudeMeshSelectionTool.h
@@ -137,6 +137,8 @@ public:
 	// IInteractiveToolCameraFocusAPI implementation
 	virtual FBox GetWorldSpaceFocusBox() override;
 
+	FTransform GetTransform() const;
+	void SetTransform(const FTransform& NewTransform);
 
 protected:
 	UPROPERTY()
