ModelingTools: Fix localization conflict in ExtrudeMeshSelectionTool

diff --git a/Engine/Plugins/Experimental/MeshModelingToolsetExp/Source/MeshModelingToolsExp/Private/Polymodeling/ExtrudeMeshSelectionTool.cpp b/Engine/Plugins/Experimental/MeshModelingToolsetExp/Source/MeshModelingToolsExp/Private/Polymodeling/ExtrudeMeshSelectionTool.cpp
--- a/Engine/Plugins/Experimental/MeshModelingToolsetExp/Source/MeshModelingToolsExp/Private/Polymodeling/ExtrudeMeshSelectionTool.cpp
+++ b/Engine/Plugins/Experimental/MeshModelingToolsetExp/Source/MeshModelingToolsExp/Private/Polymodeling/ExtrudeMeshSelectionTool.cpp
@@ -399,7 +399,7 @@ void UExtrudeMeshSelectionTool::SetTransform(const FTransform& NewTransform)
 	TUniquePtr<FToolCommandChange> Change = ChangeSource.EndChange();
 	if (Change.IsValid())
 	{
-		GetToolManager()->EmitObjectChange(this, MoveTemp(Change), LOCTEXT("ExtrudeSelection", "Set Transform"));
+		GetToolManager()->EmitObjectChange(this, MoveTemp(Change), LOCTEXT("SetTransform", "Set Transform"));
 	}
 }
 
