Modeling Tools: add ability to customize the display type of SToolInputAssetComboPanel

diff --git a/Engine/Plugins/Experimental/MeshModelingToolsetExp/Source/ModelingEditorUI/Private/ModelingWidgets/SToolInputAssetComboPanel.cpp b/Engine/Plugins/Experimental/MeshModelingToolsetExp/Source/ModelingEditorUI/Private/ModelingWidgets/SToolInputAssetComboPanel.cpp
--- a/Engine/Plugins/Experimental/MeshModelingToolsetExp/Source/ModelingEditorUI/Private/ModelingWidgets/SToolInputAssetComboPanel.cpp
+++ b/Engine/Plugins/Experimental/MeshModelingToolsetExp/Source/ModelingEditorUI/Private/ModelingWidgets/SToolInputAssetComboPanel.cpp
@@ -33,6 +33,7 @@ void SToolInputAssetComboPanel::Construct(const FArguments& InArgs)
 	this->CollectionSets = InArgs._CollectionSets;
 	this->OnSelectionChanged = InArgs._OnSelectionChanged;
 	RecentAssetsProvider = InArgs._RecentAssetsProvider;
+	this->ThumbnailDisplayMode = InArgs._ThumbnailDisplayMode;
 
 	PRAGMA_DISABLE_DEPRECATION_WARNINGS
 	for (FNamedCollectionList& CollectionSet : CollectionSets)
@@ -71,12 +72,49 @@ void SToolInputAssetComboPanel::Construct(const FArguments& InArgs)
 	ThumbnailConfig.bAllowRealTimeOnHovered = false;
 	this->AssetThumbnail = MakeShared<FAssetThumbnail>(FAssetData(), ComboButtonTileSize.X, ComboButtonTileSize.Y, ThumbnailPool);
 
-	TSharedRef<SBox> IconBox = SNew(SBox)
+	TSharedPtr<SWidget> AssetDisplay = nullptr;
+	
+	if (ThumbnailDisplayMode == EThumbnailDisplayMode::AssetThumbnail)
+	{
+		AssetDisplay = SNew(SBox)
 		.WidthOverride(ComboButtonTileSize.X)
 		.HeightOverride(ComboButtonTileSize.Y)
 		[
-			AssetThumbnail->MakeThumbnailWidget(ThumbnailConfig)
+			SAssignNew(ThumbnailBorder, SBorder)
+			.Padding(0)
+			.BorderImage(FStyleDefaults::GetNoBrush())
+			.HAlign(HAlign_Center)
+			.VAlign(VAlign_Center)
+			.OnMouseDoubleClick(this, &SToolInputAssetComboPanel::OnAssetThumbnailDoubleClick)
+			[
+				SNew(SBox)
+				.WidthOverride(ComboButtonTileSize.X)
+				.HeightOverride(ComboButtonTileSize.Y)
+				[
+					AssetThumbnail->MakeThumbnailWidget(ThumbnailConfig)
+				]
+			]
+		];
+	}
+	else if (ThumbnailDisplayMode == EThumbnailDisplayMode::AssetName)
+	{
+		AssetDisplay = SNew(SBox)
+		[
+			SAssignNew(ThumbnailBorder, SBorder)
+			.Padding(0)
+			.BorderImage(FStyleDefaults::GetNoBrush())
+			.HAlign(HAlign_Center)
+			.VAlign(VAlign_Center)
+			.OnMouseDoubleClick(this, &SToolInputAssetComboPanel::OnAssetThumbnailDoubleClick)
+			[
+				SNew(STextBlock)
+					.Text_Lambda([this]()
+					{
+						return FText::FromName(AssetThumbnail->GetAssetData().AssetName);
+					})
+			]
 		];
+	}
 
 	ComboButton = SNew(SComboButton)
 		.ToolTipText(UseTooltipText)
@@ -106,36 +144,25 @@ void SToolInputAssetComboPanel::Construct(const FArguments& InArgs)
 						+SOverlay::Slot()
 						//.Padding(1)
 						[
-							SNew(SVerticalBox)
-							// Thumbnail
-							+ SVerticalBox::Slot()
-							.AutoHeight()
-							.HAlign(HAlign_Center)
-							.VAlign(VAlign_Center)
+							SNew(SBox)
+							.WidthOverride(75.f)
 							[
-								SNew(SBox)
-								.WidthOverride(ComboButtonTileSize.X)
-								.HeightOverride(ComboButtonTileSize.Y)
+								SNew(SVerticalBox)
+								// Thumbnail or Label
+								+ SVerticalBox::Slot()
+								.AutoHeight()
+								.HAlign(HAlign_Center)
+								.VAlign(VAlign_Center)
 								[
-									SAssignNew(ThumbnailBorder, SBorder)
-									.Padding(0)
-									.BorderImage(FStyleDefaults::GetNoBrush())
-									.HAlign(HAlign_Center)
-									.VAlign(VAlign_Center)
-									.OnMouseDoubleClick(this, &SToolInputAssetComboPanel::OnAssetThumbnailDoubleClick)
-									[
-										IconBox
-									]
+									AssetDisplay.ToSharedRef()
 								]
 							]
 						]
 
 					]
-
 			]
 		];
 
-
 	EThumbnailLabel::Type TileThumbnailLabel = InArgs._AssetThumbnailLabel;
 	bool bForceShowEngineContent = InArgs._bForceShowEngineContent;
 	bool bForceShowPluginContent = InArgs._bForceShowPluginContent;
@@ -265,6 +292,13 @@ void SToolInputAssetComboPanel::RefreshThumbnailFromProperty()
 	}
 }
 
+void SToolInputAssetComboPanel::RefreshThumbnail(const FAssetData& InAssetData) const
+{
+	if (InAssetData.IsValid())
+	{
+		AssetThumbnail->SetAsset(InAssetData);
+	}
+}
 
 TSharedRef<SWidget> SToolInputAssetComboPanel::MakeCollectionSetsButtonPanel(TSharedRef<SToolInputAssetPicker> AssetPickerView)
 {
diff --git a/Engine/Plugins/Experimental/MeshModelingToolsetExp/Source/ModelingEditorUI/Public/ModelingWidgets/SToolInputAssetComboPanel.h b/Engine/Plugins/Experimental/MeshModelingToolsetExp/Source/ModelingEditorUI/Public/ModelingWidgets/SToolInputAssetComboPanel.h
--- a/Engine/Plugins/Experimental/MeshModelingToolsetExp/Source/ModelingEditorUI/Public/ModelingWidgets/SToolInputAssetComboPanel.h
+++ b/Engine/Plugins/Experimental/MeshModelingToolsetExp/Source/ModelingEditorUI/Public/ModelingWidgets/SToolInputAssetComboPanel.h
@@ -21,6 +21,12 @@ class ITableRow;
 class STableViewBase;
 class SToolInputAssetPicker;
 
+enum class EThumbnailDisplayMode
+{
+	AssetThumbnail,
+	AssetName
+};
+
 /**
 * SToolInputAssetComboPanel provides a similar UI to SComboPanel but
 * specifically for picking Assets. The standard widget is a SComboButton
@@ -87,6 +93,7 @@ public:
 		, _bForceShowEngineContent(false)
 		, _bForceShowPluginContent(false)
 		, _AssetViewType(EAssetViewType::Tile)
+		, _ThumbnailDisplayMode(EThumbnailDisplayMode::AssetThumbnail)
 	{
 		}
 
@@ -132,6 +139,9 @@ public:
 		/** Sets the type of display we want the asset view to take. */
 		SLATE_ARGUMENT(EAssetViewType::Type, AssetViewType)
 
+		/** Display mode of the asset thumbnail */
+		SLATE_ARGUMENT(EThumbnailDisplayMode, ThumbnailDisplayMode)
+
 	SLATE_END_ARGS()
 
 
@@ -142,7 +152,10 @@ public:
 	 */
 	void Construct( const FArguments& InArgs );
 
+	// refresh the thumbnail using the AssetComboPanel's Property
 	void RefreshThumbnailFromProperty();
+	// refresh the thumbnail using provided AssetData, for cases where Property may not be available
+	void RefreshThumbnail(const FAssetData& InAssetData) const;
 
 protected:
 
@@ -151,6 +164,7 @@ protected:
 	FVector2D FlyoutSize;
 	TSharedPtr<IPropertyHandle> Property;
 	UClass* AssetClassType = nullptr;
+	EThumbnailDisplayMode ThumbnailDisplayMode;
 
 
 	/** Delegate to invoke when selection changes. */
