ModelingTools: Fixed Union tool only transferring UV Channel 0 in the result mesh.

diff --git a/Engine/Plugins/Experimental/MeshModelingToolsetExp/Source/MeshModelingToolsExp/Private/SelfUnionMeshesTool.cpp b/Engine/Plugins/Experimental/MeshModelingToolsetExp/Source/MeshModelingToolsExp/Private/SelfUnionMeshesTool.cpp
--- a/Engine/Plugins/Experimental/MeshModelingToolsetExp/Source/MeshModelingToolsExp/Private/SelfUnionMeshesTool.cpp
+++ b/Engine/Plugins/Experimental/MeshModelingToolsetExp/Source/MeshModelingToolsExp/Private/SelfUnionMeshesTool.cpp
@@ -125,6 +125,11 @@ void USelfUnionMeshesTool::ConvertInputsAndSetPreviewMaterials(bool bSetPreviewM
 		{
 			ComponentMesh.ReverseOrientation(false);
 		}
+
+		// ensure all UV layers are transferred
+		const int NumUVLayers = FMath::Max(CombinedSourceMeshes->Attributes()->NumUVLayers(), ComponentMesh.Attributes()->NumUVLayers());
+		CombinedSourceMeshes->Attributes()->SetNumUVLayers(NumUVLayers);
+		
 		FMeshIndexMappings IndexMaps;
 		AppendEditor.AppendMesh(&ComponentMesh, IndexMaps,
 			[WorldTransform, this](int VID, const FVector3d& Pos)
