Add Modeling - Create tests (Box, CubeGrid, Extrude Polygon, Revolve Path)

#rnx

diff --git a/Engine/Plugins/Experimental/MeshModelingToolsetExp/Source/MeshModelingToolsExp/Private/CubeGridTool.cpp b/Engine/Plugins/Experimental/MeshModelingToolsetExp/Source/MeshModelingToolsExp/Private/CubeGridTool.cpp
--- a/Engine/Plugins/Experimental/MeshModelingToolsetExp/Source/MeshModelingToolsExp/Private/CubeGridTool.cpp
+++ b/Engine/Plugins/Experimental/MeshModelingToolsetExp/Source/MeshModelingToolsExp/Private/CubeGridTool.cpp
@@ -43,6 +43,7 @@
 #include "TargetInterfaces/PrimitiveComponentBackedTarget.h"
 #include "ToolSceneQueriesUtil.h"
 #include "ToolSetupUtil.h"
+#include "MeshOpPreviewHelpers.h"
 
 #include UE_INLINE_GENERATED_CPP_BY_NAME(CubeGridTool)
 
@@ -1472,7 +1473,7 @@ void UCubeGridTool::OnTick(float DeltaTime)
 		// Note that the check for bPreviewMayDiffer is needed, because when we do a forced
 		// reset, HaveValidResult() returns false, yet we may still want to do an application
 		// to slide the selection.
-		(!bPreviewMayDiffer || Preview->HaveValidResult()))
+		CanApplyPreview())
 	{
 		ApplyPreview();
 	}
@@ -2685,6 +2686,11 @@ FBox UCubeGridTool::GetWorldSpaceFocusBox()
 	return (FBox)Bounds;
 }
 
+bool UCubeGridTool::CanApplyPreview() const
+{
+	return !bPreviewMayDiffer || Preview->HaveValidResult();
+}
+
 void UCubeGridTool::StartCornerMode()
 {
 	using namespace CubeGridToolLocals;
diff --git a/Engine/Plugins/Experimental/MeshModelingToolsetExp/Source/MeshModelingToolsExp/Public/CubeGridTool.h b/Engine/Plugins/Experimental/MeshModelingToolsetExp/Source/MeshModelingToolsExp/Public/CubeGridTool.h
--- a/Engine/Plugins/Experimental/MeshModelingToolsetExp/Source/MeshModelingToolsExp/Public/CubeGridTool.h
+++ b/Engine/Plugins/Experimental/MeshModelingToolsetExp/Source/MeshModelingToolsExp/Public/CubeGridTool.h
@@ -427,6 +427,8 @@ public:
 		bHasInitialGridPivot = true;
 	}
 
+	UE_API bool CanApplyPreview() const;
+
 protected:
 
 	UPROPERTY()
diff --git a/Engine/Plugins/Runtime/MeshModelingToolset/Source/MeshModelingTools/Private/DrawAndRevolveTool.cpp b/Engine/Plugins/Runtime/MeshModelingToolset/Source/MeshModelingTools/Private/DrawAndRevolveTool.cpp
--- a/Engine/Plugins/Runtime/MeshModelingToolset/Source/MeshModelingTools/Private/DrawAndRevolveTool.cpp
+++ b/Engine/Plugins/Runtime/MeshModelingToolset/Source/MeshModelingTools/Private/DrawAndRevolveTool.cpp
@@ -359,5 +359,14 @@ void UDrawAndRevolveTool::Render(IToolsContextRenderAPI* RenderAPI)
 	}
 }
 
+bool UDrawAndRevolveTool::CanGenerateAsset() const
+{
+	if (Preview)
+	{
+		return 0 < Preview->PreviewMesh->GetPreviewDynamicMesh()->TriangleCount();
+	}
+	return false;
+}
+
 #undef LOCTEXT_NAMESPACE
 
diff --git a/Engine/Plugins/Runtime/MeshModelingToolset/Source/MeshModelingTools/Public/DrawAndRevolveTool.h b/Engine/Plugins/Runtime/MeshModelingToolset/Source/MeshModelingTools/Public/DrawAndRevolveTool.h
--- a/Engine/Plugins/Runtime/MeshModelingToolset/Source/MeshModelingTools/Public/DrawAndRevolveTool.h
+++ b/Engine/Plugins/Runtime/MeshModelingToolset/Source/MeshModelingTools/Public/DrawAndRevolveTool.h
@@ -115,6 +115,8 @@ public:
 		InitialDrawFrame = InFrame;
 	}
 
+	UE_API bool CanGenerateAsset() const;
+
 protected:
 
 	UWorld* TargetWorld;
