ModelingTools: Fix crash in UV transfer tools when there's a  mismatch in number of UV channels in source and destination.

diff --git a/Engine/Plugins/Experimental/MeshModelingToolsetExp/Source/MeshModelingToolsExp/Private/UVTransferTool.cpp b/Engine/Plugins/Experimental/MeshModelingToolsetExp/Source/MeshModelingToolsExp/Private/UVTransferTool.cpp
--- a/Engine/Plugins/Experimental/MeshModelingToolsetExp/Source/MeshModelingToolsExp/Private/UVTransferTool.cpp
+++ b/Engine/Plugins/Experimental/MeshModelingToolsetExp/Source/MeshModelingToolsExp/Private/UVTransferTool.cpp
@@ -280,7 +280,9 @@ void UUVTransferTool::ReinitializePreviews()
 	SourcePreview->UpdatePreview(Meshes[SourceIndex].Get());
 	SourcePreview->SetTransform((FTransform)UE::ToolTarget::GetLocalToWorldTransform(Targets[SourceIndex]));
 
-	UVChannelProperties->Initialize(Meshes[SourceIndex].Get(), false);
+	UVChannelProperties->Initialize(FMath::Min(
+		Meshes[SourceIndex] && Meshes[SourceIndex]->HasAttributes() ? Meshes[SourceIndex]->Attributes()->NumUVLayers() : 0,
+		Meshes[DestinationIndex] && Meshes[DestinationIndex]->HasAttributes() ? Meshes[DestinationIndex]->Attributes()->NumUVLayers() : 0), true);
 	UVChannelProperties->ValidateSelection(true);
 	DestinationMaterialSettings->UpdateUVChannels(
 		UVChannelProperties->GetSelectedChannelIndex(),
