ModelingTools: Give Transfer UVs Tool the ability to operate on different source/destination channel indices.

diff --git a/Engine/Plugins/Experimental/MeshModelingToolsetExp/Source/MeshModelingToolsExp/Private/UVTransferTool.cpp b/Engine/Plugins/Experimental/MeshModelingToolsetExp/Source/MeshModelingToolsExp/Private/UVTransferTool.cpp
--- a/Engine/Plugins/Experimental/MeshModelingToolsetExp/Source/MeshModelingToolsExp/Private/UVTransferTool.cpp
+++ b/Engine/Plugins/Experimental/MeshModelingToolsetExp/Source/MeshModelingToolsExp/Private/UVTransferTool.cpp
@@ -9,7 +9,6 @@
 #include "ModelingToolTargetUtil.h"
 #include "Parameterization/UVTransfer.h"
 #include "Properties/MeshMaterialProperties.h"
-#include "Properties/MeshUVChannelProperties.h"
 #include "Selections/GeometrySelectionUtil.h" // EnumerateSelectionTriangles
 #include "TargetInterfaces/DynamicMeshCommitter.h"
 #include "TargetInterfaces/DynamicMeshProvider.h"
@@ -35,7 +34,8 @@ namespace UVTransferToolLocals
 		TSharedPtr<const FDynamicMesh3> DestinationMesh;
 		TOptional<TSet<int32>> SourceSelectionTids;
 		TOptional<TSet<int32>> DestinationSelectionTids;
-		int UVLayerIndex = 0;
+		int32 SourceUVLayerIndex = 0;
+		int32 DestUVLayerIndex = 0;
 		bool bTransferSeamsOnly = false;
 		bool bClearExistingSeams = true;
 
@@ -53,9 +53,12 @@ namespace UVTransferToolLocals
 			if (Progress && Progress->Cancelled()) { return; }
 
 			if (!DestinationMesh || !SourceMesh
-				|| UVLayerIndex < 0 || UVLayerIndex > 7
+				|| SourceUVLayerIndex < 0 || SourceUVLayerIndex > 7
+				|| DestUVLayerIndex < 0 || DestUVLayerIndex > 7
 				|| !SourceMesh->HasAttributes()
-				|| SourceMesh->Attributes()->NumUVLayers() <= UVLayerIndex)
+				|| SourceMesh->Attributes()->NumUVLayers() <= SourceUVLayerIndex
+				|| !DestinationMesh->HasAttributes()
+				|| DestinationMesh->Attributes()->NumUVLayers() <= DestUVLayerIndex)
 			{
 				ensure(false);
 				return;
@@ -65,7 +68,8 @@ namespace UVTransferToolLocals
 
 			if (Progress && Progress->Cancelled()) { return; }
 
-			UE::Geometry::FDynamicMeshUVTransfer Transferer(SourceMesh.Get(), ResultMesh.Get(), UVLayerIndex);
+			UE::Geometry::FDynamicMeshUVTransfer Transferer(SourceMesh.Get(), ResultMesh.Get(), 
+				SourceUVLayerIndex, DestUVLayerIndex);
 
 			Transferer.VertexSearchDistance = VertexSearchDistance;
 			Transferer.VertexSearchCellSize = VertexSearchCellSize;
@@ -99,7 +103,6 @@ void UUVTransferTool::OnShutdown(EToolShutdownType ShutdownType)
 
 	Settings->SaveProperties(this);
 	DestinationMaterialSettings->SaveProperties(this, CacheIdentifier);
-	UVChannelProperties->SaveProperties(this);
 
 	SourceSeamVisualizer->Disconnect();
 	DestinationSeamVisualizer->Disconnect();
@@ -118,7 +121,6 @@ void UUVTransferTool::OnShutdown(EToolShutdownType ShutdownType)
 	SourcePreview->Disconnect();
 
 	Settings = nullptr;
-	UVChannelProperties = nullptr;
 	DestinationMaterialSettings = nullptr;
 	DestinationPreview = nullptr;
 	SourcePreview = nullptr;
@@ -148,11 +150,6 @@ void UUVTransferTool::Setup()
 	DestinationMaterialSettings->RestoreProperties(this, CacheIdentifier);
 	AddToolPropertySource(DestinationMaterialSettings);
 
-	UVChannelProperties = NewObject<UMeshUVChannelProperties>(this);
-	UVChannelProperties->RestoreProperties(this);
-	AddToolPropertySource(UVChannelProperties);
-	UVChannelProperties->ValidateSelection(true);
-
 	// Hide inputs
 	for (TObjectPtr<UToolTarget> Target : Targets)
 	{
@@ -206,11 +203,19 @@ void UUVTransferTool::Setup()
 	
 	ReinitializePreviews();
 
-	UVChannelProperties->WatchProperty(UVChannelProperties->UVChannel, [this](const FString& NewValue)
+	Settings->WatchProperty(Settings->SourceUVChannel, [this](const FString& NewValue)
+	{
+		InvalidatePreview();
+	});
+	Settings->WatchProperty(Settings->bSameDestinationChannel, [this](bool NewValue)
+	{
+		ReinitializeUVChannelOptions();
+		InvalidatePreview();
+	});
+	Settings->WatchProperty(Settings->DestUVChannel, [this](const FString& NewValue)
 	{
-		DestinationMaterialSettings->UpdateUVChannels(
-			UVChannelProperties->UVChannelNamesList.IndexOfByKey(UVChannelProperties->UVChannel),
-			UVChannelProperties->UVChannelNamesList);
+		DestinationMaterialSettings->UpdateUVChannels(Settings->GetDestUVChannel(),
+			Settings->GetDestUVChannelNames());
 		InvalidatePreview();
 	});
 
@@ -252,7 +257,6 @@ void UUVTransferTool::Setup()
 	Settings->WatchProperty(Settings->bShowSeams, [this](bool) { UpdateVisualizations(); });
 
 	Settings->SilentUpdateWatched();
-	UVChannelProperties->SilentUpdateWatched();
 	DestinationMaterialSettings->SilentUpdateWatched();
 
 	SetToolDisplayName(LOCTEXT("ToolName", "UV Transfer"));
@@ -264,6 +268,8 @@ void UUVTransferTool::Setup()
 
 void UUVTransferTool::ReinitializePreviews()
 {
+	ReinitializeUVChannelOptions();
+
 	int SourceIndex = Settings->bReverseDirection ? 1 : 0;
 	int DestinationIndex = Settings->bReverseDirection ? 0 : 1;
 
@@ -280,13 +286,9 @@ void UUVTransferTool::ReinitializePreviews()
 	SourcePreview->UpdatePreview(Meshes[SourceIndex].Get());
 	SourcePreview->SetTransform((FTransform)UE::ToolTarget::GetLocalToWorldTransform(Targets[SourceIndex]));
 
-	UVChannelProperties->Initialize(FMath::Min(
-		Meshes[SourceIndex] && Meshes[SourceIndex]->HasAttributes() ? Meshes[SourceIndex]->Attributes()->NumUVLayers() : 0,
-		Meshes[DestinationIndex] && Meshes[DestinationIndex]->HasAttributes() ? Meshes[DestinationIndex]->Attributes()->NumUVLayers() : 0), true);
-	UVChannelProperties->ValidateSelection(true);
 	DestinationMaterialSettings->UpdateUVChannels(
-		UVChannelProperties->GetSelectedChannelIndex(),
-		UVChannelProperties->UVChannelNamesList);
+		Settings->GetDestUVChannel(),
+		Settings->GetDestUVChannelNames());
 
 	DestinationMaterialSettings->UpdateMaterials();
 	DestinationPreview->OverrideMaterial = DestinationMaterialSettings->GetActiveOverrideMaterial();
@@ -321,6 +323,26 @@ void UUVTransferTool::UpdateVisualizations()
 	}
 }
 
+void UUVTransferTool::ReinitializeUVChannelOptions()
+{
+	if (!ensure(Settings && Meshes[0] && Meshes[1] && Meshes[0]->HasAttributes() && Meshes[1]->HasAttributes()))
+	{
+		return;
+	}
+
+	if (Settings->bSameDestinationChannel)
+	{
+		int32 NumChannels = FMath::Min(Meshes[0]->Attributes()->NumUVLayers(), Meshes[1]->Attributes()->NumUVLayers());
+		Settings->ReinitializeChannelOptions(NumChannels, NumChannels);
+	}
+	else
+	{
+		int SourceIndex = Settings->bReverseDirection ? 1 : 0;
+		int DestinationIndex = Settings->bReverseDirection ? 0 : 1;
+		Settings->ReinitializeChannelOptions(Meshes[SourceIndex]->Attributes()->NumUVLayers(), Meshes[DestinationIndex]->Attributes()->NumUVLayers());
+	}
+}
+
 void UUVTransferTool::OnTick(float DeltaTime)
 {
 	if (DestinationPreview)
@@ -374,7 +396,8 @@ TUniquePtr<UE::Geometry::FDynamicMeshOperator> UUVTransferTool::MakeNewOperator(
 	//	Op->DestinationSelectionTids = SelectionTidSets[DestinationIndex];
 	//}
 
-	Op->UVLayerIndex = UVChannelProperties->GetSelectedChannelIndex(true);
+	Op->SourceUVLayerIndex = Settings->GetSourceUVChannel();
+	Op->DestUVLayerIndex = Settings->GetDestUVChannel();
 	Op->bTransferSeamsOnly = Settings->bTransferSeamsOnly;
 	Op->bClearExistingSeams = Settings->bClearExistingSeams;
 
@@ -425,4 +448,55 @@ const FToolTargetTypeRequirements& UUVTransferToolBuilder::GetTargetRequirements
 	return TypeRequirements;
 }
 
-#undef LOCTEXT_NAMESPACE
+void UUVTransferToolProperties::ReinitializeChannelOptions(int32 NumSourceChannels, int32 NumDestChannels)
+{
+	SourceUVChannels.Reset();
+	DestUVChannels.Reset();
+
+	auto InitializeChannelNames = [](int32 NumChannelsIn, TArray<FString>& ChannelsOut)
+	{
+		for (int32 i = 0; i < NumChannelsIn; ++i)
+		{
+			ChannelsOut.Add(FString::Printf(TEXT("UV %d"), i));
+		}
+	};
+	InitializeChannelNames(NumSourceChannels, SourceUVChannels);
+	InitializeChannelNames(NumDestChannels, DestUVChannels);
+
+	auto GetUpdatedCurrentChannel = [](const TArray<FString>& Channels, const FString& CurrentChannelIn) -> FString
+	{
+		if (Channels.IsEmpty())
+		{
+			return TEXT("");
+		}
+
+		if (CurrentChannelIn.IsEmpty())
+		{
+			return Channels[0];
+		}
+
+		int32 FoundIndex = Channels.IndexOfByKey(CurrentChannelIn);
+		if (FoundIndex < 0)
+		{
+			// Probably a UV layer past the one we have, so use last layer
+			return Channels.Last();
+		}
+
+		return Channels[FoundIndex];
+	};
+
+	SourceUVChannel = GetUpdatedCurrentChannel(SourceUVChannels, SourceUVChannel);
+	DestUVChannel = GetUpdatedCurrentChannel(DestUVChannels, DestUVChannel);
+}
+
+int32 UUVTransferToolProperties::GetSourceUVChannel()
+{
+	return SourceUVChannels.IndexOfByKey(SourceUVChannel);
+}
+
+int32 UUVTransferToolProperties::GetDestUVChannel()
+{
+	return bSameDestinationChannel ? GetSourceUVChannel() : DestUVChannels.IndexOfByKey(DestUVChannel);
+}
+
+#undef LOCTEXT_NAMESPACE
\ No newline at end of file
diff --git a/Engine/Plugins/Experimental/MeshModelingToolsetExp/Source/MeshModelingToolsExp/Public/UVTransferTool.h b/Engine/Plugins/Experimental/MeshModelingToolsetExp/Source/MeshModelingToolsExp/Public/UVTransferTool.h
--- a/Engine/Plugins/Experimental/MeshModelingToolsetExp/Source/MeshModelingToolsExp/Public/UVTransferTool.h
+++ b/Engine/Plugins/Experimental/MeshModelingToolsetExp/Source/MeshModelingToolsExp/Public/UVTransferTool.h
@@ -3,6 +3,7 @@
 #pragma once
 
 #include "BaseTools/MultiTargetWithSelectionTool.h" 
+#include "Containers/UnrealString.h"
 #include "GeometryBase.h"
 #include "InteractiveToolBuilder.h"
 #include "Misc/Optional.h"
@@ -13,7 +14,6 @@
 #define UE_API MESHMODELINGTOOLSEXP_API
 
 struct FDynamicMeshOpResult;
-class UMeshUVChannelProperties;
 class UExistingMeshMaterialProperties;
 class UMeshElementsVisualizer;
 class UMeshOpPreviewWithBackgroundCompute;
@@ -83,6 +83,39 @@ public:
 	UPROPERTY(EditAnywhere, Category = Options, AdvancedDisplay, meta = (
 		ClampMin = 0.0001, ClampMax = 10000, UIMax = 1.0))
 	double VertexSearchDistance = 0.0001;
+
+private:
+	UPROPERTY(EditAnywhere, Category = "Options", meta = (DisplayName = "Source UV Channel", 
+		GetOptions = GetSourceUVChannelNames, NoResetToDefault))
+	FString SourceUVChannel;
+
+	/** 
+	 * When true, the source channel is the same index in the source and destination, and options
+	 *  are limited to channels that exist on both meshes.
+	 */
+	UPROPERTY(EditAnywhere, Category = "Options")
+	bool bSameDestinationChannel = true;
+	
+	/** UV Channel to use as destination, if Same Destination Channel is false. */
+	UPROPERTY(EditAnywhere, Category = "Options", meta = (DisplayName = "Destination UV Channel", 
+		GetOptions = GetDestUVChannelNames, NoResetToDefault,
+		EditCondition = "!bSameDestinationChannel", EditConditionHides))
+	FString DestUVChannel;
+
+	UFUNCTION()
+	const TArray<FString>& GetSourceUVChannelNames() const { return SourceUVChannels; }
+
+	UFUNCTION()
+	const TArray<FString>& GetDestUVChannelNames() const { return DestUVChannels; }
+	
+	void ReinitializeChannelOptions(int32 NumSourceChannels, int32 NumDestChannels);
+	int32 GetSourceUVChannel();
+	int32 GetDestUVChannel();
+
+	TArray<FString> SourceUVChannels;
+	TArray<FString> DestUVChannels;
+
+	friend class UUVTransferTool;
 };
 
 
@@ -126,9 +159,6 @@ private:
 	UPROPERTY()
 	TObjectPtr<UUVTransferToolProperties> Settings = nullptr;
 
-	UPROPERTY()
-	TObjectPtr<UMeshUVChannelProperties> UVChannelProperties = nullptr;
-
 	UPROPERTY()
 	TObjectPtr<UExistingMeshMaterialProperties> DestinationMaterialSettings = nullptr;
 
@@ -148,11 +178,12 @@ private:
 	TSharedPtr<UE::Geometry::FDynamicMesh3> Meshes[2];
 	TOptional<TSet<int32>> SelectionTidSets[2];
 	
-	UE_API void ReinitializePreviews();
-	UE_API void UpdateVisualizations();
+	void ReinitializePreviews();
+	void UpdateVisualizations();
+	void ReinitializeUVChannelOptions();
 	
-	UE_API void InvalidatePreview();
-	UE_API void GenerateAsset(const FDynamicMeshOpResult& Result);
+	void InvalidatePreview();
+	void GenerateAsset(const FDynamicMeshOpResult& Result);
 };
 
 #undef UE_API
diff --git a/Engine/Plugins/Runtime/GeometryProcessing/Source/DynamicMesh/Private/Parameterization/UVTransfer.cpp b/Engine/Plugins/Runtime/GeometryProcessing/Source/DynamicMesh/Private/Parameterization/UVTransfer.cpp
--- a/Engine/Plugins/Runtime/GeometryProcessing/Source/DynamicMesh/Private/Parameterization/UVTransfer.cpp
+++ b/Engine/Plugins/Runtime/GeometryProcessing/Source/DynamicMesh/Private/Parameterization/UVTransfer.cpp
@@ -11,10 +11,11 @@
 #include "Util/ProgressCancel.h"
 
 UE::Geometry::FDynamicMeshUVTransfer::FDynamicMeshUVTransfer(
-	const FDynamicMesh3* SourceMeshIn, FDynamicMesh3* DestinationMeshIn, int32 UVLayerIndexIn)
+	const FDynamicMesh3* SourceMeshIn, FDynamicMesh3* DestinationMeshIn, int32 SourceUVLayerIndexIn, int32 DestUVLayerIndexIn)
 : SourceMesh(SourceMeshIn)
 , DestinationMesh(DestinationMeshIn)
-, UVLayerIndex(FMath::Clamp(UVLayerIndexIn, 0, 7))
+, UVLayerIndex(FMath::Clamp(SourceUVLayerIndexIn, 0, 7))
+, DestUVLayerIndex(FMath::Clamp(DestUVLayerIndexIn, 0, 7))
 {}
 
 bool UE::Geometry::FDynamicMeshUVTransfer::TransferSeams(FProgressCancel* Progress)
@@ -22,7 +23,10 @@ bool UE::Geometry::FDynamicMeshUVTransfer::TransferSeams(FProgressCancel* Progre
 	if (!DestinationMesh || !SourceMesh
 		|| UVLayerIndex < 0 || UVLayerIndex > 7
 		|| !SourceMesh->HasAttributes()
-		|| SourceMesh->Attributes()->NumUVLayers() <= UVLayerIndex)
+		|| SourceMesh->Attributes()->NumUVLayers() <= UVLayerIndex
+		|| DestUVLayerIndex < 0 || DestUVLayerIndex > 7
+		|| !DestinationMesh->HasAttributes()
+		|| DestinationMesh->Attributes()->NumUVLayers() <= DestUVLayerIndex)
 	{
 		return ensure(false);
 	}
@@ -33,7 +37,7 @@ bool UE::Geometry::FDynamicMeshUVTransfer::TransferSeams(FProgressCancel* Progre
 		return false;
 	}
 
-	DestOverlay = DestinationMesh->Attributes()->GetUVLayer(UVLayerIndex);
+	DestOverlay = DestinationMesh->Attributes()->GetUVLayer(DestUVLayerIndex);
 	if (!ensure(DestOverlay))
 	{
 		return false;
@@ -335,6 +339,12 @@ bool UE::Geometry::FDynamicMeshUVTransfer::PerformSeamTransfer(FProgressCancel*
 
 bool UE::Geometry::FDynamicMeshUVTransfer::PerformElementsTransfer(FProgressCancel* Progress)
 {
+	if (!SourceOverlay || !DestOverlay)
+	{
+		// Must have had invalid parameters
+		return false;
+	}
+
 	// We'll check Progress->Cancelled every 512 elements.
 	uint32 Count = 0;
 	uint32 IntervalCheckMask = (1 << 9) - 1;
diff --git a/Engine/Plugins/Runtime/GeometryProcessing/Source/DynamicMesh/Public/Parameterization/UVTransfer.h b/Engine/Plugins/Runtime/GeometryProcessing/Source/DynamicMesh/Public/Parameterization/UVTransfer.h
--- a/Engine/Plugins/Runtime/GeometryProcessing/Source/DynamicMesh/Public/Parameterization/UVTransfer.h
+++ b/Engine/Plugins/Runtime/GeometryProcessing/Source/DynamicMesh/Public/Parameterization/UVTransfer.h
@@ -22,12 +22,18 @@ class FDynamicMesh3;
 class FDynamicMeshUVTransfer
 {
 public:
+	FDynamicMeshUVTransfer(const FDynamicMesh3* SourceMeshIn, FDynamicMesh3* DestinationMeshIn,
+		int32 UVLayerIndexIn)
+		: FDynamicMeshUVTransfer(SourceMeshIn, DestinationMeshIn, UVLayerIndexIn, UVLayerIndexIn)
+	{}
+
 	DYNAMICMESH_API FDynamicMeshUVTransfer(const FDynamicMesh3* SourceMeshIn, FDynamicMesh3* DestinationMeshIn,
-		int32 UVLayerIndexIn);
+		int32 SourceUVLayerIndex, int32 DestUVLayerIndex);
 
 	const FDynamicMesh3* SourceMesh = nullptr;
 	FDynamicMesh3* DestinationMesh = nullptr;
-	int UVLayerIndex = 0;
+	int UVLayerIndex = 0; // source UV layer index
+	int DestUVLayerIndex = 0;
 
 	/** How far to search for a matching vertex on the destination mesh. */
 	double VertexSearchDistance = KINDA_SMALL_NUMBER;
