fix crash on starting modeling mode w/ an actor selected that has null components attached

diff --git a/Engine/Plugins/Experimental/MeshModelingToolsetExp/Source/MeshModelingToolsExp/Private/Splines/BaseMeshFromSplinesTool.cpp b/Engine/Plugins/Experimental/MeshModelingToolsetExp/Source/MeshModelingToolsExp/Private/Splines/BaseMeshFromSplinesTool.cpp
--- a/Engine/Plugins/Experimental/MeshModelingToolsetExp/Source/MeshModelingToolsExp/Private/Splines/BaseMeshFromSplinesTool.cpp
+++ b/Engine/Plugins/Experimental/MeshModelingToolsetExp/Source/MeshModelingToolsExp/Private/Splines/BaseMeshFromSplinesTool.cpp
@@ -264,7 +264,7 @@ bool UBaseMeshFromSplinesToolBuilder::CanBuildTool(const FToolBuilderState& Scen
 {
 	int32 NumSplines = ToolBuilderUtil::CountComponents(SceneState, [&](UActorComponent* Object) -> bool
 		{
-			return Object->IsA<USplineComponent>();
+			return Object && Object->IsA<USplineComponent>();
 		});
 	FIndex2i SupportedRange = GetSupportedSplineCountRange();
 	return (NumSplines >= SupportedRange.A && (SupportedRange.B == -1 || NumSplines <= SupportedRange.B));
@@ -274,7 +274,7 @@ void UBaseMeshFromSplinesToolBuilder::InitializeNewTool(UBaseMeshFromSplinesTool
 {
 	TArray<UActorComponent*> Components = ToolBuilderUtil::FindAllComponents(SceneState, [&](UActorComponent* Object)
 		{
-			return Object->IsA<USplineComponent>();
+			return Object && Object->IsA<USplineComponent>();
 		});
 	TArray<TWeakObjectPtr<AActor>> ActorsWithSplines;
 	TSet<AActor*> FoundActors;
