World Partition - HLOD: Properly compute static meshes UVChannelData upon HLOD builds for simplified & approximate HLODs
* Regression introduced in CL 41487634 when calls to PostEditChange() were replaced by calls to Build(), missing some parameters

diff --git a/Engine/Plugins/Editor/WorldPartitionHLODUtilities/Source/Private/WorldPartition/HLOD/Builders/HLODBuilderMeshApproximate.cpp b/Engine/Plugins/Editor/WorldPartitionHLODUtilities/Source/Private/WorldPartition/HLOD/Builders/HLODBuilderMeshApproximate.cpp
--- a/Engine/Plugins/Editor/WorldPartitionHLODUtilities/Source/Private/WorldPartition/HLOD/Builders/HLODBuilderMeshApproximate.cpp
+++ b/Engine/Plugins/Editor/WorldPartitionHLODUtilities/Source/Private/WorldPartition/HLOD/Builders/HLODBuilderMeshApproximate.cpp
@@ -47,7 +47,7 @@ uint32 UHLODBuilderMeshApproximateSettings::GetCRC() const
 	FArchiveCrc32 Ar;
 
 	// Base key, changing this will force a rebuild of all HLODs from this builder
-	FString HLODBaseKey = "4293829EF930455CA00E56A0B3F4E67F";
+	FString HLODBaseKey = "4E9BB187B38645A08B424895C5268B8A";
 	Ar << HLODBaseKey;
 
 	Ar << This.MeshApproximationSettings;
diff --git a/Engine/Plugins/Editor/WorldPartitionHLODUtilities/Source/Private/WorldPartition/HLOD/Builders/HLODBuilderMeshSimplify.cpp b/Engine/Plugins/Editor/WorldPartitionHLODUtilities/Source/Private/WorldPartition/HLOD/Builders/HLODBuilderMeshSimplify.cpp
--- a/Engine/Plugins/Editor/WorldPartitionHLODUtilities/Source/Private/WorldPartition/HLOD/Builders/HLODBuilderMeshSimplify.cpp
+++ b/Engine/Plugins/Editor/WorldPartitionHLODUtilities/Source/Private/WorldPartition/HLOD/Builders/HLODBuilderMeshSimplify.cpp
@@ -44,7 +44,7 @@ uint32 UHLODBuilderMeshSimplifySettings::GetCRC() const
 	FArchiveCrc32 Ar;
 
 	// Base mesh simplify key, changing this will force a rebuild of all HLODs from this builder
-	FString HLODBaseKey = "AD4D6BE5CFA5447BB43DD31171A429BF";
+	FString HLODBaseKey = "833ECCC91448424CACC620661712928A";
 	Ar << HLODBaseKey;
 
 	Ar << This.MeshSimplifySettings;
diff --git a/Engine/Plugins/Experimental/MeshModelingToolsetExp/Source/GeometryProcessingAdapters/Private/ApproximateActorsImpl.cpp b/Engine/Plugins/Experimental/MeshModelingToolsetExp/Source/GeometryProcessingAdapters/Private/ApproximateActorsImpl.cpp
--- a/Engine/Plugins/Experimental/MeshModelingToolsetExp/Source/GeometryProcessingAdapters/Private/ApproximateActorsImpl.cpp
+++ b/Engine/Plugins/Experimental/MeshModelingToolsetExp/Source/GeometryProcessingAdapters/Private/ApproximateActorsImpl.cpp
@@ -1427,6 +1427,9 @@ UStaticMesh* FApproximateActorsImpl::EmitGeneratedMeshAsset(
 
 	
 	UStaticMesh::FBuildParameters BuildParameters;
+	BuildParameters.bInSilent = true;
+	BuildParameters.bInRebuildUVChannelData = true;
+	BuildParameters.bInEnforceLightmapRestrictions = true;
 
 	// Nanite meshes can sometime contain smaller triangles that could be discarded in the proxy mesh.
 	// If those are found on the outskirt of a merged mesh, the static mesh build will complain about
diff --git a/Engine/Source/Developer/MeshMergeUtilities/Private/ProxyGenerationProcessor.cpp b/Engine/Source/Developer/MeshMergeUtilities/Private/ProxyGenerationProcessor.cpp
--- a/Engine/Source/Developer/MeshMergeUtilities/Private/ProxyGenerationProcessor.cpp
+++ b/Engine/Source/Developer/MeshMergeUtilities/Private/ProxyGenerationProcessor.cpp
@@ -361,6 +361,9 @@ void FProxyGenerationProcessor::ProcessJob(const FGuid& JobGuid, FProxyGeneratio
 	}
 
 	UStaticMesh::FBuildParameters BuildParameters;
+	BuildParameters.bInSilent = true;
+	BuildParameters.bInRebuildUVChannelData = true;
+	BuildParameters.bInEnforceLightmapRestrictions = true;
 
 	// Nanite meshes can sometime contain smaller triangles that could be discarded in the proxy mesh.
 	// If those are found on the outskirt of a merged mesh, the static mesh build will complain about
