ModelingTools: Relax target requirements for the Inspect tool.

diff --git a/Engine/Plugins/Experimental/MeshModelingToolsetExp/Source/MeshModelingToolsExp/Private/MeshInspectorTool.cpp b/Engine/Plugins/Experimental/MeshModelingToolsetExp/Source/MeshModelingToolsExp/Private/MeshInspectorTool.cpp
--- a/Engine/Plugins/Experimental/MeshModelingToolsetExp/Source/MeshModelingToolsExp/Private/MeshInspectorTool.cpp
+++ b/Engine/Plugins/Experimental/MeshModelingToolsetExp/Source/MeshModelingToolsExp/Private/MeshInspectorTool.cpp
@@ -27,8 +27,8 @@
 #include "AssetUtils/MeshDescriptionUtil.h"
 #include "ModelingToolTargetUtil.h"
 #include "TargetInterfaces/MaterialProvider.h"
-#include "TargetInterfaces/MeshDescriptionProvider.h"
-#include "TargetInterfaces/PrimitiveComponentBackedTarget.h"
+#include "TargetInterfaces/DynamicMeshProvider.h"
+#include "TargetInterfaces/SceneComponentBackedTarget.h"
 
 #include "CanvasTypes.h"
 #include "CanvasItem.h"
@@ -169,8 +169,8 @@ UMaterialInterface* UMeshInspectorMaterialProperties::GetActiveOverrideMaterial(
 const FToolTargetTypeRequirements& UMeshInspectorToolBuilder::GetTargetRequirements() const
 {
 	static FToolTargetTypeRequirements TypeRequirements({
-		UMeshDescriptionProvider::StaticClass(),
-		UPrimitiveComponentBackedTarget::StaticClass(),
+		UDynamicMeshProvider::StaticClass(),
+		USceneComponentBackedTarget::StaticClass(),
 		UMaterialProvider::StaticClass()
 		});
 	return TypeRequirements;
