Modeling Tools: add support for new Mesh Layers panel

diff --git a/Engine/Plugins/Editor/ModelingToolsEditorMode/Source/ModelingToolsEditorMode/Private/DetailsCustomizations/MeshVertexSculptToolCustomizations.cpp b/Engine/Plugins/Editor/ModelingToolsEditorMode/Source/ModelingToolsEditorMode/Private/DetailsCustomizations/MeshVertexSculptToolCustomizations.cpp
--- a/Engine/Plugins/Editor/ModelingToolsEditorMode/Source/ModelingToolsEditorMode/Private/DetailsCustomizations/MeshVertexSculptToolCustomizations.cpp
+++ b/Engine/Plugins/Editor/ModelingToolsEditorMode/Source/ModelingToolsEditorMode/Private/DetailsCustomizations/MeshVertexSculptToolCustomizations.cpp
@@ -18,6 +18,10 @@
 #include "ModelingToolsEditorModeSettings.h"
 
 #include "MeshVertexSculptTool.h"
+#include "ModelingWidgets/SculptLayersController.h"
+#include "ModelingWidgets/SMeshLayersStack.h"
+#include "Properties/MeshSculptLayerProperties.h"
+
 #include "Widgets/Layout/SBox.h"
 #include "Widgets/SBoxPanel.h"
 #include "Sculpting/KelvinletBrushOp.h"
@@ -114,6 +118,41 @@ void FSculptBrushPropertiesDetails::CustomizeDetails(IDetailLayoutBuilder& Detai
 
 }
 
+FMeshSculptLayerPropertiesDetails::FMeshSculptLayerPropertiesDetails()
+{
+	Controller = MakeShared<FSculptLayersController>();
+}
+
+TSharedRef<IDetailCustomization> FMeshSculptLayerPropertiesDetails::MakeInstance()
+{
+	return MakeShareable(new FMeshSculptLayerPropertiesDetails);
+}
+
+void FMeshSculptLayerPropertiesDetails::CustomizeDetails(IDetailLayoutBuilder& DetailBuilder)
+{
+	TArray<TWeakObjectPtr<UObject>> ObjectsBeingCustomized;
+	DetailBuilder.GetObjectsBeingCustomized(ObjectsBeingCustomized);
+	check(ObjectsBeingCustomized.Num() > 0);
+	UMeshSculptLayerProperties* SculptLayersProperties = Cast<UMeshSculptLayerProperties>(ObjectsBeingCustomized[0]);
+
+	Controller->SetProperties(SculptLayersProperties);
+
+	TSharedPtr<IPropertyHandle> ActiveLayerHandle = DetailBuilder.GetProperty(GET_MEMBER_NAME_CHECKED(UMeshSculptLayerProperties, ActiveLayer), UMeshSculptLayerProperties::StaticClass());
+	ensure(ActiveLayerHandle->IsValidHandle());
+	
+	TSharedPtr<IPropertyHandle> LayerWeightsHandle = DetailBuilder.GetProperty(GET_MEMBER_NAME_CHECKED(UMeshSculptLayerProperties, LayerWeights), UMeshSculptLayerProperties::StaticClass());
+	ensure(LayerWeightsHandle->IsValidHandle());
+	LayerWeightsHandle->MarkHiddenByCustomization();
+	
+	DetailBuilder
+	.EditDefaultProperty(ActiveLayerHandle)->CustomWidget()
+	.OverrideResetToDefault(FResetToDefaultOverride::Hide())
+	.WholeRowContent()
+	[
+		SNew(SMeshLayersStack, Controller)
+	];
+}
+
 template<typename BrushClass>
 TSharedRef<IDetailCustomization> FMeshSculptBrushOpPropertiesDetails<BrushClass>::MakeInstance()
 {
diff --git a/Engine/Plugins/Editor/ModelingToolsEditorMode/Source/ModelingToolsEditorMode/Private/ModelingToolsEditorModeModule.cpp b/Engine/Plugins/Editor/ModelingToolsEditorMode/Source/ModelingToolsEditorMode/Private/ModelingToolsEditorModeModule.cpp
--- a/Engine/Plugins/Editor/ModelingToolsEditorMode/Source/ModelingToolsEditorMode/Private/ModelingToolsEditorModeModule.cpp
+++ b/Engine/Plugins/Editor/ModelingToolsEditorMode/Source/ModelingToolsEditorMode/Private/ModelingToolsEditorModeModule.cpp
@@ -33,6 +33,7 @@
 #include "Sculpting/MeshPlaneBrushOps.h"
 #include "Sculpting/MeshSculptBrushOps.h"
 #include "Sculpting/MeshSmoothingBrushOps.h"
+#include "Properties/MeshSculptLayerProperties.h"
 
 
 #define LOCTEXT_NAMESPACE "FModelingToolsEditorModeModule"
@@ -98,6 +99,8 @@ void FModelingToolsEditorModeModule::OnPostEngineInit()
 	ClassesToUnregisterOnShutdown.Add(UVertexBrushSculptProperties::StaticClass()->GetFName());
 	PropertyModule.RegisterCustomClassLayout("VertexBrushAlphaProperties", FOnGetDetailCustomizationInstance::CreateStatic(&FVertexBrushAlphaPropertiesDetails::MakeInstance));
 	ClassesToUnregisterOnShutdown.Add(UVertexBrushAlphaProperties::StaticClass()->GetFName());
+	PropertyModule.RegisterCustomClassLayout("MeshSculptLayerProperties", FOnGetDetailCustomizationInstance::CreateStatic(&FMeshSculptLayerPropertiesDetails::MakeInstance));
+	ClassesToUnregisterOnShutdown.Add(UMeshSculptLayerProperties::StaticClass()->GetFName());
 	// Sculpt - BrushOpProps
 	PropertyModule.RegisterCustomClassLayout("PinchBrushOpProps", FOnGetDetailCustomizationInstance::CreateStatic(&FMeshSculptBrushOpPropertiesDetails<UPinchBrushOpProps>::MakeInstance));
 	ClassesToUnregisterOnShutdown.Add(UPinchBrushOpProps::StaticClass()->GetFName());
diff --git a/Engine/Plugins/Editor/ModelingToolsEditorMode/Source/ModelingToolsEditorMode/Public/DetailsCustomizations/MeshVertexSculptToolCustomizations.h b/Engine/Plugins/Editor/ModelingToolsEditorMode/Source/ModelingToolsEditorMode/Public/DetailsCustomizations/MeshVertexSculptToolCustomizations.h
--- a/Engine/Plugins/Editor/ModelingToolsEditorMode/Source/ModelingToolsEditorMode/Public/DetailsCustomizations/MeshVertexSculptToolCustomizations.h
+++ b/Engine/Plugins/Editor/ModelingToolsEditorMode/Source/ModelingToolsEditorMode/Public/DetailsCustomizations/MeshVertexSculptToolCustomizations.h
@@ -20,6 +20,7 @@ class SBox;
 class UMeshVertexSculptTool;
 class FRecentAlphasProvider;
 class SToolInputAssetComboPanel;
+class IMeshLayersController;
 
 
 // customization for USculptBrushProperties, creates two-column layout
@@ -94,4 +95,18 @@ private:
 		const TSharedPtr<IPropertyHandle>& StrengthHandle,
 		TSharedPtr<IPropertyHandle> EnablePressureSensitivityHandle
 	);
+};
+
+// customization for UMeshSculptLayerProperties
+// - creates custom mesh layer panel
+class FMeshSculptLayerPropertiesDetails : public IDetailCustomization
+{
+public:
+	FMeshSculptLayerPropertiesDetails();
+	
+	static TSharedRef<IDetailCustomization> MakeInstance();
+	void CustomizeDetails(IDetailLayoutBuilder& DetailBuilder) override;
+	
+private:
+	TSharedPtr<IMeshLayersController> Controller;
 };
\ No newline at end of file
diff --git a/Engine/Plugins/Experimental/MeshModelingToolsetExp/Source/ModelingEditorUI/ModelingEditorUI.build.cs b/Engine/Plugins/Experimental/MeshModelingToolsetExp/Source/ModelingEditorUI/ModelingEditorUI.build.cs
--- a/Engine/Plugins/Experimental/MeshModelingToolsetExp/Source/ModelingEditorUI/ModelingEditorUI.build.cs
+++ b/Engine/Plugins/Experimental/MeshModelingToolsetExp/Source/ModelingEditorUI/ModelingEditorUI.build.cs
@@ -40,7 +40,8 @@ public class ModelingEditorUI : ModuleRules
 				"DeveloperSettings",
 				"Persona",
 				"ModelingUI",
-				"MeshModelingToolsEditorOnlyExp"
+				"MeshModelingToolsEditorOnlyExp",
+				"MeshModelingTools",
 			}
 			);
 
diff --git a/Engine/Plugins/Experimental/MeshModelingToolsetExp/Source/ModelingEditorUI/Private/ModelingEditorUIStyle.h b/Engine/Plugins/Experimental/MeshModelingToolsetExp/Source/ModelingEditorUI/Private/ModelingEditorUIStyle.h
new file mode 100644
--- /dev/null
+++ b/Engine/Plugins/Experimental/MeshModelingToolsetExp/Source/ModelingEditorUI/Private/ModelingEditorUIStyle.h
@@ -0,0 +1,68 @@
+﻿// Copyright Epic Games, Inc. All Rights Reserved.
+
+#pragma once
+
+#include "Styling/SlateStyle.h"
+#include "Styling/SlateStyleRegistry.h"
+#include "Styling/SlateStyleMacros.h"
+#include "Styling/StyleColors.h"
+#include "Styling/SlateTypes.h"
+#include "Styling/AppStyle.h"
+
+class FModelingEditorUIStyle final : public FSlateStyleSet
+{
+public:
+	
+	FModelingEditorUIStyle() : FSlateStyleSet("ModelingEditorUIStyle")
+	{
+		const FVector2D Icon16x16(16.0f, 16.0f);
+		const FVector2D Icon20x20(20.0f, 20.0f);
+		const FVector2D Icon40x40(40.0f, 40.0f);
+		const FVector2D Icon64x64(64.0f, 64.0f);
+
+		SetContentRoot(FPaths::EnginePluginsDir() / TEXT("Experimental/MeshModelingToolsetExp/Content"));
+		
+		Set("MeshLayers.DragSolver", new IMAGE_BRUSH("Icons/DragSolver", FVector2D(6, 15)));
+
+		const FSlateColor LayerOutlineColor = FSlateColor(FLinearColor(0.1843f, 0.1843f, 0.1843f, 0.5f));
+		const FSlateColor LayerOutlineColorSelected = FSlateColor(FLinearColor(0.1843f, 0.1843f, 0.1843f, 1.0f));
+		Set("MeshLayers.LayerBorder", new FSlateRoundedBoxBrush(
+			FStyleColors::Header, 4.f,
+			LayerOutlineColor, 1.f));
+		Set("MeshLayers.SelectedLayerBorder", new FSlateRoundedBoxBrush(
+			FStyleColors::Select, 4.f,
+			LayerOutlineColorSelected, 1.f));
+		
+		FCheckBoxStyle TransparentCheckBox = FCheckBoxStyle(FAppStyle::Get().GetWidgetStyle<FCheckBoxStyle>("DetailsView.SectionButton"))
+			.SetUncheckedImage(FSlateRoundedBoxBrush(FStyleColors::Transparent))
+			.SetUncheckedHoveredImage(FSlateRoundedBoxBrush(FStyleColors::Transparent))
+			.SetUncheckedPressedImage(FSlateRoundedBoxBrush(FStyleColors::Transparent))
+			.SetCheckedImage(FSlateRoundedBoxBrush(FStyleColors::Transparent))
+			.SetCheckedHoveredImage(FSlateRoundedBoxBrush(FStyleColors::Transparent))
+			.SetCheckedPressedImage(FSlateRoundedBoxBrush(FStyleColors::Transparent));
+		Set("MeshLayers.TransparentCheckBox", TransparentCheckBox);
+
+
+		FButtonStyle TransparentButton = FButtonStyle()
+			.SetNormal(FSlateRoundedBoxBrush(FStyleColors::Header, 4.0f, FStyleColors::Input, 1.0f))
+			.SetHovered(FSlateRoundedBoxBrush(FStyleColors::Input, 4.0f, FStyleColors::Input, 1.0f))
+			.SetPressed(FSlateRoundedBoxBrush(FStyleColors::Header, 4.0f, FStyleColors::Input, 1.0f))
+			.SetDisabled(FSlateRoundedBoxBrush(FStyleColors::Header, 4.0f, FStyleColors::Input, 1.0f))
+			.SetNormalPadding(FMargin(4.f, 4.f, 4.f, 4.f))
+			.SetPressedPadding(FMargin(4.f, 5.f, 4.f, 3.f));
+		Set("MeshLayers.TransparentButton", TransparentButton);
+		
+		FSlateStyleRegistry::RegisterSlateStyle(*this);
+	}
+
+	static FModelingEditorUIStyle& Get()
+	{
+		static FModelingEditorUIStyle Inst;
+		return Inst;
+	}
+	
+	~FModelingEditorUIStyle()
+	{
+		FSlateStyleRegistry::UnRegisterSlateStyle(*this);
+	}
+};
diff --git a/Engine/Plugins/Experimental/MeshModelingToolsetExp/Source/ModelingEditorUI/Private/ModelingWidgets/SMeshLayers.cpp b/Engine/Plugins/Experimental/MeshModelingToolsetExp/Source/ModelingEditorUI/Private/ModelingWidgets/SMeshLayers.cpp
new file mode 100644
--- /dev/null
+++ b/Engine/Plugins/Experimental/MeshModelingToolsetExp/Source/ModelingEditorUI/Private/ModelingWidgets/SMeshLayers.cpp
@@ -0,0 +1,563 @@
+﻿// Copyright Epic Games, Inc. All Rights Reserved.
+
+#include "ModelingWidgets/SMeshLayers.h"
+
+#include "ModelingEditorUIStyle.h"
+#include "ModelingWidgets/SculptLayersController.h"
+
+#include "Widgets/Input/SButton.h"
+#include "Widgets/Input/SCheckBox.h"
+#include "Widgets/Input/SSpinBox.h"
+#include "Widgets/Layout/SSpacer.h"
+#include "Widgets/Text/SInlineEditableTextBlock.h"
+
+#define LOCTEXT_NAMESPACE "SMeshLayers"
+
+float SMeshLayerItem::LayersHorizontalPadding = 6.0f;;
+float SMeshLayerItem::LayersVerticalPadding = 3.0f;
+
+//
+// FMeshLayerElement
+//
+
+TSharedRef<ITableRow> FMeshLayerElement::MakeListRowWidget(
+	const TSharedRef<STableViewBase>& InOwnerTable,
+	TSharedRef<FMeshLayerElement> InMeshLayerElement,
+	TSharedPtr<SMeshLayersList> InMeshLayersListWidget)
+{
+	return SNew(SMeshLayerItem, InOwnerTable)
+		.InMeshLayerElement(InMeshLayerElement)
+		.InMeshLayersList(InMeshLayersListWidget);
+}
+
+//
+// SMeshLayerList
+//
+
+void SMeshLayersList::Construct(const FArguments& InArgs)
+{
+	Controller = InArgs._InController;
+	
+	SListView::FArguments SuperArgs;
+	SuperArgs.ListItemsSource(&Layers);
+	SuperArgs.SelectionMode(ESelectionMode::Single);
+	SuperArgs.OnGenerateRow(this, &SMeshLayersList::MakeListRowWidget);
+	SuperArgs.OnSelectionChanged(this, &SMeshLayersList::OnSelectionChanged);
+	SuperArgs.OnMouseButtonClick(this, &SMeshLayersList::OnItemClicked);
+
+	SListView::Construct(SuperArgs);
+
+	bAllowReordering = InArgs._InAllowReordering;
+	bAllowAddRemove = InArgs._InAllowAddRemove;
+}
+
+TSharedRef<ITableRow> SMeshLayersList::MakeListRowWidget(TSharedPtr<FMeshLayerElement> InLayerElement,
+	const TSharedRef<STableViewBase>& OwnerTable)
+{
+	return InLayerElement->MakeListRowWidget(OwnerTable,InLayerElement.ToSharedRef(),SharedThis(this));
+}
+
+FReply SMeshLayersList::OnKeyDown(const FGeometry& MyGeometry, const FKeyEvent& InKeyEvent)
+{
+	TArray<TSharedPtr<FMeshLayerElement>> SelectedLayers = GetSelectedItems();
+	if (!SelectedLayers.IsEmpty() && InKeyEvent.GetKey() == EKeys::Delete)
+	{
+		DeleteMeshLayer(SelectedLayers.Last());
+		return FReply::Handled();
+	}
+	
+	return SListView::OnKeyDown(MyGeometry, InKeyEvent);
+}
+
+void SMeshLayersList::DeleteMeshLayer(TSharedPtr<FMeshLayerElement> LayerToDelete)
+{
+	if (!LayerToDelete.IsValid())
+	{
+		return;
+	}
+
+	// remove the layer
+	if (!Controller.Pin()->RemoveMeshLayer(LayerToDelete->GetIndexInStack()))
+	{
+		return;
+	}
+
+	// remove layer from this List of layers
+	const int32 StackIndexToDelete = LayerToDelete->GetIndexInStack();
+
+	for (TSharedPtr<FMeshLayerElement> Layer : Layers)
+	{
+		const int32 StackIndex = Layer->GetIndexInStack();
+		// for all layers after the removed one, adjust their IndexInStack
+		if (StackIndex != INDEX_NONE && StackIndex > StackIndexToDelete)
+		{
+			Layer->SetIndexInStack(StackIndex - 1);
+		}
+	}
+	RequestListRefresh();
+}
+
+void SMeshLayersList::RefreshAndRestore()
+{
+	IMeshLayersController* Ctrl = Controller.Pin().Get();
+	if (!ensure(Ctrl))
+	{
+		return;
+	}
+	RequestListRefresh();
+}
+
+void SMeshLayersList::OnSelectionChanged(TSharedPtr<FMeshLayerElement> InLayer, ESelectInfo::Type SelectInfo) const
+{
+	// adds support for keyboard navigation of layer stack
+	if (SelectInfo == ESelectInfo::OnNavigation || SelectInfo == ESelectInfo::Direct)
+	{
+		if (Controller.IsValid() && InLayer.IsValid())
+		{
+			Controller.Pin()->SetActiveLayer(InLayer->GetIndexInStack());
+		}
+		return;
+	}
+
+	// set first layer in stack as selected layer
+	if (!InLayer.IsValid() && Controller.IsValid())
+	{
+		Controller.Pin()->SetActiveLayer(0);
+	}
+}
+
+void SMeshLayersList::OnItemClicked(TSharedPtr<FMeshLayerElement> InLayer)
+{
+	// to rename an item, you have to select it first, then click on it again
+	const bool ClickedOnSameItem = LastSelectedLayer.Pin().Get() == InLayer.Get();
+	if (ClickedOnSameItem)
+	{
+		RegisterActiveTimer(0.f,
+			FWidgetActiveTimerDelegate::CreateLambda([this](double, float)
+			{
+				RequestRenameSelectedLayer();
+				return EActiveTimerReturnType::Stop;
+			}));
+	}
+
+	LastSelectedLayer = InLayer;
+
+	if (Controller.IsValid() && InLayer.IsValid())
+	{
+		Controller.Pin()->SetActiveLayer(InLayer->GetIndexInStack());
+	}
+}
+
+void SMeshLayersList::RequestRenameSelectedLayer() const
+{
+	if (!LastSelectedLayer.IsValid())
+	{
+		return;
+	}
+	LastSelectedLayer.Pin()->OnRenameRequested.ExecuteIfBound();
+}
+
+FReply SMeshLayersList::OnDragDetected(const FGeometry& MyGeometry, const FPointerEvent& MouseEvent)
+{
+	if (!bAllowReordering)
+	{
+		return FReply::Unhandled();
+	}
+	const TArray<TSharedPtr<FMeshLayerElement>> CurrentSelectedItems = GetSelectedItems();
+	if (CurrentSelectedItems.Num() != 1)
+	{
+		return FReply::Unhandled();
+	}
+	
+	if (MouseEvent.IsMouseButtonDown(EKeys::LeftMouseButton))
+	{
+		const TSharedPtr<FMeshLayerElement> DraggedElement = CurrentSelectedItems[0];
+		const TSharedRef<FMeshLayerStackDragDropOp> DragDropOp = FMeshLayerStackDragDropOp::New(DraggedElement);
+		return FReply::Handled().BeginDragDrop(DragDropOp);
+	}
+
+	return FReply::Unhandled();
+}
+
+TOptional<EItemDropZone> SMeshLayersList::OnCanAcceptDrop(const FDragDropEvent& DragDropEvent, EItemDropZone DropZone,
+	TSharedPtr<FMeshLayerElement> TargetLayer)
+{
+	TOptional<EItemDropZone> ReturnedDropZone;
+	const TSharedPtr<FMeshLayerStackDragDropOp> DragDropOp = DragDropEvent.GetOperationAs<FMeshLayerStackDragDropOp>();
+	if (!DragDropOp.IsValid())
+	{
+		return ReturnedDropZone;
+	}
+
+	const TSharedPtr<FMeshLayerElement>& DraggedElement = DragDropOp.Get()->Element.Pin();
+	if (!DraggedElement)
+	{
+		return ReturnedDropZone;
+	}
+
+	// validate index to move to
+	const int32 IndexToMoveTo = FMeshLayerStackDragDropOp::GetIndexToMoveTo(DraggedElement, TargetLayer, DropZone);
+
+	// GetIndexToMoveTo returns index none for invalid moves
+	if (IndexToMoveTo == INDEX_NONE)
+	{
+		return ReturnedDropZone;
+	}
+	
+	ReturnedDropZone = DropZone;
+	return ReturnedDropZone;
+}
+
+FReply SMeshLayersList::OnAcceptDrop(const FDragDropEvent& DragDropEvent, EItemDropZone DropZone, TSharedPtr<FMeshLayerElement> TargetLayer)
+{
+	const TSharedPtr<FMeshLayerStackDragDropOp> DragDropOp = DragDropEvent.GetOperationAs<FMeshLayerStackDragDropOp>();
+	if (!DragDropOp.IsValid())
+	{
+		return FReply::Unhandled();
+	}
+	
+	const TSharedPtr<IMeshLayersController> Ctrl = Controller.Pin();
+	if (!Ctrl.IsValid())
+	{
+		return FReply::Handled();
+	}
+
+	const TSharedPtr<FMeshLayerElement>& DraggedElement = DragDropOp.Get()->Element.Pin();
+	// stack index
+	const int32 IndexToMoveTo = FMeshLayerStackDragDropOp::GetIndexToMoveTo(DraggedElement, TargetLayer, DropZone);
+	if (IndexToMoveTo == INDEX_NONE)
+	{
+		// don't do anything if the drop location is invalid
+		return FReply::Unhandled();
+	}
+	
+	const int32 LayerToMoveIndex = DraggedElement->GetIndexInStack();
+	// perform move in underlying property set
+	Ctrl->MoveLayerInStack(LayerToMoveIndex, IndexToMoveTo);
+
+	// also need to perform move in Layers (array of MeshElements) here
+	const int Dir = IndexToMoveTo > LayerToMoveIndex ? 1 : -1;
+	int32 CurrentLayerIndex = LayerToMoveIndex;
+	// perform move by swapping with previous/next layer in stack until layer is in desired location
+	while (CurrentLayerIndex != IndexToMoveTo)
+	{
+		Layers.Swap(CurrentLayerIndex, CurrentLayerIndex + Dir);
+		// adjusting indices as we go 
+		Layers[CurrentLayerIndex]->SetIndexInStack(CurrentLayerIndex);
+		Layers[CurrentLayerIndex  + Dir]->SetIndexInStack(CurrentLayerIndex + Dir);
+		CurrentLayerIndex += Dir;
+	}
+
+	// if we are moving the active layer, update our ActiveLayer to reflect new index
+	const int32 ActiveLayer = Ctrl->GetActiveLayer();
+	if (ActiveLayer == LayerToMoveIndex)
+	{
+		Ctrl->SetActiveLayer(IndexToMoveTo);
+	}
+	// if we are placing a layer before the ActiveLayer (that was previously after it), increment ActiveLayer's index
+	else if (ActiveLayer >= IndexToMoveTo && LayerToMoveIndex > ActiveLayer)
+	{
+		Ctrl->SetActiveLayer(ActiveLayer + 1);
+	}
+	// if we are moving a layer that was before the ActiveLayer to be after it, decrement ActiveLayer index
+	else if (ActiveLayer > LayerToMoveIndex && ActiveLayer <= IndexToMoveTo)
+	{
+		Ctrl->SetActiveLayer(ActiveLayer - 1);
+	}
+	Ctrl->RefreshLayersStackView();
+	
+	return FReply::Handled();
+}
+
+//
+// SMeshLayerItem
+//
+
+void SMeshLayerItem::Construct(const FArguments& InArgs, const TSharedRef<STableViewBase>& OwnerTable)
+{
+	ListView = InArgs._InMeshLayersList;
+	Element = InArgs._InMeshLayerElement;
+
+	SHorizontalBox::FArguments BoxArgs;
+
+	// if reordering is enabled, add drag icon to UI
+	if (ListView.Pin()->GetAllowReordering())
+	{
+		// drag icon
+		BoxArgs
+		+ SHorizontalBox::Slot()
+		.MaxWidth(18.f)
+		.AutoWidth()
+		.HAlign(HAlign_Left)
+		.VAlign(VAlign_Center)
+		.Padding(LayersHorizontalPadding, LayersVerticalPadding)
+		[
+			SNew(SImage)
+			.Image(FModelingEditorUIStyle::Get().GetBrush("MeshLayers.DragSolver"))
+		];
+	}
+
+	// enable Checkbox - TODO
+	BoxArgs
+	+ SHorizontalBox::Slot()
+	.AutoWidth()
+	.HAlign(HAlign_Left)
+	.VAlign(VAlign_Center)
+	.Padding(FMargin(LayersHorizontalPadding, LayersVerticalPadding))
+	[
+		SNew(SBox)
+		[
+			SNew(SCheckBox)
+			.Style(FModelingEditorUIStyle::Get().Get(), "MeshLayers.TransparentCheckBox")
+			.Padding(FMargin(4, 2))
+			.HAlign(HAlign_Center)
+			.IsEnabled(false) // TODO
+			.ToolTipText( LOCTEXT("EnabledTooltip", "Enables/disables the mesh layer") )
+			.OnCheckStateChanged_Lambda([this](const ECheckBoxState NewState)
+			{
+				//IMeshLayersController* Ctrl = ListView.Pin()->Controller.Pin().Get();
+				//Ctrl->AssetController->SetMeshLayerEnabled( Element.Pin()->GetIndexInStack(), bIsChecked);
+			})
+			.IsChecked_Lambda([this]()
+			{
+				//IMeshLayersController* Ctrl = ListView.Pin()->Controller.Pin().Get();
+				//bool bEnabled = Ctrl->AssetController->GetMeshLayerEnabled(Element.Pin()->GetIndexInStack());
+				//return bEnabled ? ECheckBoxState::Checked : ECheckBoxState::Unchecked;
+				return ECheckBoxState::Checked;
+			})
+			[
+				SNew(SHorizontalBox)
+				+ SHorizontalBox::Slot()
+				.VAlign(VAlign_Center)
+				.Padding(FMargin(0))
+				.AutoWidth()
+				[
+					SNew(SImage)
+					.Image_Lambda([this]()
+					{
+						return IsLayerEnabled() ?
+							FAppStyle::Get().GetBrush(TEXT("Level.VisibleIcon16x")) :
+						FAppStyle::Get().GetBrush(TEXT("Level.NotVisibleIcon16x"));
+					})
+				]
+			]
+		]
+	];
+
+	// display Name
+	BoxArgs
+	+ SHorizontalBox::Slot()
+	.AutoWidth()
+	.HAlign(HAlign_Left)
+	.VAlign(VAlign_Center)
+	.Padding(FMargin(LayersHorizontalPadding, LayersVerticalPadding))
+	[
+		SAssignNew(EditNameWidget, SInlineEditableTextBlock)
+		.Text_Lambda([this]()
+		{
+			const IMeshLayersController* Ctrl = ListView.Pin()->GetController().Pin().Get();
+			const FName LayerName = Ctrl->GetLayerName(Element.Pin()->GetIndexInStack());
+			// TODO
+			// if layer name has not been set, display a unique placeholder name like "Layer 2" to
+			// eliminate the confusion of having multiple layers with the display label "None"
+			/*if (LayerName == NAME_None)
+			{
+				return FText::FromName(Ctrl->GetPlaceholderNames()[Element.Pin()->GetIndexInStack()]);
+			}*/
+			// otherwise, display the layer's actual name
+			return FText::FromName(LayerName);
+		})
+		.OnVerifyTextChanged_Lambda([](const FText& InText, FText& OutErrorMessage)
+		{
+			static FString IllegalNameCharacters = "^<>:\"/\\|?*";
+			return FName::IsValidXName(InText.ToString(), FString(IllegalNameCharacters), &OutErrorMessage);
+		})
+		.OnTextCommitted(this, &SMeshLayerItem::OnNameCommitted)
+		.MultiLine(false)
+	];
+	
+	// spacer
+	BoxArgs
+	+ SHorizontalBox::Slot()
+	.HAlign(HAlign_Fill)
+	.FillWidth(1.0f)
+	[
+		SNew(SSpacer)
+		.Size(FVector2D(0.0f, 1.0f))
+	];
+	
+	// spin Box
+	BoxArgs
+	+ SHorizontalBox::Slot()
+	.AutoWidth()
+	.HAlign(HAlign_Right)
+	.VAlign(VAlign_Center)
+	.Padding(FMargin(LayersHorizontalPadding, LayersVerticalPadding))
+	[
+		SNew(SSpinBox<double>)
+		.Value_Lambda([this]() -> double
+		{
+			IMeshLayersController* Ctrl = ListView.Pin()->GetController().Pin().Get();
+			return Ctrl->GetLayerWeight(Element.Pin()->GetIndexInStack());
+		})
+		.OnValueChanged_Lambda([this](const float NewValue)
+		{
+			IMeshLayersController* Ctrl = ListView.Pin()->GetController().Pin().Get();
+			Ctrl->SetLayerWeight(Element.Pin()->GetIndexInStack(), NewValue, EPropertyChangeType::Interactive);
+		})
+		// only log the slider movement after its complete
+		.OnEndSliderMovement_Lambda([this](float LastValue)
+		{
+			IMeshLayersController* Ctrl = ListView.Pin()->GetController().Pin().Get();
+			Ctrl->SetLayerWeight(Element.Pin()->GetIndexInStack(), LastValue, EPropertyChangeType::ValueSet);
+		})
+		.MinValue(-1.0f)
+		.MaxValue(1.5f)
+		.Style(&FAppStyle::Get().GetWidgetStyle<FSpinBoxStyle>("Sequencer.HyperlinkSpinBox"))
+	];
+
+	if (ListView.Pin()->GetAllowAddRemove())
+	{
+		// delete Button
+		BoxArgs
+		+ SHorizontalBox::Slot()
+		.AutoWidth()
+		.HAlign(HAlign_Right)
+		.VAlign(VAlign_Center)
+		.Padding(FMargin(LayersHorizontalPadding, LayersVerticalPadding))
+		[
+			SNew(SButton)
+			.ContentPadding(FMargin(0,0))
+			.ToolTipText(LOCTEXT("DeleteLayer", "Delete mesh layer and remove from stack."))
+			.ButtonStyle(FModelingEditorUIStyle::Get().Get(), "MeshLayers.TransparentButton")
+			.OnClicked_Lambda([this]() -> FReply
+			{
+				ListView.Pin()->DeleteMeshLayer(Element.Pin());
+				const IMeshLayersController* Ctrl = ListView.Pin()->GetController().Pin().Get();
+				Ctrl->RefreshLayersStackView();
+				return FReply::Handled();
+			})
+			.Content()
+			[
+				SNew(SImage)
+				.Image(FAppStyle::Get().GetBrush("Icons.Delete"))
+				.ColorAndOpacity(FSlateColor::UseForeground())
+			]
+		];
+	}
+
+	STableRow::Construct(
+		STableRow::FArguments()
+		.OnDragDetected(ListView.Pin().Get(), &SMeshLayersList::OnDragDetected)
+		.OnCanAcceptDrop(ListView.Pin().Get(), &SMeshLayersList::OnCanAcceptDrop)
+		.OnAcceptDrop(ListView.Pin().Get(), &SMeshLayersList::OnAcceptDrop)
+		.ShowSelection(false)
+		.Padding(FMargin(LayersHorizontalPadding, LayersVerticalPadding))
+		.Content()
+		[
+			SNew(SBorder)
+			.BorderImage_Lambda([this]()
+			{
+				// display highlight on selected layer only
+				if (Element.Pin()->GetIndexInStack() == ListView.Pin()->GetController().Pin()->GetActiveLayer())
+				{
+					return FModelingEditorUIStyle::Get().GetBrush("MeshLayers.SelectedLayerBorder");
+				}
+				return FModelingEditorUIStyle::Get().GetBrush("MeshLayers.LayerBorder");
+			})
+			.Padding(FMargin(LayersHorizontalPadding, LayersVerticalPadding))
+			[
+				SArgumentNew(BoxArgs, SHorizontalBox)
+			]
+		], OwnerTable);
+	
+	// bind the rename callback to enter text editing mode when item is double-clicked
+	Element.Pin()->OnRenameRequested.BindSP(EditNameWidget.Get(), &SInlineEditableTextBlock::EnterEditingMode);
+}
+
+void SMeshLayerItem::OnNameCommitted(const FText& InText, ETextCommit::Type InCommitType) const
+{
+	const FName NewName = FName(InText.ToString());
+	const IMeshLayersController* Ctrl = ListView.Pin()->GetController().Pin().Get();
+	Ctrl->SetLayerName(Element.Pin()->GetIndexInStack(), NewName);
+}
+
+FText SMeshLayerItem::GetName() const
+{
+	const IMeshLayersController* Ctrl = ListView.Pin()->GetController().Pin().Get();
+	const FName LayerName = Ctrl->GetLayerName(Element.Pin()->GetIndexInStack());
+	return FText::FromName(LayerName);
+}
+
+bool SMeshLayerItem::IsLayerEnabled() const
+{
+	// TODO
+	
+	return true;
+}
+
+//
+// FMeshLayerStackDragDropOp
+//
+
+
+TSharedRef<FMeshLayerStackDragDropOp> FMeshLayerStackDragDropOp::New(TWeakPtr<FMeshLayerElement> InElement)
+{
+	TSharedRef<FMeshLayerStackDragDropOp> Operation = MakeShared<FMeshLayerStackDragDropOp>();
+	Operation->Element = InElement;
+	Operation->Construct();
+	return Operation;
+}
+
+TSharedPtr<SWidget> FMeshLayerStackDragDropOp::GetDefaultDecorator() const
+{
+	return SNew(SBorder)
+		.Visibility(EVisibility::Visible)
+		.BorderImage(FAppStyle::GetBrush("Menu.Background"))
+		[
+			SNew(STextBlock)
+			.Text(FText::FromName(Element.Pin()->GetLayerList().Pin()->GetController().Pin()->
+				GetLayerName(Element.Pin()->GetIndexInStack())))
+		];
+}
+
+int32 FMeshLayerStackDragDropOp::GetIndexToMoveTo(
+	const TSharedPtr<FMeshLayerElement>& InDraggedElement,
+	const TSharedPtr<FMeshLayerElement>& InTargetElement,
+	const EItemDropZone InDropZone)
+{
+	// disallow dropping on self
+	const int32 DraggedItemIndex = InDraggedElement->GetIndexInStack();
+	const int32 TargetItemIndex = InTargetElement->GetIndexInStack();
+	if (DraggedItemIndex == TargetItemIndex)
+	{
+		return INDEX_NONE;
+	}
+
+	// disallow dropping in a place that would not change the order (like below the one above, or above the one below)
+	int32 IndexToMoveDragItemTo = INDEX_NONE;
+	bool bDraggingDown = DraggedItemIndex <= TargetItemIndex;
+	switch (InDropZone)
+	{
+	case EItemDropZone::AboveItem:
+		IndexToMoveDragItemTo = bDraggingDown ? TargetItemIndex - 1 : TargetItemIndex;
+		break;
+	case EItemDropZone::BelowItem:
+		IndexToMoveDragItemTo = bDraggingDown ? TargetItemIndex : TargetItemIndex + 1;
+		break;
+	case EItemDropZone::OntoItem:
+		IndexToMoveDragItemTo = TargetItemIndex;
+		break;
+	default:
+		checkNoEntry();
+	}
+	
+	if (DraggedItemIndex == IndexToMoveDragItemTo)
+	{
+		return INDEX_NONE;
+	}
+
+	return IndexToMoveDragItemTo;
+}
+
+#undef LOCTEXT_NAMESPACE
diff --git a/Engine/Plugins/Experimental/MeshModelingToolsetExp/Source/ModelingEditorUI/Private/ModelingWidgets/SMeshLayers.h b/Engine/Plugins/Experimental/MeshModelingToolsetExp/Source/ModelingEditorUI/Private/ModelingWidgets/SMeshLayers.h
new file mode 100644
--- /dev/null
+++ b/Engine/Plugins/Experimental/MeshModelingToolsetExp/Source/ModelingEditorUI/Private/ModelingWidgets/SMeshLayers.h
@@ -0,0 +1,167 @@
+﻿// Copyright Epic Games, Inc. All Rights Reserved.
+
+#pragma once
+
+#include "DragAndDrop/DecoratedDragDropOp.h"
+#include "Widgets/Views/SListView.h"
+
+class SMeshLayersList;
+class SInlineEditableTextBlock;
+class IMeshLayersController;
+
+/**
+ * Represents a single mesh layer within the List
+ */
+class FMeshLayerElement
+{
+public:
+	static TSharedRef<ITableRow> MakeListRowWidget(
+		const TSharedRef<STableViewBase>& InOwnerTable,
+		TSharedRef<FMeshLayerElement> InMeshLayerElement,
+		TSharedPtr<SMeshLayersList> InMeshLayerListWidget);
+
+	static TSharedRef<FMeshLayerElement> Make(
+		const int32 InLayerIndex,
+		const TSharedPtr<SMeshLayersList>& InLayerListWidget)
+	{
+		FMeshLayerElement* NewElement = new FMeshLayerElement(InLayerIndex, InLayerListWidget);
+		return MakeShareable(NewElement);
+	}
+	
+	int32 GetIndexInStack() const { return IndexInStack; };
+	void SetIndexInStack(const int32 InIndex) { IndexInStack = InIndex;}
+	
+	TWeakPtr<SMeshLayersList> GetLayerList() const { return LayerListWidget; };
+	
+	DECLARE_DELEGATE(FOnRenameRequested);
+	FOnRenameRequested OnRenameRequested;
+private:
+
+	int32 IndexInStack = INDEX_NONE;
+	TWeakPtr<SMeshLayersList> LayerListWidget;
+	
+	// hidden constructors, always use Make above
+	FMeshLayerElement() = default;
+	FMeshLayerElement(
+		const int32 InIndexInStack,
+		const TSharedPtr<SMeshLayersList>& InLayerListWidget) :
+		IndexInStack(InIndexInStack),
+		LayerListWidget(InLayerListWidget){}
+};
+
+/**
+ * Represents the list of mesh layers - works within the MeshLayersStack widget
+ */
+class SMeshLayersList : public SListView<TSharedPtr<FMeshLayerElement>>
+{
+public:
+	SLATE_BEGIN_ARGS(SMeshLayersList)
+		: _InAllowAddRemove(true)
+		, _InAllowReordering(true){}
+		SLATE_ARGUMENT(TWeakPtr<IMeshLayersController>, InController)
+		// TODO put these props on the controller ?
+		SLATE_ARGUMENT(bool, InAllowAddRemove)
+		SLATE_ARGUMENT(bool, InAllowReordering)
+	SLATE_END_ARGS()
+
+	virtual ~SMeshLayersList() {}
+
+	void Construct(const FArguments& InArgs);
+
+	TSharedRef<ITableRow> MakeListRowWidget(TSharedPtr<FMeshLayerElement> InLayerElement, const TSharedRef<STableViewBase>& OwnerTable);
+	
+	// drag and drop operations
+	FReply OnDragDetected(
+		const FGeometry& MyGeometry,
+		const FPointerEvent& MouseEvent);
+	TOptional<EItemDropZone> OnCanAcceptDrop(
+		const FDragDropEvent& DragDropEvent,
+		EItemDropZone DropZone,
+		TSharedPtr<FMeshLayerElement> TargetLayer);
+	FReply OnAcceptDrop(
+		const FDragDropEvent& DragDropEvent,
+		EItemDropZone DropZone,
+		TSharedPtr<FMeshLayerElement> TargetLayer);
+	// end drag and drop operations
+
+	// renaming
+	void OnSelectionChanged(TSharedPtr<FMeshLayerElement> InLayer, ESelectInfo::Type SelectInfo) const;
+	void OnItemClicked(TSharedPtr<FMeshLayerElement> InLayer);
+	void RequestRenameSelectedLayer() const;
+	// END renaming
+	
+	// delete layer from stack
+	virtual FReply OnKeyDown(const FGeometry& MyGeometry, const FKeyEvent& InKeyEvent);
+	void DeleteMeshLayer(TSharedPtr<FMeshLayerElement> LayerToDelete);
+
+	// must be called after refresh
+	void RefreshAndRestore();
+	
+	TArray<TSharedPtr<FMeshLayerElement>>& GetLayers() { return Layers; }
+	TWeakPtr<IMeshLayersController> GetController() { return Controller; }
+
+	bool GetAllowReordering() const { return bAllowReordering;};
+	bool GetAllowAddRemove() const { return bAllowAddRemove;}
+
+private:
+	bool bAllowReordering = true;
+	bool bAllowAddRemove = true;
+	
+	// the layers contained in this list
+	TArray<TSharedPtr<FMeshLayerElement>> Layers;
+	
+	TWeakPtr<IMeshLayersController> Controller;
+
+	TWeakPtr<FMeshLayerElement> LastSelectedLayer;
+};
+
+/**
+ * Builds a UI element representing a single mesh layer
+ */
+class SMeshLayerItem : public STableRow<TSharedPtr<FMeshLayerElement>>
+{
+public:
+	SLATE_BEGIN_ARGS(SMeshLayerItem) {}
+		SLATE_ARGUMENT(TWeakPtr<SMeshLayersList>, InMeshLayersList)
+		SLATE_ARGUMENT(TWeakPtr<FMeshLayerElement>, InMeshLayerElement)
+	SLATE_END_ARGS()
+
+	void Construct(const FArguments& InArgs, const TSharedRef<STableViewBase>& OwnerTable);
+
+	// RENAMING op
+	void OnNameCommitted(const FText& InText, ETextCommit::Type InCommitType) const;
+	FText GetName() const;
+	// END renaming op
+
+	static float LayersHorizontalPadding;
+	static float LayersVerticalPadding;
+	
+private:
+
+	bool IsLayerEnabled() const;
+	
+	TWeakPtr<SMeshLayersList> ListView;
+	TWeakPtr<FMeshLayerElement> Element;
+	TSharedPtr<SInlineEditableTextBlock> EditNameWidget;
+};
+
+/**
+ * Support for drag and drop operations within the MeshLayers
+ */
+class FMeshLayerStackDragDropOp : public FDecoratedDragDropOp
+{
+public:
+	
+	DRAG_DROP_OPERATOR_TYPE(FMeshLayerStackDragDropOp, FDecoratedDragDropOp)
+
+	static TSharedRef<FMeshLayerStackDragDropOp> New(TWeakPtr<FMeshLayerElement> InElement);
+	
+	virtual TSharedPtr<SWidget> GetDefaultDecorator() const override;
+	
+	static int32 GetIndexToMoveTo(
+		const TSharedPtr<FMeshLayerElement>& InDraggedElement,
+		const TSharedPtr<FMeshLayerElement>& InTargetElement,
+		const EItemDropZone InDropZone);
+	
+	TWeakPtr<FMeshLayerElement> Element;
+};
\ No newline at end of file
diff --git a/Engine/Plugins/Experimental/MeshModelingToolsetExp/Source/ModelingEditorUI/Private/ModelingWidgets/SMeshLayersStack.cpp b/Engine/Plugins/Experimental/MeshModelingToolsetExp/Source/ModelingEditorUI/Private/ModelingWidgets/SMeshLayersStack.cpp
new file mode 100644
--- /dev/null
+++ b/Engine/Plugins/Experimental/MeshModelingToolsetExp/Source/ModelingEditorUI/Private/ModelingWidgets/SMeshLayersStack.cpp
@@ -0,0 +1,113 @@
+﻿// Copyright Epic Games, Inc. All Rights Reserved.
+
+#include "ModelingWidgets/SMeshLayersStack.h"
+
+#include "ModelingWidgets/SculptLayersController.h"
+#include "ModelingWidgets/SMeshLayers.h"
+#include "SPositiveActionButton.h"
+
+#define LOCTEXT_NAMESPACE "SMeshLayersStack"
+
+void SMeshLayersStack::Construct(
+	const FArguments& InArgs,
+	const TWeakPtr<IMeshLayersController>& InController)
+{
+	Controller = InController;
+	Controller.Pin()->SetLayerStackView(SharedThis(this));
+
+	SVerticalBox::FArguments BoxArgs;
+
+	if (InArgs._InAllowAddRemove)
+	{
+		BoxArgs
+		+ SVerticalBox::Slot()
+		.AutoHeight()
+		.VAlign(VAlign_Top)
+		.Padding(0.0f)
+		[
+			SNew(SBorder)
+			.Padding(0.0f)
+			.BorderImage(FAppStyle::GetBrush("DetailsView.CategoryTop"))
+			.BorderBackgroundColor(FLinearColor(0.6f, 0.6f, 0.6f, 1.0f))
+			[
+				SNew(SVerticalBox)
+				+ SVerticalBox::Slot()
+				.AutoHeight()
+				.VAlign(VAlign_Top)
+				[
+					SNew(SHorizontalBox)
+					+ SHorizontalBox::Slot()
+					.VAlign(VAlign_Center)
+					.HAlign(HAlign_Left)
+					.FillWidth(1.f)
+					.Padding(6.0f, 4.0f)
+					[
+						SNew(SPositiveActionButton)
+							.Icon(FAppStyle::Get().GetBrush("Icons.Plus"))
+							.Text(LOCTEXT("AddNewMeshLayerLabel", "Add Layer"))
+							.ToolTipText(LOCTEXT("AddNewToolTip", "Add a new mesh layer"))
+							.OnClicked_Lambda([this]()
+							{
+								IMeshLayersController* AssetController = Controller.Pin().Get();
+								if (!AssetController)
+								{
+									return FReply::Unhandled();
+								}
+
+								// add a new layer to the stack
+								AssetController->AddMeshLayer();
+								RefreshStackView();
+								return FReply::Handled();
+							})
+					]
+				]
+			]
+		];
+	}
+
+	BoxArgs
+	+SVerticalBox::Slot()
+	.Padding(0.0f)
+	.HAlign(HAlign_Fill)
+	[
+		SAssignNew( ListView, SMeshLayersList )
+		.InController(Controller)
+		.InAllowReordering(InArgs._InAllowReordering)
+		.InAllowAddRemove(InArgs._InAllowAddRemove)
+	];
+
+	ChildSlot
+	[
+		SArgumentNew(BoxArgs, SVerticalBox)
+	];
+
+	RefreshStackView();
+}
+
+void SMeshLayersStack::RefreshStackView() const
+{
+	const TSharedPtr<IMeshLayersController> SLController = Controller.Pin();
+	if (!SLController.IsValid())
+	{
+		return; 
+	}
+
+	// rebuild list of top level elements
+	const int32 NumLayers = SLController->GetNumMeshLayers();
+	
+	// empty old elements
+	ListView->GetLayers().Reset();
+
+	// rebuild list of top level elements
+	for (int32 LayerIdx = 0; LayerIdx < NumLayers; ++LayerIdx)
+	{
+		// add op to main stack
+		TSharedPtr<FMeshLayerElement> StackElement = FMeshLayerElement::Make(LayerIdx, ListView);
+		ListView->GetLayers().Add(StackElement);
+	}
+
+	// refresh the list and restore the selection
+	ListView->RefreshAndRestore();
+}
+
+#undef LOCTEXT_NAMESPACE
\ No newline at end of file
diff --git a/Engine/Plugins/Experimental/MeshModelingToolsetExp/Source/ModelingEditorUI/Private/ModelingWidgets/SculptLayersController.cpp b/Engine/Plugins/Experimental/MeshModelingToolsetExp/Source/ModelingEditorUI/Private/ModelingWidgets/SculptLayersController.cpp
new file mode 100644
--- /dev/null
+++ b/Engine/Plugins/Experimental/MeshModelingToolsetExp/Source/ModelingEditorUI/Private/ModelingWidgets/SculptLayersController.cpp
@@ -0,0 +1,148 @@
+// Copyright Epic Games, Inc. All Rights Reserved.
+
+#include "ModelingWidgets/SculptLayersController.h"
+
+#include "ScopedTransaction.h"
+#include "Properties/MeshSculptLayerProperties.h"
+#include "ModelingWidgets/SMeshLayersStack.h"
+
+
+DEFINE_LOG_CATEGORY(LogMeshLayers);
+
+#define LOCTEXT_NAMESPACE "SculptLayersController"
+
+FSculptLayersController::FSculptLayersController() { }
+
+int32 FSculptLayersController::AddMeshLayer()
+{
+	// ensure the properties are valid
+	if (!ensure(Properties))
+	{
+		return INDEX_NONE;
+	}
+	FScopedTransaction Transaction(LOCTEXT("AddSculptLayer_Label", "Add Sculpt Layer"));
+
+	// add the layer
+	Properties->AddLayer();
+	return Properties->LayerWeights.Num() - 1;
+}
+
+bool FSculptLayersController::RemoveMeshLayer(const int32 InLayerIndex)
+{
+	// ensure properties are valid
+	if (!ensure(Properties))
+	{
+		return false;
+	}
+	FScopedTransaction Transaction(LOCTEXT("RemoveSculptLayer_Label", "Remove Sculpt Layer"));
+	
+	// remove the layer and return if the remove was successful
+	return Properties->RemoveLayerAtIndex(InLayerIndex);
+}
+
+void FSculptLayersController::SetActiveLayer(const int32 InLayerIndex) const
+{
+	// ensure the properties are valid
+	if (!ensure(Properties))
+	{
+		return;
+	}
+	
+	Properties->SetActiveLayer(InLayerIndex);
+}
+
+int32 FSculptLayersController::GetActiveLayer() const
+{
+	// ensure the properties are valid
+	if (!ensure(Properties))
+	{
+		return INDEX_NONE;
+	}
+	return Properties->ActiveLayer - Properties->GetNumLockedBaseLayers();
+}
+
+FName FSculptLayersController::GetLayerName(const int32 InLayerIndex) const
+{
+	// ensure the properties are valid
+	if (!ensure(Properties))
+	{
+		return NAME_None;
+	}
+	
+	return Properties->GetLayerName(InLayerIndex);
+}
+
+void FSculptLayersController::SetLayerName(const int32 InLayerIndex, const FName InName) const
+{
+	// ensure the properties are valid
+	if (!ensure(Properties))
+	{
+		return;
+	}
+	
+	return Properties->SetLayerName(InLayerIndex, InName);
+}
+
+
+double FSculptLayersController::GetLayerWeight(const int32 InLayerIndex) const
+{
+	// ensure the properties are valid
+	if (!ensure(Properties))
+	{
+		return 0.0;
+	}
+	
+	// ensure the provided index is valid
+	if (!Properties->LayerWeights.IsValidIndex(InLayerIndex))
+	{
+		UE_LOG(LogMeshLayers, Warning, TEXT("Sculpt Layer not found. Invalid index: %d."), InLayerIndex);
+		return INDEX_NONE;
+	}
+	
+	return Properties->LayerWeights[InLayerIndex];
+}
+
+void FSculptLayersController::SetLayerWeight(const int32 InLayerIndex, const double InWeight, const EPropertyChangeType::Type ChangeType) const
+{
+	// ensure the properties are valid
+	if (!ensure(Properties))
+	{
+		return;
+	}
+	FScopedTransaction Transaction(LOCTEXT("SetLayerWeight_Label", "Set Layer Weight"));
+
+	Properties->SetLayerWeight(InLayerIndex, InWeight, ChangeType);
+}
+
+int32 FSculptLayersController::GetNumMeshLayers() const
+{
+	if (!ensure(Properties))
+	{
+		return INDEX_NONE;
+	}
+	return Properties->LayerWeights.Num();
+}
+
+bool FSculptLayersController::MoveLayerInStack(int32 InLayerToMoveIndex, int32 InTargetIndex)
+{
+	// ensure the properties are valid
+	if (!ensure(Properties))
+	{
+		return false;
+	}
+	FScopedTransaction Transaction(LOCTEXT("ReorderSculptLayers_Label", "Reorder Sculpt Layers"));
+
+	Properties->MoveLayer(InLayerToMoveIndex, InTargetIndex);
+
+	return true;
+}
+
+void FSculptLayersController::RefreshLayersStackView() const
+{
+	if (LayersStackView.IsValid())
+	{
+		LayersStackView->RefreshStackView();
+	}
+}
+
+#undef LOCTEXT_NAMESPACE
diff --git a/Engine/Plugins/Experimental/MeshModelingToolsetExp/Source/ModelingEditorUI/Public/ModelingWidgets/MeshLayersController.h b/Engine/Plugins/Experimental/MeshModelingToolsetExp/Source/ModelingEditorUI/Public/ModelingWidgets/MeshLayersController.h
new file mode 100644
--- /dev/null
+++ b/Engine/Plugins/Experimental/MeshModelingToolsetExp/Source/ModelingEditorUI/Public/ModelingWidgets/MeshLayersController.h
@@ -0,0 +1,93 @@
+﻿// Copyright Epic Games, Inc. All Rights Reserved.
+
+#pragma once
+
+#include "UObject/Object.h"
+#include "UObject/UnrealType.h"
+
+class UMeshSculptLayerProperties;
+class UMegaMeshProjectSculptLayersModifier;
+class SMeshLayersStack;
+
+class MODELINGEDITORUI_API IMeshLayersController
+{
+public:
+	virtual ~IMeshLayersController() = default;
+	
+	/** Adds a new mesh layer to the bottom of the stack
+	 * @return the stack index of the new layer. */
+	virtual int32 AddMeshLayer() { return INDEX_NONE; };
+	
+	/** Remove the mesh layer at the given stack index.
+	 * @param InLayerIndex the index of the layer to remove
+	 * @return true if layer was found and removed, false otherwise */
+	virtual bool RemoveMeshLayer(const int32 InLayerIndex) { return false; };
+
+	/** Sets the layer with the given index as the Active one
+	 * @param InLayerIndex the index of the layer to set as active */
+	virtual void SetActiveLayer(const int32 InLayerIndex) const { };
+
+	/** Gets the layer that is currently active
+	 * @return index of the currently active layer */
+	virtual int32 GetActiveLayer() const { return INDEX_NONE; };
+
+	/** Set the name of the layer at the given index in the stack.
+	 * @param InName the new name to use
+	 * @param InLayerIndex the index of the layer to be renamed */
+	virtual void SetLayerName(const int32 InLayerIndex, const FName InName) const { };
+
+	/** Get the name of the layer at the given index in the stack.
+	 * @param InLayerIndex the index of the layer to get the name of
+	 * @return the name of the layer or None if the index was invalid */
+	virtual FName GetLayerName(const int32 InLayerIndex) const { return FName(); };
+
+	/** Sets the layer weight at the given index in the stack.
+	 * @param InLayerIndex the index of the layer to set the weight of
+	 * @param InWeight the new weight to set the layer to
+	 * @param ChangeType the type of property change we want when setting the layer weight. */
+	virtual void SetLayerWeight(const int32 InLayerIndex, const double InWeight, const EPropertyChangeType::Type ChangeType) const { };
+	
+	/** Gets the layer weight at the given index in the stack.
+	 * @param InLayerIndex the index of the layer to get the weight of
+	 * @return the weight of the layer or none if the index was invalid */
+	virtual double GetLayerWeight(const int32 InLayerIndex) const { return 0.0; };
+
+	/** Get the number of Layers in the stack.
+	 * @return int, the number of layers */
+	virtual int32 GetNumMeshLayers() const = 0;
+	
+	/** Move the mesh layer at the given index to the target index.
+	 * @param InLayerToMoveIndex the index of the layer to be moved
+	 * @param InTargetIndex the index where the layer should be moved to */
+	virtual bool MoveLayerInStack(int32 InLayerToMoveIndex, int32 InTargetIndex)  { return false; };
+	
+	/** Toggle a mesh layer on/off (visible/invisible)
+	 * @param InLayerIndex the index of the layer to toggle the visibility of
+	 * @param bIsEnabled if true, turns the layer on/visible, else off/invisible
+	 * @return true if layer was found at index */
+	//bool SetMeshLayerEnabled(int32 InLayerIndex, bool bIsEnabled) const};
+
+	/** Get the enabled/visibility status of the given layer
+	 * @param InLayerIndex the index of the layer to get the enabled/visibility state for
+	 * @return true of the layer is enabled/visible, false if disabled/invisible or not found */
+	//bool GetMeshLayerEnabled(int32 InLayerIndex) const { return false; };
+
+	/** Refresh the view of the layer stack */
+	virtual void RefreshLayersStackView() const { };
+
+	/** Connect the properties to be displayed and modified */
+	virtual void SetProperties(UMeshSculptLayerProperties* InProperties) { }
+	/** Connect the modifier to be displayed and modified */
+	virtual void SetProperties(UMegaMeshProjectSculptLayersModifier* InProperties) { }
+
+	/** Retrieve the Stack View */
+	virtual TSharedPtr<SMeshLayersStack> GetLayerStackView() { return LayersStackView; };
+	/** Set the stack view */
+	virtual void SetLayerStackView(const TSharedPtr<SMeshLayersStack>& InLayersStackView) { LayersStackView = InLayersStackView; };
+
+protected:
+
+	TSharedPtr<SMeshLayersStack> LayersStackView;
+};
+
+DECLARE_LOG_CATEGORY_EXTERN(LogMeshLayers, Warning, All);
\ No newline at end of file
diff --git a/Engine/Plugins/Experimental/MeshModelingToolsetExp/Source/ModelingEditorUI/Public/ModelingWidgets/SMeshLayersStack.h b/Engine/Plugins/Experimental/MeshModelingToolsetExp/Source/ModelingEditorUI/Public/ModelingWidgets/SMeshLayersStack.h
new file mode 100644
--- /dev/null
+++ b/Engine/Plugins/Experimental/MeshModelingToolsetExp/Source/ModelingEditorUI/Public/ModelingWidgets/SMeshLayersStack.h
@@ -0,0 +1,34 @@
+﻿// Copyright Epic Games, Inc. All Rights Reserved.
+
+#pragma once
+
+#include "EditorUndoClient.h"
+#include "Widgets/SCompoundWidget.h"
+
+class SMeshLayersList;
+class IMeshLayersController;
+
+/**
+ * top level view of a stack of mesh layers
+ * shows the main menu for adding layers to the stack, and a list view of layers
+*/
+class SMeshLayersStack : public SCompoundWidget, public FEditorUndoClient
+{
+public:
+	
+	SLATE_BEGIN_ARGS(SMeshLayersStack)
+		: _InAllowAddRemove(true)
+		, _InAllowReordering(true){}
+	SLATE_ARGUMENT(TWeakPtr<IMeshLayersController>, InController)
+	SLATE_ARGUMENT(bool, InAllowAddRemove)
+	SLATE_ARGUMENT(bool, InAllowReordering)
+	SLATE_END_ARGS()
+
+	MODELINGEDITORUI_API void Construct(const FArguments& InArgs, const TWeakPtr<IMeshLayersController>& InController);
+	
+	MODELINGEDITORUI_API void RefreshStackView() const;
+
+private:
+	TWeakPtr<IMeshLayersController> Controller;
+	TSharedPtr<SMeshLayersList> ListView;
+};
diff --git a/Engine/Plugins/Experimental/MeshModelingToolsetExp/Source/ModelingEditorUI/Public/ModelingWidgets/SculptLayersController.h b/Engine/Plugins/Experimental/MeshModelingToolsetExp/Source/ModelingEditorUI/Public/ModelingWidgets/SculptLayersController.h
new file mode 100644
--- /dev/null
+++ b/Engine/Plugins/Experimental/MeshModelingToolsetExp/Source/ModelingEditorUI/Public/ModelingWidgets/SculptLayersController.h
@@ -0,0 +1,39 @@
+// Copyright Epic Games, Inc. All Rights Reserved.
+
+#pragma once
+
+#include "MeshLayersController.h"
+
+class MODELINGEDITORUI_API FSculptLayersController : public IMeshLayersController
+{
+public:
+	 FSculptLayersController();
+
+	// get access to the Properties.
+	// Warning: do not make modifications to the asset directly, use the controller API
+	UMeshSculptLayerProperties* GetProperties() const { return Properties;};
+	virtual void SetProperties(UMeshSculptLayerProperties* InProperties) override
+	{
+		Properties = const_cast<UMeshSculptLayerProperties*>(InProperties);
+	};
+
+private:
+	TObjectPtr<UMeshSculptLayerProperties> Properties = nullptr;
+
+public:
+	
+	virtual int32 AddMeshLayer() override;
+	virtual bool RemoveMeshLayer(const int32 InLayerIndex) override;
+	
+	virtual void SetActiveLayer(const int32 InLayerIndex) const override;
+	virtual int32 GetActiveLayer() const override;
+	
+	virtual void SetLayerName(const int32 InLayerIndex, const FName InName) const override;
+	virtual FName GetLayerName(const int32 InLayerIndex) const override;
+	virtual void SetLayerWeight(const int32 InLayerIndex, const double InWeight, const EPropertyChangeType::Type ChangeType) const override;
+	virtual double GetLayerWeight(const int32 InLayerIndex) const override;
+	
+	virtual int32 GetNumMeshLayers() const override;
+	virtual bool MoveLayerInStack(int32 InLayerToMoveIndex, int32 InTargetIndex) override;
+	virtual void RefreshLayersStackView() const override;
+};
\ No newline at end of file
diff --git a/Engine/Plugins/Runtime/MeshModelingToolset/Source/MeshModelingTools/Private/Properties/MeshSculptLayerProperties.cpp b/Engine/Plugins/Runtime/MeshModelingToolset/Source/MeshModelingTools/Private/Properties/MeshSculptLayerProperties.cpp
--- a/Engine/Plugins/Runtime/MeshModelingToolset/Source/MeshModelingTools/Private/Properties/MeshSculptLayerProperties.cpp
+++ b/Engine/Plugins/Runtime/MeshModelingToolset/Source/MeshModelingTools/Private/Properties/MeshSculptLayerProperties.cpp
@@ -8,6 +8,8 @@
 
 #include UE_INLINE_GENERATED_CPP_BY_NAME(MeshSculptLayerProperties)
 
+DEFINE_LOG_CATEGORY(LogSculptLayerProps);
+
 #define LOCTEXT_NAMESPACE "UMeshSculptLayerProperties"
 
 void UMeshSculptLayerProperties::Init(IModelingToolExternalDynamicMeshUpdateAPI* InTool, int32 InNumLockedBaseLayers)
@@ -121,6 +123,8 @@ void UMeshSculptLayerProperties::PostEditChangeProperty(FPropertyChangedEvent& P
 
 void UMeshSculptLayerProperties::AddLayer()
 {
+	Modify();
+	
 	EditSculptLayers([this](UE::Geometry::FDynamicMesh3& Mesh, UE::Geometry::FDynamicMeshSculptLayers* SculptLayers)
 		{
 			Mesh.Attributes()->EnableSculptLayers(SculptLayers->NumLayers() + 1);
@@ -128,13 +132,15 @@ void UMeshSculptLayerProperties::AddLayer()
 		}, true);
 }
 
-void UMeshSculptLayerProperties::RemoveLayer()
+bool UMeshSculptLayerProperties::RemoveLayer()
 {
 	if (LayerWeights.Num() <= 1)
 	{
-		return;
+		return false;
 	}
-	EditSculptLayers([this](UE::Geometry::FDynamicMesh3& Mesh, UE::Geometry::FDynamicMeshSculptLayers* SculptLayers)
+	bool bRemoveSuccessful = false;
+	
+	EditSculptLayers([this, &bRemoveSuccessful](UE::Geometry::FDynamicMesh3& Mesh, UE::Geometry::FDynamicMeshSculptLayers* SculptLayers)
 		{
 			if (!SculptLayers->DiscardSculptLayer(ActiveLayer))
 			{
@@ -162,10 +168,200 @@ void UMeshSculptLayerProperties::RemoveLayer()
 					SetLayerWeights(SculptLayers);
 				}
 				ActiveLayer = SculptLayers->SetActiveLayer(NumLockedBaseLayers + NonZeroLayerIdx);
+				bRemoveSuccessful = true;
 			}
 		}, true);
+	return bRemoveSuccessful;
 }
 
+bool UMeshSculptLayerProperties::RemoveLayerAtIndex(const int32 IndexOfLayerToRemove)
+{
+	// ensure the provided index is valid
+	if (!LayerWeights.IsValidIndex(IndexOfLayerToRemove))
+	{
+		UE_LOG(LogSculptLayerProps, Warning, TEXT("Sculpt Layer not removed. Invalid index: %d."), IndexOfLayerToRemove);
+		return false;
+	}
+
+	Modify();
+
+	bool bRemoveSuccessful = false;
+
+	// adjust index to take into account any locked base layers
+	const int32 AdjustedIndex = IndexOfLayerToRemove + NumLockedBaseLayers;
+	
+	if (AdjustedIndex == ActiveLayer)
+	{
+		bRemoveSuccessful = RemoveLayer();
+	}
+	else
+	{
+		EditSculptLayers([this, &AdjustedIndex, &IndexOfLayerToRemove, &bRemoveSuccessful](UE::Geometry::FDynamicMesh3& Mesh, UE::Geometry::FDynamicMeshSculptLayers* SculptLayers)
+		{
+			if (!SculptLayers->DiscardSculptLayer(AdjustedIndex))
+			{
+				return;
+			}
+			if (ActiveLayer > IndexOfLayerToRemove)
+			{
+				// ideally, when deleting the active layer, we should set the prior layer to the new active one,
+				// unless that layer's weight is 0.
+				double LayerWeightTarget = SculptLayers->GetLayerWeights()[ActiveLayer - 1];
+				if (LayerWeightTarget != 0.0)
+				{
+					ActiveLayer = SculptLayers->SetActiveLayer(ActiveLayer - 1);
+				}
+				// find a different layer to set as the active one
+				else
+				{
+					int32 NonZeroLayerIdx = INDEX_NONE;
+					for (int32 Idx = 0; Idx < LayerWeights.Num(); ++Idx)
+					{
+						if (LayerWeights[Idx] != 0.0)
+						{
+							NonZeroLayerIdx = Idx;
+							break;
+						}
+					}
+					// if all layers had zero weight, just pick the first layer and set its weight to 1.0 so it is ready to sculpt on
+					if (NonZeroLayerIdx == INDEX_NONE)
+					{
+						LayerWeights[0] = 1.0;
+						NonZeroLayerIdx = 0;
+						SetLayerWeights(SculptLayers);
+					}
+					ActiveLayer = SculptLayers->SetActiveLayer(NumLockedBaseLayers + NonZeroLayerIdx);
+				}
+			}
+			UpdateSettingsFromMesh(SculptLayers);
+			bRemoveSuccessful = true;
+		}, true);
+	}
+	return bRemoveSuccessful;
+}
+
+FName UMeshSculptLayerProperties::GetLayerName(const int32 InIndex)
+{
+	// ensure the provided index is valid
+	if (!LayerWeights.IsValidIndex(InIndex))
+	{
+		UE_LOG(LogSculptLayerProps, Warning, TEXT("Sculpt Layer not found. Invalid index: %d."), InIndex);
+		return NAME_None;
+	}
+	
+	FName LayerName;
+	EditSculptLayers([this, &InIndex, &LayerName](UE::Geometry::FDynamicMesh3& Mesh, UE::Geometry::FDynamicMeshSculptLayers* SculptLayers)
+		{
+			LayerName = SculptLayers->GetLayer(InIndex + NumLockedBaseLayers)->GetName();
+			UpdateSettingsFromMesh(SculptLayers);
+		}, false);
+	return LayerName;
+}
+
+void UMeshSculptLayerProperties::SetLayerName(const int32 InIndex, const FName InLayerName)
+{
+	// ensure the provided index is valid
+	if (!LayerWeights.IsValidIndex(InIndex))
+	{
+		UE_LOG(LogSculptLayerProps, Warning, TEXT("Sculpt Layer not found. Invalid index: %d."), InIndex);
+		return;
+	}
+	
+	EditSculptLayers([this, &InIndex, &InLayerName](UE::Geometry::FDynamicMesh3& Mesh, UE::Geometry::FDynamicMeshSculptLayers* SculptLayers)
+		{
+			SculptLayers->GetLayer(InIndex + NumLockedBaseLayers)->SetName(InLayerName);
+			UpdateSettingsFromMesh(SculptLayers);
+		}, true);
+}
+
+void UMeshSculptLayerProperties::MoveLayer(const int32 InInitIndex, int32 InTargetIndex)
+{
+	// ensure target is within range
+	InTargetIndex = FMath::Clamp(InTargetIndex, 0, LayerWeights.Num() - 1);
+	
+	if (!LayerWeights.IsValidIndex(InInitIndex))
+	{
+		UE_LOG(LogSculptLayerProps, Warning, TEXT("Sculpt Layer not moved. Invalid source index, %d."), InInitIndex);
+		return;
+	}
+
+	if (!LayerWeights.IsValidIndex(InTargetIndex))
+	{
+		UE_LOG(LogSculptLayerProps, Warning, TEXT("Sculpt Layer not moved. Invalid target index, %d."), InTargetIndex);
+		return;
+	}
+
+	if (InInitIndex == InTargetIndex)
+	{
+		UE_LOG(LogSculptLayerProps, Warning, TEXT("Sculpt Layer not moved. Source and target index cannot be the same."));
+		return;
+	}
+	
+	Modify();
+	
+	// adjust indices to take into account any locked base layers
+	int32 AdjustedInitIndex = InInitIndex + NumLockedBaseLayers;
+	int32 AdjustedTargetIndex = InTargetIndex + NumLockedBaseLayers;
+	EditSculptLayers([this, &AdjustedInitIndex, &AdjustedTargetIndex](UE::Geometry::FDynamicMesh3& Mesh, UE::Geometry::FDynamicMeshSculptLayers* SculptLayers)
+		{
+			SculptLayers->MoveLayer(AdjustedInitIndex, AdjustedTargetIndex);
+			UpdateSettingsFromMesh(SculptLayers);
+		}, true);
+}
+
+void UMeshSculptLayerProperties::SetActiveLayer(const int32 InIndex)
+{
+	// ensure the provided index is valid
+	if (!LayerWeights.IsValidIndex(InIndex))
+	{
+		UE_LOG(LogSculptLayerProps, Warning, TEXT("Sculpt Layer not found. Invalid index: %d."), InIndex);
+		return;
+	}
+	// account for the locked base layers
+	const int32 AdjustedIndex = InIndex + NumLockedBaseLayers;
+
+	// if the provided index is the already active layer, do nothing
+	if (AdjustedIndex == ActiveLayer)
+	{
+		return;
+	}
+
+	ActiveLayer = FMath::Clamp(AdjustedIndex, NumLockedBaseLayers, LayerWeights.Num());
+	FProperty* ActiveLayerProp = FindFProperty<FProperty>(UMeshSculptLayerProperties::StaticClass(),
+		GET_MEMBER_NAME_CHECKED(UMeshSculptLayerProperties, ActiveLayer));
+	FPropertyChangedEvent PropertyChangedEvent(ActiveLayerProp);
+	#if WITH_EDITOR
+	PostEditChangeProperty(PropertyChangedEvent);
+	#endif
+}
+
+void UMeshSculptLayerProperties::SetLayerWeight(const int32 InIndex, double InWeight, const EPropertyChangeType::Type ChangeType)
+{
+	// ensure the provided index is valid
+	if (!LayerWeights.IsValidIndex(InIndex))
+	{
+		UE_LOG(LogSculptLayerProps, Warning, TEXT("Sculpt Layer not found. Invalid index: %d."), InIndex);
+		return;
+	}
+	
+	Modify();
+	
+	// adjust index to take into account any locked base layers
+	const int32 AdjustedIndex = InIndex + NumLockedBaseLayers;
+
+	// do not allow the ActiveLayer's weight to be set to 0
+	if (AdjustedIndex == ActiveLayer && InWeight == 0.0)
+	{
+		InWeight = 0.1;
+	}
+	
+	LayerWeights[InIndex] = InWeight;
+
+	EditSculptLayers([this, &AdjustedIndex, &InWeight](UE::Geometry::FDynamicMesh3& Mesh, UE::Geometry::FDynamicMeshSculptLayers* SculptLayers)
+	{
+		SetLayerWeights(SculptLayers);
+	}, ChangeType != EPropertyChangeType::Interactive);
+}
 
 #undef LOCTEXT_NAMESPACE
 
diff --git a/Engine/Plugins/Runtime/MeshModelingToolset/Source/MeshModelingTools/Private/Properties/MeshSculptLayerProperties.h b/Engine/Plugins/Runtime/MeshModelingToolset/Source/MeshModelingTools/Private/Properties/MeshSculptLayerProperties.h
deleted file mode 100644
diff --git a/Engine/Plugins/Runtime/MeshModelingToolset/Source/MeshModelingTools/Public/Properties/MeshSculptLayerProperties.h b/Engine/Plugins/Runtime/MeshModelingToolset/Source/MeshModelingTools/Public/Properties/MeshSculptLayerProperties.h
new file mode 100644
--- /dev/null
+++ b/Engine/Plugins/Runtime/MeshModelingToolset/Source/MeshModelingTools/Public/Properties/MeshSculptLayerProperties.h
@@ -0,0 +1,92 @@
+// Copyright Epic Games, Inc. All Rights Reserved.
+
+#pragma once
+
+#include "UObject/Object.h"
+#include "GeometryBase.h"
+
+#include "MeshSculptLayerProperties.generated.h"
+
+PREDECLARE_USE_GEOMETRY_CLASS(FDynamicMesh3);
+PREDECLARE_USE_GEOMETRY_CLASS(FDynamicMeshSculptLayers);
+class IModelingToolExternalDynamicMeshUpdateAPI;
+
+UCLASS(MinimalAPI)
+class UMeshSculptLayerProperties : public UObject
+{
+	GENERATED_BODY()
+
+public:
+
+	/** Set the active mesh layer */
+	UPROPERTY(EditAnywhere, Category = MeshLayers, meta = (ClampMin = 0, HideEditConditionToggle, EditCondition = bCanEditLayers, ModelingQuickSettings))
+	int32 ActiveLayer = 0;
+
+	/** Set the mesh layer weights */
+	UPROPERTY(EditAnywhere, EditFixedSize, NoClear, Category = MeshLayers, meta = (NoResetToDefault, HideEditConditionToggle, EditCondition = bCanEditLayers, ModelingQuickSettings))
+	TArray<double> LayerWeights;
+
+	UPROPERTY(meta = (TransientToolProperty))
+	bool bCanEditLayers = false;
+
+	/** Add a sculpt layer */
+	MESHMODELINGTOOLS_API void AddLayer();
+
+	/** Remove the active sculpt layer */
+	MESHMODELINGTOOLS_API bool RemoveLayer();
+
+	/** Remove the sculpt layer at the provided index */
+	UFUNCTION(Category = MeshLayers)
+	MESHMODELINGTOOLS_API bool RemoveLayerAtIndex(const int32 IndexOfLayerToRemove);
+
+	/** Return the name of the sculpt layer at the provided index */
+	UFUNCTION(Category = MeshLayers)
+	MESHMODELINGTOOLS_API FName GetLayerName(const int32 InIndex);
+
+	/** Set the name of the layer at the provided index to the provided name */
+	UFUNCTION(Category = MeshLayers)
+	MESHMODELINGTOOLS_API void SetLayerName(const int32 InIndex, const FName InLayerName);
+
+	/** Move the layer at the first provided index to the location provided by the second index */
+	UFUNCTION(Category = MeshLayers)
+	MESHMODELINGTOOLS_API void MoveLayer(const int32 InInitIndex, int32 InTargetIndex);
+
+	/** Sets the currently active sculpt layer to be the one at the provided index */
+	UFUNCTION(Category = MeshLayers)
+	MESHMODELINGTOOLS_API void SetActiveLayer(const int32 InIndex);
+
+	/** Sets the weight of the layer at the provided index to the provided weight */
+	UFUNCTION(Category = MeshLayers)
+	MESHMODELINGTOOLS_API void SetLayerWeight(const int32 InIndex, double InWeight, const uint32 ChangeType);
+
+	MESHMODELINGTOOLS_API void Init(IModelingToolExternalDynamicMeshUpdateAPI* InTool, int32 InNumLockedBaseLayers);
+
+#if WITH_EDITOR
+	MESHMODELINGTOOLS_API virtual void PostEditChangeProperty(FPropertyChangedEvent& PropertyChangedEvent) override;
+#endif
+
+	int32 GetNumLockedBaseLayers() const { return NumLockedBaseLayers; }
+
+private:
+
+	IModelingToolExternalDynamicMeshUpdateAPI* Tool;
+
+	int32 NumLockedBaseLayers = 0;
+
+	// Mesh before layer change has been applied. Used for tracking mesh changes that occur over multiple frames (i.e., from an interactive drag)
+	// TODO: Add/update a mesh FChange to track sculpt layer changes, and use that instead of saving an entire mesh here
+	TSharedPtr<FDynamicMesh3> InitialMesh;
+
+	// Helper to set sculpt layers from the current LayerWeights property (accounting for the NumLockedBaseLayers)
+	void SetLayerWeights(FDynamicMeshSculptLayers* SculptLayers) const;
+
+	// Update the ActiveLayer and LayerWeights settings from the current sculpt layers
+	void UpdateSettingsFromMesh(const FDynamicMeshSculptLayers* SculptLayers);
+
+	// Helper to apply edits to the current sculpt layers if possible, with associated book keeping
+	// @param EditFn The edit to apply if possible
+	// @param bEmitChange Whether to emit a change object along with the edit
+	void EditSculptLayers(TFunctionRef<void(UE::Geometry::FDynamicMesh3& Mesh, UE::Geometry::FDynamicMeshSculptLayers* SculptLayers)> EditFn, bool bEmitChange);
+};
+
+DECLARE_LOG_CATEGORY_EXTERN(LogSculptLayerProps, Warning, All);
\ No newline at end of file
diff --git a/Engine/Source/Runtime/GeometryCore/Private/DynamicMesh/DynamicMeshSculptLayers.cpp b/Engine/Source/Runtime/GeometryCore/Private/DynamicMesh/DynamicMeshSculptLayers.cpp
--- a/Engine/Source/Runtime/GeometryCore/Private/DynamicMesh/DynamicMeshSculptLayers.cpp
+++ b/Engine/Source/Runtime/GeometryCore/Private/DynamicMesh/DynamicMeshSculptLayers.cpp
@@ -29,6 +29,28 @@ bool FDynamicMeshSculptLayers::DiscardSculptLayer(int32 LayerIndex)
 	return true;
 }
 
+bool FDynamicMeshSculptLayers::MoveLayer(const int32 InitLayerIndex, const int32 TargetIndex)
+{
+	if (!Layers.IsValidIndex(InitLayerIndex) || !Layers.IsValidIndex(TargetIndex) || InitLayerIndex == TargetIndex)
+	{
+		return false;
+	}
+	UpdateLayersFromMesh();
+
+	const int Dir = InitLayerIndex < TargetIndex ? 1 : -1;
+	int32 CurrentMovedLayerIndex = InitLayerIndex;
+	while (CurrentMovedLayerIndex != TargetIndex)
+	{
+		Layers.Swap(CurrentMovedLayerIndex, CurrentMovedLayerIndex  + Dir);
+		LayerWeights.Swap(CurrentMovedLayerIndex, CurrentMovedLayerIndex + Dir);
+		CurrentMovedLayerIndex += Dir;
+	}
+
+	UpdateMeshFromLayers();
+	ValidateActiveLayer();
+	return true;
+}
+
 bool FDynamicMeshSculptLayers::MergeSculptLayers(int32 StartIndex, int32 EndIndex, bool bUseWeights)
 {
 	if (!Layers.IsValidIndex(StartIndex) || !Layers.IsValidIndex(EndIndex) || StartIndex >= EndIndex)
diff --git a/Engine/Source/Runtime/GeometryCore/Public/DynamicMesh/DynamicMeshSculptLayers.h b/Engine/Source/Runtime/GeometryCore/Public/DynamicMesh/DynamicMeshSculptLayers.h
--- a/Engine/Source/Runtime/GeometryCore/Public/DynamicMesh/DynamicMeshSculptLayers.h
+++ b/Engine/Source/Runtime/GeometryCore/Public/DynamicMesh/DynamicMeshSculptLayers.h
@@ -55,6 +55,9 @@ public:
 	// Set new sculpt layer weights
 	GEOMETRYCORE_API void UpdateLayerWeights(TConstArrayView<double> InLayerWeights);
 
+	// moves the layer at the first provided index to the location given by the second provided index
+	GEOMETRYCORE_API bool MoveLayer(const int32 InitLayerIndex, const int32 TargetIndex);
+
 	// Directly access the sculpt layer data
 	const FDynamicMeshSculptLayerAttribute* GetLayer(int32 LayerIndex) const
 	{
