ModelingTools: Remove some unused lines from DrawSplineTool to fix static analysis warning.

diff --git a/Engine/Plugins/Experimental/MeshModelingToolsetExp/Source/MeshModelingToolsEditorOnlyExp/Private/DrawSplineTool.cpp b/Engine/Plugins/Experimental/MeshModelingToolsetExp/Source/MeshModelingToolsEditorOnlyExp/Private/DrawSplineTool.cpp
--- a/Engine/Plugins/Experimental/MeshModelingToolsetExp/Source/MeshModelingToolsEditorOnlyExp/Private/DrawSplineTool.cpp
+++ b/Engine/Plugins/Experimental/MeshModelingToolsetExp/Source/MeshModelingToolsEditorOnlyExp/Private/DrawSplineTool.cpp
@@ -1005,8 +1005,6 @@ void UDrawSplineTool::OnTerminateDragSequence()
 	{
 	case EDrawSplineDrawMode::ClickAutoTangent:
 	{
-		WorkingSpline->GetLocationAtSplinePoint(NumSplinePoints-1, ESplineCoordinateSpace::World);
-		WorkingSpline->GetUpVectorAtSplinePoint(NumSplinePoints - 1, ESplineCoordinateSpace::World);
 		GetToolManager()->EmitObjectChange(this,
 			MakeUnique<FSimplePointInsertionChange>(
 				WorkingSpline->GetLocationAtSplinePoint(NumSplinePoints - 1, ESplineCoordinateSpace::World),
