Modeling Mode: Fix Sculpt Layers undo/redo bug and other minor cleanup

diff --git a/Engine/Plugins/Experimental/MeshModelingToolsetExp/Source/ModelingEditorUI/Private/ModelingWidgets/SMeshLayers.cpp b/Engine/Plugins/Experimental/MeshModelingToolsetExp/Source/ModelingEditorUI/Private/ModelingWidgets/SMeshLayers.cpp
--- a/Engine/Plugins/Experimental/MeshModelingToolsetExp/Source/ModelingEditorUI/Private/ModelingWidgets/SMeshLayers.cpp
+++ b/Engine/Plugins/Experimental/MeshModelingToolsetExp/Source/ModelingEditorUI/Private/ModelingWidgets/SMeshLayers.cpp
@@ -10,6 +10,7 @@
 #include "Widgets/Input/SSpinBox.h"
 #include "Widgets/Layout/SSpacer.h"
 #include "Widgets/Text/SInlineEditableTextBlock.h"
+#include "Editor.h"
 
 #define LOCTEXT_NAMESPACE "SMeshLayers"
 
@@ -71,7 +72,7 @@ FReply SMeshLayersList::OnKeyDown(const FGeometry& MyGeometry, const FKeyEvent&
 
 void SMeshLayersList::DeleteMeshLayer(TSharedPtr<FMeshLayerElement> LayerToDelete)
 {
-	if (!LayerToDelete.IsValid())
+	if (!Controller.IsValid() || !LayerToDelete.IsValid())
 	{
 		return;
 	}
@@ -82,9 +83,12 @@ void SMeshLayersList::DeleteMeshLayer(TSharedPtr<FMeshLayerElement> LayerToDelet
 		return;
 	}
 
-	// remove layer from this List of layers
+	// remove layer from this List of layers - not technically necessary if base layer is removed above and List is being refreshed below,
+	// since the refresh will generate a new Layers list based on the updated layers (without the deleted one) from the Controller
 	const int32 StackIndexToDelete = LayerToDelete->GetIndexInStack();
+	Layers.RemoveAt(StackIndexToDelete);
 
+	// correct indices of other layers
 	for (TSharedPtr<FMeshLayerElement> Layer : Layers)
 	{
 		const int32 StackIndex = Layer->GetIndexInStack();
@@ -400,16 +404,31 @@ void SMeshLayerItem::Construct(const FArguments& InArgs, const TSharedRef<STable
 			IMeshLayersController* Ctrl = ListView.Pin()->GetController().Pin().Get();
 			return Ctrl->GetLayerWeight(Element.Pin()->GetIndexInStack());
 		})
+		.OnBeginSliderMovement_Lambda([this]()
+		{
+			GEditor->BeginTransaction(LOCTEXT("MeshLayers", "Set Layer Weight"));
+		})
 		.OnValueChanged_Lambda([this](const float NewValue)
 		{
+			// in the case of text entry, no transaction has been created yet
+			const bool bIsTextEntry = !GEditor->IsTransactionActive();
+			if (bIsTextEntry)
+			{
+				GEditor->BeginTransaction(LOCTEXT("MeshLayers", "Set Layer Weight"));
+			}
 			IMeshLayersController* Ctrl = ListView.Pin()->GetController().Pin().Get();
 			Ctrl->SetLayerWeight(Element.Pin()->GetIndexInStack(), NewValue, EPropertyChangeType::Interactive);
+			if (bIsTextEntry)
+			{
+				GEditor->EndTransaction();
+			}
 		})
 		// only log the slider movement after its complete
 		.OnEndSliderMovement_Lambda([this](float LastValue)
 		{
 			IMeshLayersController* Ctrl = ListView.Pin()->GetController().Pin().Get();
 			Ctrl->SetLayerWeight(Element.Pin()->GetIndexInStack(), LastValue, EPropertyChangeType::ValueSet);
+			GEditor->EndTransaction();
 		})
 		.MinValue(-1.0f)
 		.MaxValue(1.5f)
diff --git a/Engine/Plugins/Runtime/MeshModelingToolset/Source/MeshModelingTools/Public/Properties/MeshSculptLayerProperties.h b/Engine/Plugins/Runtime/MeshModelingToolset/Source/MeshModelingTools/Public/Properties/MeshSculptLayerProperties.h
--- a/Engine/Plugins/Runtime/MeshModelingToolset/Source/MeshModelingTools/Public/Properties/MeshSculptLayerProperties.h
+++ b/Engine/Plugins/Runtime/MeshModelingToolset/Source/MeshModelingTools/Public/Properties/MeshSculptLayerProperties.h
@@ -69,7 +69,7 @@ public:
 
 private:
 
-	IModelingToolExternalDynamicMeshUpdateAPI* Tool;
+	IModelingToolExternalDynamicMeshUpdateAPI* Tool = nullptr;
 
 	int32 NumLockedBaseLayers = 0;
 
