Make SubdividePolyTool's OnTick a virtual override method

diff --git a/Engine/Plugins/Experimental/MeshModelingToolsetExp/Source/MeshModelingToolsEditorOnlyExp/Public/SubdividePolyTool.h b/Engine/Plugins/Experimental/MeshModelingToolsetExp/Source/MeshModelingToolsEditorOnlyExp/Public/SubdividePolyTool.h
--- a/Engine/Plugins/Experimental/MeshModelingToolsetExp/Source/MeshModelingToolsEditorOnlyExp/Public/SubdividePolyTool.h
+++ b/Engine/Plugins/Experimental/MeshModelingToolsetExp/Source/MeshModelingToolsEditorOnlyExp/Public/SubdividePolyTool.h
@@ -100,7 +100,7 @@ public:
 	virtual bool HasAccept() const override { return true; }
 	UE_API virtual bool CanAccept() const override;
 
-	UE_API void OnTick(float DeltaTime);
+	UE_API virtual void OnTick(float DeltaTime) override;
 
 protected:
 
