[MeshPartition]

Rename part 1.

This rename is the first of several changelists which renames and moves the MegaMesh classes to a new namespace UE::MeshPartition.
Not all classes are renamed. This changelist focuses on the UObjects and especially the UObjects which require redirectors. All uobjects which correspond to saved content should be renamed as of this cl and include a redirector to map to the new class names.

diff --git a/Engine/Plugins/Experimental/MeshModelingToolsetExp/Source/ModelingEditorUI/Public/ModelingWidgets/MeshLayersController.h b/Engine/Plugins/Experimental/MeshModelingToolsetExp/Source/ModelingEditorUI/Public/ModelingWidgets/MeshLayersController.h
--- a/Engine/Plugins/Experimental/MeshModelingToolsetExp/Source/ModelingEditorUI/Public/ModelingWidgets/MeshLayersController.h
+++ b/Engine/Plugins/Experimental/MeshModelingToolsetExp/Source/ModelingEditorUI/Public/ModelingWidgets/MeshLayersController.h
@@ -6,9 +6,13 @@
 #include "UObject/UnrealType.h"
 
 class UMeshSculptLayerProperties;
-class UMegaMeshProjectSculptLayersModifier;
 class SMeshLayersStack;
 
+namespace UE::MeshPartition
+{
+	class UProjectMeshLayersModifier;
+}
+
 class IMeshLayersController
 {
 public:
@@ -78,7 +82,7 @@ public:
 	/** Connect the properties to be displayed and modified */
 	virtual void SetProperties(UMeshSculptLayerProperties* InProperties) { }
 	/** Connect the modifier to be displayed and modified */
-	virtual void SetProperties(UMegaMeshProjectSculptLayersModifier* InProperties) { }
+	virtual void SetProperties(UE::MeshPartition::UProjectMeshLayersModifier* InProperties) { }
 
 	/** Retrieve the Stack View */
 	virtual TSharedPtr<SMeshLayersStack> GetLayerStackView() { return LayersStackView; };
