ModelingTools: Fix alpha texture picker filtering being broken in sculpt tool.

diff --git a/Engine/Plugins/Experimental/MeshModelingToolsetExp/Source/ModelingEditorUI/Private/ModelingWidgets/SToolInputAssetPicker.cpp b/Engine/Plugins/Experimental/MeshModelingToolsetExp/Source/ModelingEditorUI/Private/ModelingWidgets/SToolInputAssetPicker.cpp
--- a/Engine/Plugins/Experimental/MeshModelingToolsetExp/Source/ModelingEditorUI/Private/ModelingWidgets/SToolInputAssetPicker.cpp
+++ b/Engine/Plugins/Experimental/MeshModelingToolsetExp/Source/ModelingEditorUI/Private/ModelingWidgets/SToolInputAssetPicker.cpp
@@ -171,6 +171,11 @@ void SToolInputAssetPicker::Construct( const FArguments& InArgs )
 		}
 	}
 
+	if (!InArgs._AssetPickerConfig.bAutohideSearchBar)
+	{
+		TextFilter = MakeShared<FAssetTextFilter>();
+	}
+
 	VerticalBox->AddSlot()
 	.FillHeight(1.f)
 	[
@@ -185,6 +190,7 @@ void SToolInputAssetPicker::Construct( const FArguments& InArgs )
 		.OnVisualizeAssetToolTip(InArgs._AssetPickerConfig.OnVisualizeAssetToolTip)
 		.OnAssetToolTipClosing(InArgs._AssetPickerConfig.OnAssetToolTipClosing)
 		.FrontendFilters(FrontendFilters)
+		.TextFilter(TextFilter)
 		.InitialContentSources(CurrentContentSources)
 		.InitialBackendFilter(CurrentBackendFilter)
 		.InitialViewType(InArgs._AssetPickerConfig.InitialAssetViewType)
@@ -216,10 +222,8 @@ void SToolInputAssetPicker::Construct( const FArguments& InArgs )
 		.ShowTypeInTileView(false)
 	];
 
-
-	if (AssetViewPtr.IsValid() && !InArgs._AssetPickerConfig.bAutohideSearchBar)
+	if (AssetViewPtr.IsValid() && TextFilter)
 	{
-		TextFilter = MakeShared<FAssetTextFilter>();
 		bool bClassNamesProvided = (InArgs._AssetPickerConfig.Filter.ClassPaths.Num() != 1);
 		TextFilter->SetIncludeClassName(bClassNamesProvided || AssetViewPtr->IsIncludingClassNames());
 		TextFilter->SetIncludeAssetPath(AssetViewPtr->IsIncludingAssetPaths());
