Modeling Mode - fix mesh layers naming bug

diff --git a/Engine/Plugins/Experimental/MeshModelingToolsetExp/Source/ModelingEditorUI/Private/ModelingWidgets/SMeshLayers.cpp b/Engine/Plugins/Experimental/MeshModelingToolsetExp/Source/ModelingEditorUI/Private/ModelingWidgets/SMeshLayers.cpp
--- a/Engine/Plugins/Experimental/MeshModelingToolsetExp/Source/ModelingEditorUI/Private/ModelingWidgets/SMeshLayers.cpp
+++ b/Engine/Plugins/Experimental/MeshModelingToolsetExp/Source/ModelingEditorUI/Private/ModelingWidgets/SMeshLayers.cpp
@@ -361,14 +361,6 @@ void SMeshLayerItem::Construct(const FArguments& InArgs, const TSharedRef<STable
 		{
 			const IMeshLayersController* Ctrl = ListView.Pin()->GetController().Pin().Get();
 			const FName LayerName = Ctrl->GetLayerName(Element.Pin()->GetIndexInStack());
-			// TODO
-			// if layer name has not been set, display a unique placeholder name like "Layer 2" to
-			// eliminate the confusion of having multiple layers with the display label "None"
-			/*if (LayerName == NAME_None)
-			{
-				return FText::FromName(Ctrl->GetPlaceholderNames()[Element.Pin()->GetIndexInStack()]);
-			}*/
-			// otherwise, display the layer's actual name
 			return FText::FromName(LayerName);
 		})
 		.OnVerifyTextChanged_Lambda([](const FText& InText, FText& OutErrorMessage)
@@ -496,9 +488,11 @@ void SMeshLayerItem::Construct(const FArguments& InArgs, const TSharedRef<STable
 
 void SMeshLayerItem::OnNameCommitted(const FText& InText, ETextCommit::Type InCommitType) const
 {
+	GEditor->BeginTransaction(LOCTEXT("MeshLayers", "Change Mesh Layer Name"));
 	const FName NewName = FName(InText.ToString());
 	const IMeshLayersController* Ctrl = ListView.Pin()->GetController().Pin().Get();
 	Ctrl->SetLayerName(Element.Pin()->GetIndexInStack(), NewName);
+	GEditor->EndTransaction();
 }
 
 FText SMeshLayerItem::GetName() const
