ModelingTools: Updated the OutputType properties' display name and tooltip to SourceType in the BakeVertex tool to better reflect the function of those properties.

diff --git a/Engine/Plugins/Experimental/MeshModelingToolsetExp/Source/MeshModelingToolsExp/Public/BakeMeshAttributeVertexTool.h b/Engine/Plugins/Experimental/MeshModelingToolsetExp/Source/MeshModelingToolsExp/Public/BakeMeshAttributeVertexTool.h
--- a/Engine/Plugins/Experimental/MeshModelingToolsetExp/Source/MeshModelingToolsExp/Public/BakeMeshAttributeVertexTool.h
+++ b/Engine/Plugins/Experimental/MeshModelingToolsetExp/Source/MeshModelingToolsExp/Public/BakeMeshAttributeVertexTool.h
@@ -75,32 +75,37 @@ public:
 	UPROPERTY(EditAnywhere, Category = BakeOutput)
 	EBakeVertexOutput OutputMode = EBakeVertexOutput::RGBA;
 
-	/** The bake output type to generate */
+	/** The source data to bake to vertex colors */
 	UPROPERTY(EditAnywhere, Category = BakeOutput, meta=(Bitmask, BitmaskEnum = "/Script/MeshModelingToolsExp.EBakeMapType",
+		DisplayName="Source Type",
 		ValidEnumValues="TangentSpaceNormal, AmbientOcclusion, BentNormal, Curvature, Texture, ObjectSpaceNormal, FaceNormal, Position, MaterialID, PolyGroupID, MultiTexture, VertexColor, Height",
 		EditCondition="OutputMode == EBakeVertexOutput::RGBA", EditConditionHides))
 	int32 OutputType = static_cast<int32>(EBakeMapType::TangentSpaceNormal);
 
-	/** The bake output type to generate in the Red channel */
+	/** The source data to bake to the Red channel */
 	UPROPERTY(EditAnywhere, Category = BakeOutput, meta=(Bitmask, BitmaskEnum = "/Script/MeshModelingToolsExp.EBakeMapType",
+		DisplayName="Source Type R",
 		ValidEnumValues="None, AmbientOcclusion, Curvature, Height, One, Zero",
 		EditCondition="OutputMode == EBakeVertexOutput::PerChannel", EditConditionHides))
 	int32 OutputTypeR = static_cast<int32>(EBakeMapType::None);
 
-	/** The bake output type to generate in the Green channel */
+	/** The source data to bake to the Green channel */
 	UPROPERTY(EditAnywhere, Category = BakeOutput, meta=(Bitmask, BitmaskEnum = "/Script/MeshModelingToolsExp.EBakeMapType",
+		DisplayName="Source Type G",
 		ValidEnumValues="None, AmbientOcclusion, Curvature, Height, One, Zero",
 		EditCondition="OutputMode == EBakeVertexOutput::PerChannel", EditConditionHides))
 	int32 OutputTypeG = static_cast<int32>(EBakeMapType::None);
 
-	/** The bake output type to generate in the Blue channel */
+	/** The source data to bake to the Blue channel */
 	UPROPERTY(EditAnywhere, Category = BakeOutput, meta=(Bitmask, BitmaskEnum = "/Script/MeshModelingToolsExp.EBakeMapType",
+		DisplayName="Source Type B",
 		ValidEnumValues="None, AmbientOcclusion, Curvature, Height, One, Zero",
 		EditCondition="OutputMode == EBakeVertexOutput::PerChannel", EditConditionHides))
 	int32 OutputTypeB = static_cast<int32>(EBakeMapType::None);
 
-	/** The bake output type to generate in the Alpha channel */
+	/** The source data to bake to the Alpha channel */
 	UPROPERTY(EditAnywhere, Category = BakeOutput, meta=(Bitmask, BitmaskEnum = "/Script/MeshModelingToolsExp.EBakeMapType",
+		DisplayName="Source Type A",
 		ValidEnumValues="None, AmbientOcclusion, Curvature, Height, One, Zero",
 		EditCondition="OutputMode == EBakeVertexOutput::PerChannel", EditConditionHides))
 	int32 OutputTypeA = static_cast<int32>(EBakeMapType::None);
