Modeling Tools: fix transaction naming conflict

diff --git a/Engine/Plugins/Experimental/MeshModelingToolsetExp/Source/ModelingEditorUI/Private/ModelingWidgets/SMeshLayers.cpp b/Engine/Plugins/Experimental/MeshModelingToolsetExp/Source/ModelingEditorUI/Private/ModelingWidgets/SMeshLayers.cpp
--- a/Engine/Plugins/Experimental/MeshModelingToolsetExp/Source/ModelingEditorUI/Private/ModelingWidgets/SMeshLayers.cpp
+++ b/Engine/Plugins/Experimental/MeshModelingToolsetExp/Source/ModelingEditorUI/Private/ModelingWidgets/SMeshLayers.cpp
@@ -398,7 +398,7 @@ void SMeshLayerItem::Construct(const FArguments& InArgs, const TSharedRef<STable
 		})
 		.OnBeginSliderMovement_Lambda([this]()
 		{
-			GEditor->BeginTransaction(LOCTEXT("MeshLayers", "Set Layer Weight"));
+			GEditor->BeginTransaction(LOCTEXT("MeshLayers_SetWeight_Slider", "Set Layer Weight"));
 		})
 		.OnValueChanged_Lambda([this](const float NewValue)
 		{
@@ -406,7 +406,7 @@ void SMeshLayerItem::Construct(const FArguments& InArgs, const TSharedRef<STable
 			const bool bIsTextEntry = !GEditor->IsTransactionActive();
 			if (bIsTextEntry)
 			{
-				GEditor->BeginTransaction(LOCTEXT("MeshLayers", "Set Layer Weight"));
+				GEditor->BeginTransaction(LOCTEXT("MeshLayers_SetWeight_ValueChange", "Set Layer Weight"));
 			}
 			IMeshLayersController* Ctrl = ListView.Pin()->GetController().Pin().Get();
 			Ctrl->SetLayerWeight(Element.Pin()->GetIndexInStack(), NewValue, EPropertyChangeType::Interactive);
@@ -488,7 +488,7 @@ void SMeshLayerItem::Construct(const FArguments& InArgs, const TSharedRef<STable
 
 void SMeshLayerItem::OnNameCommitted(const FText& InText, ETextCommit::Type InCommitType) const
 {
-	GEditor->BeginTransaction(LOCTEXT("MeshLayers", "Change Mesh Layer Name"));
+	GEditor->BeginTransaction(LOCTEXT("MeshLayers_ChangeName", "Change Mesh Layer Name"));
 	const FName NewName = FName(InText.ToString());
 	const IMeshLayersController* Ctrl = ListView.Pin()->GetController().Pin().Get();
 	Ctrl->SetLayerName(Element.Pin()->GetIndexInStack(), NewName);
