Change mesh layer's layer weight limits to 'slider' limits, so typed values outside the (-1,1.5) range will be accepted

diff --git a/Engine/Plugins/Experimental/MeshModelingToolsetExp/Source/ModelingEditorUI/Private/ModelingWidgets/SMeshLayers.cpp b/Engine/Plugins/Experimental/MeshModelingToolsetExp/Source/ModelingEditorUI/Private/ModelingWidgets/SMeshLayers.cpp
--- a/Engine/Plugins/Experimental/MeshModelingToolsetExp/Source/ModelingEditorUI/Private/ModelingWidgets/SMeshLayers.cpp
+++ b/Engine/Plugins/Experimental/MeshModelingToolsetExp/Source/ModelingEditorUI/Private/ModelingWidgets/SMeshLayers.cpp
@@ -1,4 +1,4 @@
-﻿// Copyright Epic Games, Inc. All Rights Reserved.
+// Copyright Epic Games, Inc. All Rights Reserved.
 
 #include "ModelingWidgets/SMeshLayers.h"
 
@@ -422,8 +422,8 @@ void SMeshLayerItem::Construct(const FArguments& InArgs, const TSharedRef<STable
 			Ctrl->SetLayerWeight(Element.Pin()->GetIndexInStack(), LastValue, EPropertyChangeType::ValueSet);
 			GEditor->EndTransaction();
 		})
-		.MinValue(-1.0f)
-		.MaxValue(1.5f)
+		.MinSliderValue(-1.0f)
+		.MaxSliderValue(1.5f)
 		.Style(&FAppStyle::Get().GetWidgetStyle<FSpinBoxStyle>("Sequencer.HyperlinkSpinBox"))
 	];
 
