fix crash on attempting to split meshes w/ invalid material IDs

diff --git a/Engine/Plugins/Experimental/MeshModelingToolsetExp/Source/MeshModelingToolsExp/Private/SplitMeshesTool.cpp b/Engine/Plugins/Experimental/MeshModelingToolsetExp/Source/MeshModelingToolsExp/Private/SplitMeshesTool.cpp
--- a/Engine/Plugins/Experimental/MeshModelingToolsetExp/Source/MeshModelingToolsExp/Private/SplitMeshesTool.cpp
+++ b/Engine/Plugins/Experimental/MeshModelingToolsetExp/Source/MeshModelingToolsExp/Private/SplitMeshesTool.cpp
@@ -338,7 +338,7 @@ void USplitMeshesTool::UpdateSplitMeshes()
 					{
 						int32 NewMaterialID = UniqueIDs.Num();
 						UniqueIDs.Add(MaterialID);
-						NewMaterials.Add(SourceMaterials[MaterialID]);
+						NewMaterials.Add(SourceMaterials.IsValidIndex(MaterialID) ? SourceMaterials[MaterialID] : nullptr);
 						MaterialIDs->SetValue(tid, NewMaterialID);
 					}
 					else
