Add support for the tangents tool to choose the reference UV layer to use when computing tangents

diff --git a/Engine/Plugins/Experimental/MeshModelingToolsetExp/Source/MeshModelingToolsEditorOnlyExp/Private/MeshTangentsTool.cpp b/Engine/Plugins/Experimental/MeshModelingToolsetExp/Source/MeshModelingToolsEditorOnlyExp/Private/MeshTangentsTool.cpp
--- a/Engine/Plugins/Experimental/MeshModelingToolsetExp/Source/MeshModelingToolsEditorOnlyExp/Private/MeshTangentsTool.cpp
+++ b/Engine/Plugins/Experimental/MeshModelingToolsetExp/Source/MeshModelingToolsEditorOnlyExp/Private/MeshTangentsTool.cpp
@@ -123,6 +123,11 @@ void UMeshTangentsTool::Setup()
 	Settings = NewObject<UMeshTangentsToolProperties>(this);
 	Settings->RestoreProperties(this);
 	AddToolPropertySource(Settings);
+	UVChannelProperties = NewObject<UMeshUVChannelProperties>(this);
+	UVChannelProperties->RestoreProperties(this);
+	PreviewMesh->ProcessMesh([&](const FDynamicMesh3& ReadMesh) { UVChannelProperties->Initialize(&ReadMesh, false); });
+	UVChannelProperties->ValidateSelection(true);
+	AddToolPropertySource(UVChannelProperties);
 
 	Settings->WatchProperty(Settings->CalculationMethod, [this](EMeshTangentsType) { Compute->InvalidateResult(); });
 	Settings->WatchProperty(Settings->LineLength, [this](float) { bLengthDirty = true; });
@@ -130,6 +135,7 @@ void UMeshTangentsTool::Setup()
 	Settings->WatchProperty(Settings->bShowTangents, [this](bool) { bVisibilityChanged = true; });
 	Settings->WatchProperty(Settings->bShowNormals, [this](bool) { bVisibilityChanged = true; });
 	Settings->WatchProperty(Settings->bCompareWithMikkt, [this](bool) { ComputeMikkTDeviations(ComputeDegenerateTris()); });
+	UVChannelProperties->WatchProperty(UVChannelProperties->UVChannel, [this](const FString& NewValue) { Compute->InvalidateResult(); });
 
 	if (InputGeometrySelection.IsEmpty() == false)
 	{
@@ -227,6 +233,7 @@ void UMeshTangentsTool::OnShutdown(EToolShutdownType ShutdownType)
 	PreviewMesh->Disconnect();
 
 	Settings->SaveProperties(this);
+	UVChannelProperties->SaveProperties(this);
 
 	if (GeometrySelectionViz)
 	{
@@ -328,6 +335,7 @@ TUniquePtr<TGenericDataOperator<FMeshTangentsd>> UMeshTangentsTool::MakeNewOpera
 	TangentsOp->SourceMesh = InputMesh;
 	TangentsOp->SourceTangents = InitialTangents;
 	TangentsOp->CalculationMethod = Settings->CalculationMethod;
+	TangentsOp->TargetUVLayer = UVChannelProperties->GetSelectedChannelIndex(true /*use 0 on failure*/);
 
 	return TangentsOp;
 }
diff --git a/Engine/Plugins/Experimental/MeshModelingToolsetExp/Source/MeshModelingToolsEditorOnlyExp/Public/MeshTangentsTool.h b/Engine/Plugins/Experimental/MeshModelingToolsetExp/Source/MeshModelingToolsEditorOnlyExp/Public/MeshTangentsTool.h
--- a/Engine/Plugins/Experimental/MeshModelingToolsetExp/Source/MeshModelingToolsEditorOnlyExp/Public/MeshTangentsTool.h
+++ b/Engine/Plugins/Experimental/MeshModelingToolsetExp/Source/MeshModelingToolsEditorOnlyExp/Public/MeshTangentsTool.h
@@ -11,6 +11,7 @@
 #include "ParameterizationOps/CalculateTangentsOp.h"
 #include "BaseTools/SingleSelectionMeshEditingTool.h"
 #include "Selections/GeometrySelection.h"
+#include "Properties/MeshUVChannelProperties.h"
 #include "MeshTangentsTool.generated.h"
 
 #define UE_API MESHMODELINGTOOLSEDITORONLYEXP_API
@@ -128,6 +129,9 @@ protected:
 	UPROPERTY()
 	TObjectPtr<UMeshTangentsToolProperties> Settings = nullptr;
 
+	UPROPERTY()
+	TObjectPtr<UMeshUVChannelProperties> UVChannelProperties = nullptr;
+
 	UPROPERTY()
 	TObjectPtr<UMaterialInterface> DefaultMaterial = nullptr;
 
diff --git a/Engine/Plugins/Runtime/MeshModelingToolset/Source/ModelingOperatorsEditorOnly/Private/ParameterizationOps/CalculateTangentsOp.cpp b/Engine/Plugins/Runtime/MeshModelingToolset/Source/ModelingOperatorsEditorOnly/Private/ParameterizationOps/CalculateTangentsOp.cpp
--- a/Engine/Plugins/Runtime/MeshModelingToolset/Source/ModelingOperatorsEditorOnly/Private/ParameterizationOps/CalculateTangentsOp.cpp
+++ b/Engine/Plugins/Runtime/MeshModelingToolset/Source/ModelingOperatorsEditorOnly/Private/ParameterizationOps/CalculateTangentsOp.cpp
@@ -32,6 +32,12 @@ void FCalculateTangentsOp::CalculateResult(FProgressCancel* Progress)
 		return;
 	}
 
+	if (TargetUVLayer < 0 || TargetUVLayer >= SourceMesh->Attributes()->NumUVLayers())
+	{
+		int32 ClampedLayer = FMath::Clamp(TargetUVLayer, 0, SourceMesh->Attributes()->NumUVLayers());
+		UE_LOG(LogGeometry, Warning, TEXT("Tangents compute requested invalid UV layer: %d. Using layer %d instead"), TargetUVLayer, ClampedLayer);
+		TargetUVLayer = ClampedLayer;
+	}
 
 	switch (CalculationMethod)
 	{
@@ -66,7 +72,7 @@ void FCalculateTangentsOp::CalculateResult(FProgressCancel* Progress)
 void FCalculateTangentsOp::CalculateStandard(FProgressCancel* Progress, TUniquePtr<FMeshTangentsd>& Tangents)
 {
 	const FDynamicMeshNormalOverlay* NormalOverlay = SourceMesh->Attributes()->PrimaryNormals();
-	const FDynamicMeshUVOverlay* UVOverlay = SourceMesh->Attributes()->GetUVLayer(0);
+	const FDynamicMeshUVOverlay* UVOverlay = SourceMesh->Attributes()->GetUVLayer(TargetUVLayer);
 
 	FComputeTangentsOptions Options;
 	Options.bAveraged = (CalculationMethod != EMeshTangentsType::PerTriangle);
@@ -180,7 +186,7 @@ void FCalculateTangentsOp::CalculateMikkT(FProgressCancel* Progress, TUniquePtr<
 	DynamicMeshMikkTInterface::FDynamicMeshInfo MeshInfo;
 	MeshInfo.Mesh = SourceMesh.Get();
 	MeshInfo.NormalOverlay = SourceMesh->Attributes()->PrimaryNormals();
-	MeshInfo.UVOverlay = SourceMesh->Attributes()->GetUVLayer(0);;
+	MeshInfo.UVOverlay = SourceMesh->Attributes()->GetUVLayer(TargetUVLayer);
 	MeshInfo.TangentsOut = Tangents.Get();
 
 	SMikkTSpaceInterface MikkTInterface;
diff --git a/Engine/Plugins/Runtime/MeshModelingToolset/Source/ModelingOperatorsEditorOnly/Public/ParameterizationOps/CalculateTangentsOp.h b/Engine/Plugins/Runtime/MeshModelingToolset/Source/ModelingOperatorsEditorOnly/Public/ParameterizationOps/CalculateTangentsOp.h
--- a/Engine/Plugins/Runtime/MeshModelingToolset/Source/ModelingOperatorsEditorOnly/Public/ParameterizationOps/CalculateTangentsOp.h
+++ b/Engine/Plugins/Runtime/MeshModelingToolset/Source/ModelingOperatorsEditorOnly/Public/ParameterizationOps/CalculateTangentsOp.h
@@ -40,6 +40,8 @@ public:
 
 	// parameters
 	EMeshTangentsType CalculationMethod;
+	// Compute tangents with respect to this UV layer
+	int32 TargetUVLayer = 0;
 
 	// error flags
 	bool bNoAttributesError = false;
