Make tangents tool also use new selected layer for mikkt-vs-fastmikkt comparison visualization
+ fix tangent op's clamp for invalid UV layer

diff --git a/Engine/Plugins/Experimental/MeshModelingToolsetExp/Source/MeshModelingToolsEditorOnlyExp/Private/MeshTangentsTool.cpp b/Engine/Plugins/Experimental/MeshModelingToolsetExp/Source/MeshModelingToolsEditorOnlyExp/Private/MeshTangentsTool.cpp
--- a/Engine/Plugins/Experimental/MeshModelingToolsetExp/Source/MeshModelingToolsEditorOnlyExp/Private/MeshTangentsTool.cpp
+++ b/Engine/Plugins/Experimental/MeshModelingToolsetExp/Source/MeshModelingToolsEditorOnlyExp/Private/MeshTangentsTool.cpp
@@ -481,7 +481,8 @@ TSet<int32> UMeshTangentsTool::ComputeDegenerateTris() const
 	if (!Settings->bShowDegenerates || (Settings->bCompareWithMikkt && Settings->CalculationMethod != EMeshTangentsType::MikkTSpace))
 	{
 		FMeshTangentsd DegenTangents(InputMesh.Get());
-		DegenTangents.ComputeTriangleTangents(InputMesh->Attributes()->GetUVLayer(0));
+		int32 UVLayer = UVChannelProperties->GetSelectedChannelIndex(true /*use 0 on failure*/);
+		DegenTangents.ComputeTriangleTangents(InputMesh->Attributes()->GetUVLayer(UVLayer));
 		DegenerateTris = TSet<int32>(DegenTangents.GetDegenerateTris());
 	}
 	return DegenerateTris;
@@ -494,16 +495,20 @@ void UMeshTangentsTool::ComputeMikkTDeviations(const TSet<int32>& DegenerateTris
 	Deviations.Reset();
 	if (Settings->bCompareWithMikkt && Settings->CalculationMethod == EMeshTangentsType::FastMikkTSpace)
 	{
+		int32 UVLayer = UVChannelProperties->GetSelectedChannelIndex(true /*use 0 on failure*/);
+		
 		FProgressCancel TmpCancel;
 		FCalculateTangentsOp MikktOp;
 		MikktOp.SourceMesh = InputMesh;
 		MikktOp.CalculationMethod = EMeshTangentsType::MikkTSpace;
+		MikktOp.TargetUVLayer = UVLayer;
 		MikktOp.CalculateResult(&TmpCancel);
 		TUniquePtr<FMeshTangentsd> MikktTangents = MikktOp.ExtractResult();
 
 		FCalculateTangentsOp NewOp;
 		NewOp.SourceMesh = InputMesh;
 		NewOp.CalculationMethod = EMeshTangentsType::FastMikkTSpace;
+		NewOp.TargetUVLayer = UVLayer;
 		NewOp.CalculateResult(&TmpCancel);
 		TUniquePtr<FMeshTangentsd> NewTangents = NewOp.ExtractResult();
 
diff --git a/Engine/Plugins/Runtime/MeshModelingToolset/Source/ModelingOperatorsEditorOnly/Private/ParameterizationOps/CalculateTangentsOp.cpp b/Engine/Plugins/Runtime/MeshModelingToolset/Source/ModelingOperatorsEditorOnly/Private/ParameterizationOps/CalculateTangentsOp.cpp
--- a/Engine/Plugins/Runtime/MeshModelingToolset/Source/ModelingOperatorsEditorOnly/Private/ParameterizationOps/CalculateTangentsOp.cpp
+++ b/Engine/Plugins/Runtime/MeshModelingToolset/Source/ModelingOperatorsEditorOnly/Private/ParameterizationOps/CalculateTangentsOp.cpp
@@ -34,7 +34,7 @@ void FCalculateTangentsOp::CalculateResult(FProgressCancel* Progress)
 
 	if (TargetUVLayer < 0 || TargetUVLayer >= SourceMesh->Attributes()->NumUVLayers())
 	{
-		int32 ClampedLayer = FMath::Clamp(TargetUVLayer, 0, SourceMesh->Attributes()->NumUVLayers());
+		int32 ClampedLayer = FMath::Clamp(TargetUVLayer, 0, SourceMesh->Attributes()->NumUVLayers() - 1);
 		UE_LOG(LogGeometry, Warning, TEXT("Tangents compute requested invalid UV layer: %d. Using layer %d instead"), TargetUVLayer, ClampedLayer);
 		TargetUVLayer = ClampedLayer;
 	}
