ModelingTools: Fix PivotActorTool crashing when run on an actor whose root component is not a primitive component.

diff --git a/Engine/Plugins/Experimental/MeshModelingToolsetExp/Source/MeshModelingToolsEditorOnlyExp/Private/AddPivotActorTool.cpp b/Engine/Plugins/Experimental/MeshModelingToolsetExp/Source/MeshModelingToolsEditorOnlyExp/Private/AddPivotActorTool.cpp
--- a/Engine/Plugins/Experimental/MeshModelingToolsetExp/Source/MeshModelingToolsEditorOnlyExp/Private/AddPivotActorTool.cpp
+++ b/Engine/Plugins/Experimental/MeshModelingToolsetExp/Source/MeshModelingToolsEditorOnlyExp/Private/AddPivotActorTool.cpp
@@ -93,12 +93,15 @@ void UAddPivotActorToolBuilder::InitializeNewTool(UMultiSelectionMeshEditingTool
 	}
 	else
 	{
-		// When dealing with actor selection, only consider root components.
-		SceneState.TargetManager->EnumerateSelectedAndTargetableComponents(SceneState, Requirements, [this, &SceneState, &Targets, &Requirements](UActorComponent* Component)
+		// When dealing with actor selection, only grab one target per owning actor.
+		TSet<AActor*> OwningActors;
+		SceneState.TargetManager->EnumerateSelectedAndTargetableComponents(SceneState, Requirements, [this, &SceneState, &Targets, &Requirements, &OwningActors](UActorComponent* Component)
 		{
 			if (AActor* OwnerActor = Component->GetOwner())
 			{
-				if (OwnerActor->GetRootComponent() == Component)
+				bool bAlreadyInSet = false;
+				OwningActors.Add(OwnerActor, &bAlreadyInSet);
+				if (!bAlreadyInSet)
 				{
 					Targets.Add(SceneState.TargetManager->BuildTarget(Component, Requirements));
 				}
@@ -265,7 +268,8 @@ void UAddPivotActorTool::OnShutdown(EToolShutdownType ShutdownType)
 
 			GetToolManager()->EndUndoTransaction();
 		}
-		else // If in "add" mode
+		// If in "add" mode
+		else if (ensure(Targets.Num() > 0)) 
 		{
 			GetToolManager()->BeginUndoTransaction(LOCTEXT("AddPivotActorTransactionName", "Add Pivot Actor"));
 
