ModelingMode: Fixed issue where the pattern tool could crash when run on static mesh components that have no static mesh asset assigned.

diff --git a/Engine/Plugins/Experimental/MeshModelingToolsetExp/Source/MeshModelingToolsExp/Private/PatternTool.cpp b/Engine/Plugins/Experimental/MeshModelingToolsetExp/Source/MeshModelingToolsExp/Private/PatternTool.cpp
--- a/Engine/Plugins/Experimental/MeshModelingToolsetExp/Source/MeshModelingToolsExp/Private/PatternTool.cpp
+++ b/Engine/Plugins/Experimental/MeshModelingToolsetExp/Source/MeshModelingToolsExp/Private/PatternTool.cpp
@@ -901,14 +901,29 @@ void UPatternTool::InitializeElements()
 		if (UStaticMeshComponent* StaticMeshComp = Cast<UStaticMeshComponent>(Element.SourceComponent))
 		{
 			Element.SourceStaticMesh = StaticMeshComp->GetStaticMesh();
-			Element.LocalBounds = Element.SourceStaticMesh->GetBounds().GetBox();
+			if (Element.SourceStaticMesh != nullptr)
+			{
+				Element.LocalBounds = Element.SourceStaticMesh->GetBounds().GetBox();
+			}
+			else
+			{
+				Element.bValid = false;
+			}
 		}	
 		else if (UDynamicMeshComponent* DynamicMeshComp = Cast<UDynamicMeshComponent>(Element.SourceComponent))
 		{
 			Element.SourceDynamicMesh = DynamicMeshComp->GetDynamicMesh();
-			Element.SourceDynamicMesh->ProcessMesh([&](const FDynamicMesh3& Mesh) {
-				Element.LocalBounds = Mesh.GetBounds(true);
-			});
+			if (Element.SourceDynamicMesh != nullptr)
+			{
+				Element.SourceDynamicMesh->ProcessMesh([&](const FDynamicMesh3& Mesh) {
+					Element.LocalBounds = Mesh.GetBounds(true);
+					});
+			}
+			else
+			{
+				Element.bValid = false;
+			}
+
 		}
 		else
 		{
