Merge Actor - Approximate: Use the correct merge material to avoid showing default engine textures when some material properties aren't baked
* FullMaterialBakePreviewMaterial_PackedMRS falls back to engine default textures when normal & packed MRS aren't baked
* Now using FullMaterialBakeOutputMaterial_PackedMRS, which requires static switches to be setup
* Added missing FullMaterialBakePreviewMaterial material which uses independent textures for Metallic+Roughness+Specular

#virtualized

diff --git a/Engine/Plugins/Experimental/MeshModelingToolsetExp/Source/GeometryProcessingAdapters/Private/ApproximateActorsImpl.cpp b/Engine/Plugins/Experimental/MeshModelingToolsetExp/Source/GeometryProcessingAdapters/Private/ApproximateActorsImpl.cpp
--- a/Engine/Plugins/Experimental/MeshModelingToolsetExp/Source/GeometryProcessingAdapters/Private/ApproximateActorsImpl.cpp
+++ b/Engine/Plugins/Experimental/MeshModelingToolsetExp/Source/GeometryProcessingAdapters/Private/ApproximateActorsImpl.cpp
@@ -1279,7 +1279,10 @@ void FApproximateActorsImpl::GenerateApproximationForActorSet(const FInput& Inpu
 
 	// Make material for textures by creating MIC of input material, or fall back to known material
 	UMaterialInterface* UseBaseMaterial = (Options.BakeMaterial != nullptr) ? 
-		Options.BakeMaterial : LoadObject<UMaterial>(nullptr, TEXT("/MeshModelingToolsetExp/Materials/FullMaterialBakePreviewMaterial_PackedMRS"));
+		Options.BakeMaterial : 
+		Options.bUsePackedMRS ?
+			LoadObject<UMaterial>(nullptr, TEXT("/MeshModelingToolsetExp/Materials/FullMaterialBakeOutputMaterial_PackedMRS")) :
+			LoadObject<UMaterial>(nullptr, TEXT("/MeshModelingToolsetExp/Materials/FullMaterialBakeOutputMaterial"));
 	FMaterialAssetOptions MatOptions;
 	MatOptions.NewAssetPath = Options.BasePackagePath + TEXT("_Material");
 	FMaterialAssetResults MatResults;
@@ -1295,42 +1298,54 @@ void FApproximateActorsImpl::GenerateApproximationForActorSet(const FInput& Inpu
 	FString BaseTexturePath = MatOptions.NewAssetPath;
 	auto WriteTextureLambda = [BaseTexturePath, NewMaterial, &ResultsOut](
 		UTexture2D* Texture, 
-		FString TextureTypeSuffix, 
 		FTexture2DBuilder::ETextureType Type,
-		FName MaterialParamName )
+		FName MaterialParamName,
+		FName SwitchName)
 	{
-		if (ensure(Texture != nullptr) == false) return;
+		if (Texture)
+		{
+			FTexture2DBuilder::CopyPlatformDataToSourceData(Texture, Type);
 
-		FTexture2DBuilder::CopyPlatformDataToSourceData(Texture, Type);
+			if (Type == FTexture2DBuilder::ETextureType::Roughness
+				|| Type == FTexture2DBuilder::ETextureType::Metallic
+				|| Type == FTexture2DBuilder::ETextureType::Specular)
+			{
+				UE::AssetUtils::ConvertToSingleChannel(Texture);
+			}
 
-		if (Type == FTexture2DBuilder::ETextureType::Roughness
-			|| Type == FTexture2DBuilder::ETextureType::Metallic
-			|| Type == FTexture2DBuilder::ETextureType::Specular)
-		{
-			UE::AssetUtils::ConvertToSingleChannel(Texture);
-		}
+			// Make sure the texture is a VT if required by the material sampler
+			if (NewMaterial != nullptr)
+			{
+				UTexture* DefaultTexture = nullptr;
+				NewMaterial->GetTextureParameterValue(MaterialParamName, DefaultTexture);
+				if (ensure(DefaultTexture))
+				{
+					Texture->VirtualTextureStreaming = DefaultTexture->VirtualTextureStreaming;
+				}
+			}
 
-		// Make sure the texture is a VT if required by the material sampler
-		if (NewMaterial != nullptr)
-		{
-			UTexture* DefaultTexture = nullptr;
-			NewMaterial->GetTextureParameterValue(MaterialParamName, DefaultTexture);
-			if (ensure(DefaultTexture))
+			FTexture2DAssetOptions TexOptions;
+			TexOptions.NewAssetPath = BaseTexturePath + TEXT("_") + MaterialParamName.ToString();
+			FTexture2DAssetResults Results;
+			ECreateTexture2DResult TexResult = UE::AssetUtils::SaveGeneratedTexture2DAsset(Texture, TexOptions, Results);
+			if (ensure(TexResult == ECreateTexture2DResult::Ok))
 			{
-				Texture->VirtualTextureStreaming = DefaultTexture->VirtualTextureStreaming;
+				ResultsOut.NewTextures.Add(Texture);
+				if (NewMaterial != nullptr)
+				{
+					NewMaterial->SetTextureParameterValueEditorOnly(MaterialParamName, Texture);
+				}
 			}
 		}
 
-		FTexture2DAssetOptions TexOptions;
-		TexOptions.NewAssetPath = BaseTexturePath + TextureTypeSuffix;
-		FTexture2DAssetResults Results;
-		ECreateTexture2DResult TexResult = UE::AssetUtils::SaveGeneratedTexture2DAsset(Texture, TexOptions, Results);
-		if (ensure(TexResult == ECreateTexture2DResult::Ok))
+		if (NewMaterial && SwitchName != NAME_None)
 		{
-			ResultsOut.NewTextures.Add(Texture);
-			if (NewMaterial != nullptr)
+			FGuid SwitchGuid;
+			bool bUnusedValue = false;
+			if (NewMaterial->GetStaticSwitchParameterValue(SwitchName, bUnusedValue, SwitchGuid))
 			{
-				NewMaterial->SetTextureParameterValueEditorOnly(MaterialParamName, Texture);
+				const bool bUseSwitchValue = true;
+				NewMaterial->SetStaticSwitchParameterValueEditorOnly(SwitchName, bUseSwitchValue);
 			}
 		}
 	};
@@ -1338,34 +1353,34 @@ void FApproximateActorsImpl::GenerateApproximationForActorSet(const FInput& Inpu
 	TRACE_BOOKMARK(TEXT("ApproximateActors-Write Textures"));
 
 	// process the generated textures
-	if (Options.bBakeBaseColor && GeneratedTextures.BaseColorMap)
+	if (Options.bBakeBaseColor)
 	{
-		WriteTextureLambda(GeneratedTextures.BaseColorMap, TEXT("_BaseColor"), FTexture2DBuilder::ETextureType::Color, Options.BaseColorTexParamName);
+		WriteTextureLambda(GeneratedTextures.BaseColorMap, FTexture2DBuilder::ETextureType::Color, Options.BaseColorTexParamName, Options.BaseColorSwitchName);
 	}
-	if (Options.bBakeEmissive && GeneratedTextures.EmissiveMap)
+	if (Options.bBakeEmissive)
 	{
-		WriteTextureLambda(GeneratedTextures.EmissiveMap, TEXT("_Emissive"), FTexture2DBuilder::ETextureType::EmissiveHDR, Options.EmissiveTexParamName);
+		WriteTextureLambda(GeneratedTextures.EmissiveMap, FTexture2DBuilder::ETextureType::EmissiveHDR, Options.EmissiveTexParamName, Options.EmissiveSwitchName);
 	}
-	if (Options.bBakeNormalMap && GeneratedTextures.NormalMap)
+	if (Options.bBakeNormalMap)
 	{
-		WriteTextureLambda(GeneratedTextures.NormalMap, TEXT("_Normal"), FTexture2DBuilder::ETextureType::NormalMap, Options.NormalTexParamName);
+		WriteTextureLambda(GeneratedTextures.NormalMap, FTexture2DBuilder::ETextureType::NormalMap, Options.NormalTexParamName, Options.NormalSwitchName);
 	}
 
-	if ((Options.bBakeRoughness || Options.bBakeMetallic || Options.bBakeSpecular) && Options.bUsePackedMRS && GeneratedTextures.PackedMRSMap)
+	if ((Options.bBakeRoughness || Options.bBakeMetallic || Options.bBakeSpecular) && Options.bUsePackedMRS)
 	{
-		WriteTextureLambda(GeneratedTextures.PackedMRSMap, TEXT("_PackedMRS"), FTexture2DBuilder::ETextureType::ColorLinear, Options.PackedMRSTexParamName);
+		WriteTextureLambda(GeneratedTextures.PackedMRSMap, FTexture2DBuilder::ETextureType::ColorLinear, Options.PackedMRSTexParamName, Options.PackedMRSSwitchName);
 	}
-	if (Options.bBakeRoughness && GeneratedTextures.RoughnessMap)
+	if (Options.bBakeRoughness)
 	{
-		WriteTextureLambda(GeneratedTextures.RoughnessMap, TEXT("_Roughness"), FTexture2DBuilder::ETextureType::Roughness, Options.RoughnessTexParamName);
+		WriteTextureLambda(GeneratedTextures.RoughnessMap, FTexture2DBuilder::ETextureType::Roughness, Options.RoughnessTexParamName, Options.RoughnessSwitchName);
 	}
-	if (Options.bBakeMetallic && GeneratedTextures.MetallicMap)
+	if (Options.bBakeMetallic)
 	{
-		WriteTextureLambda(GeneratedTextures.MetallicMap, TEXT("_Metallic"), FTexture2DBuilder::ETextureType::Metallic, Options.MetallicTexParamName);
+		WriteTextureLambda(GeneratedTextures.MetallicMap, FTexture2DBuilder::ETextureType::Metallic, Options.MetallicTexParamName, Options.MetallicSwitchName);
 	}
-	if (Options.bBakeSpecular && GeneratedTextures.SpecularMap)
+	if (Options.bBakeSpecular)
 	{
-		WriteTextureLambda(GeneratedTextures.SpecularMap, TEXT("_Specular"), FTexture2DBuilder::ETextureType::Specular, Options.SpecularTexParamName);
+		WriteTextureLambda(GeneratedTextures.SpecularMap, FTexture2DBuilder::ETextureType::Specular, Options.SpecularTexParamName, Options.SpecularSwitchName);
 	}
 
 
diff --git a/Engine/Source/Developer/GeometryProcessingInterfaces/Public/GeometryProcessingInterfaces/ApproximateActors.h b/Engine/Source/Developer/GeometryProcessingInterfaces/Public/GeometryProcessingInterfaces/ApproximateActors.h
--- a/Engine/Source/Developer/GeometryProcessingInterfaces/Public/GeometryProcessingInterfaces/ApproximateActors.h
+++ b/Engine/Source/Developer/GeometryProcessingInterfaces/Public/GeometryProcessingInterfaces/ApproximateActors.h
@@ -193,28 +193,35 @@ public:
 		//
 
 		// A new MIC derived from this material will be created and assigned to the generated mesh
-		UMaterialInterface* BakeMaterial = nullptr;		// if null, will use /MeshModelingToolsetExp/Materials/FullMaterialBakePreviewMaterial_PackedMRS instead
+		UMaterialInterface* BakeMaterial = nullptr;		// if null, will use /MeshModelingToolsetExp/Materials/FullMaterialBakeOutputMaterial instead (or the _PackedMRS version if that parameter is enabled)
 		FName BaseColorTexParamName = FName("BaseColor");
+		FName BaseColorSwitchName = FName("EnableBaseColor");
 		bool bBakeBaseColor = true;
 		int32 CustomTextureSizeBaseColor = 1024;
 		FName RoughnessTexParamName = FName("Roughness");
+		FName RoughnessSwitchName = FName("EnableRoughness");
 		bool bBakeRoughness = true;
 		int32 CustomTextureSizeRoughness = 1024;
 		FName MetallicTexParamName = FName("Metallic");
+		FName MetallicSwitchName = FName("EnableMetallic");
 		bool bBakeMetallic = true;
 		int32 CustomTextureSizeMetallic = 1024;
 		FName SpecularTexParamName = FName("Specular");
+		FName SpecularSwitchName = FName("EnableSpecular");
 		bool bBakeSpecular = true;
 		int32 CustomTextureSizeSpecular = 1024;
 		FName EmissiveTexParamName = FName("Emissive");
+		FName EmissiveSwitchName = FName("EnableEmissive");
 		bool bBakeEmissive = true;
 		int32 CustomTextureSizeEmissive = 1024;
 		FName NormalTexParamName = FName("NormalMap");
+		FName NormalSwitchName = FName("EnableNormalMap");
 		bool bBakeNormalMap = true;
 		int32 CustomTextureSizeNormalMap = 1024;
 		
 		bool bUsePackedMRS = true;
 		FName PackedMRSTexParamName = FName("PackedMRS");
+		FName PackedMRSSwitchName = FName("EnablePackedMRS");
 
 		// output texture options
 		ETextureSizePolicy TextureSizePolicy = ETextureSizePolicy::TextureSize;
