ModelingTools: Fix CubeGrid "Accept and Start New" action not working correctly when editing an existing asset.

diff --git a/Engine/Plugins/Experimental/MeshModelingToolsetExp/Source/MeshModelingToolsExp/Private/CubeGridTool.cpp b/Engine/Plugins/Experimental/MeshModelingToolsetExp/Source/MeshModelingToolsExp/Private/CubeGridTool.cpp
--- a/Engine/Plugins/Experimental/MeshModelingToolsetExp/Source/MeshModelingToolsExp/Private/CubeGridTool.cpp
+++ b/Engine/Plugins/Experimental/MeshModelingToolsetExp/Source/MeshModelingToolsExp/Private/CubeGridTool.cpp
@@ -1406,7 +1406,7 @@ void UCubeGridTool::AcceptToolAndStartNew()
 		ApplyCornerMode(true);
 	}
 
-	if (!bChangesMade)
+	if (!bChangesMade && !Target)
 	{
 		return;
 	}
@@ -1414,15 +1414,20 @@ void UCubeGridTool::AcceptToolAndStartNew()
 	FText TransactionText = LOCTEXT("AcceptAndStartNewTransaction", "Cube Grid Accept and Start New");
 	GetToolManager()->BeginUndoTransaction(TransactionText);
 
-	OutputCurrentResults(/*bSetSelection =*/ false);
-
-	GetToolManager()->EmitObjectChange(this, MakeUnique<FCubeGridChangesMadeChange>(false), TransactionText);
+	if (bChangesMade)
+	{
+		OutputCurrentResults(/*bSetSelection =*/ false);
+		GetToolManager()->EmitObjectChange(this, MakeUnique<FCubeGridChangesMadeChange>(false), TransactionText);
+		bChangesMade = false;
+	}
 
 	if (Target)
 	{
 		UE::ToolTarget::SetSourceObjectVisible(Target.Get(), true);
 		GetToolManager()->EmitObjectChange(this, MakeUnique<FCubeGridTargetResetChange>(Target), TransactionText);
+		SetTarget(nullptr);
 	}
+	GetToolManager()->DisplayMessage(CreatingAssetMessage, EToolMessageLevel::UserWarning);
 
 	// Clear out the current mesh
 	if (CurrentMesh->TriangleCount() > 0)
