PR #13360: Add assign and start new keyboard command to cube grid

diff --git a/Engine/Plugins/Experimental/MeshModelingToolsetExp/Source/MeshModelingToolsExp/Private/CubeGridTool.cpp b/Engine/Plugins/Experimental/MeshModelingToolsetExp/Source/MeshModelingToolsExp/Private/CubeGridTool.cpp
--- a/Engine/Plugins/Experimental/MeshModelingToolsetExp/Source/MeshModelingToolsExp/Private/CubeGridTool.cpp
+++ b/Engine/Plugins/Experimental/MeshModelingToolsetExp/Source/MeshModelingToolsExp/Private/CubeGridTool.cpp
@@ -2619,6 +2619,13 @@ void UCubeGridTool::RegisterActions(FInteractiveToolActionSet& ActionSet)
 #endif
 		[this]() { RequestAction(ECubeGridToolAction::IncreaseGridPower); });
 
+	ActionSet.RegisterAction(this, ActionID++,
+		TEXT("AcceptAndStartNew"),
+		LOCTEXT("AcceptAndStartNewAction", "Accept and Start New"),
+		FText(),
+		EModifierKey::None, EKeys::N,
+		[this]() {RequestAction(ECubeGridToolAction::AcceptAndStartNew); });
+
 	ActionSet.RegisterAction(this, ActionID++,
 		TEXT("ToggleGizmoVisibility"),
 		LOCTEXT("ToggleGizmoVisibilityAction", "Toggle Gizmo Visibility"),
