fix crash on converting mulitple dynamic mesh sources to static meshes w/ the modeling mode convert tool, due to tarray doing an unsafe move of dynamic meshes

[RN] Fixed crash on using the modeling mode Convert tool to convert multiple dynamic meshes to static meshes

diff --git a/Engine/Plugins/Experimental/MeshModelingToolsetExp/Source/MeshModelingToolsExp/Private/ConvertMeshesTool.cpp b/Engine/Plugins/Experimental/MeshModelingToolsetExp/Source/MeshModelingToolsExp/Private/ConvertMeshesTool.cpp
--- a/Engine/Plugins/Experimental/MeshModelingToolsetExp/Source/MeshModelingToolsExp/Private/ConvertMeshesTool.cpp
+++ b/Engine/Plugins/Experimental/MeshModelingToolsetExp/Source/MeshModelingToolsExp/Private/ConvertMeshesTool.cpp
@@ -151,7 +151,7 @@ void UConvertMeshesTool::Shutdown(EToolShutdownType ShutdownType)
 
 		TArray<AActor*> NewSelectedActors;
 		TSet<AActor*> DeleteActors;
-		TArray<FCreateMeshObjectParams> NewMeshObjects;
+		TIndirectArray<FCreateMeshObjectParams> NewMeshObjects;
 
 		// Accumulate info for new mesh objects. Do not immediately create them because then
 		// the new Actors will get a unique-name incremented suffix, because the convert-from
@@ -211,7 +211,8 @@ void UConvertMeshesTool::Shutdown(EToolShutdownType ShutdownType)
 
 			FString AssetName = TargetActor->GetActorNameOrLabel();
 
-			FCreateMeshObjectParams NewMeshObjectParams;
+			FCreateMeshObjectParams* NewMeshObjectParamsPtr = new FCreateMeshObjectParams();
+			FCreateMeshObjectParams& NewMeshObjectParams = *NewMeshObjectParamsPtr;
 			NewMeshObjectParams.TargetWorld = InputComponent->GetWorld();
 			NewMeshObjectParams.Transform = (FTransform)SourceTransform;
 			NewMeshObjectParams.BaseName = AssetName;
@@ -238,7 +239,7 @@ void UConvertMeshesTool::Shutdown(EToolShutdownType ShutdownType)
 				}
 			}
 
-			NewMeshObjects.Add(MoveTemp(NewMeshObjectParams));
+			NewMeshObjects.Add(NewMeshObjectParamsPtr);
 		}
 
 		// delete all the existing Actors we want to get rid of
