fix crash in modeling tool union/csg/cut when processing meshes w/ invalid material IDs on some triangles

diff --git a/Engine/Plugins/Experimental/MeshModelingToolsetExp/Source/MeshModelingToolsExp/Private/SelfUnionMeshesTool.cpp b/Engine/Plugins/Experimental/MeshModelingToolsetExp/Source/MeshModelingToolsExp/Private/SelfUnionMeshesTool.cpp
--- a/Engine/Plugins/Experimental/MeshModelingToolsetExp/Source/MeshModelingToolsExp/Private/SelfUnionMeshesTool.cpp
+++ b/Engine/Plugins/Experimental/MeshModelingToolsetExp/Source/MeshModelingToolsExp/Private/SelfUnionMeshesTool.cpp
@@ -116,9 +116,15 @@ void USelfUnionMeshesTool::ConvertInputsAndSetPreviewMaterials(bool bSetPreviewM
 		ComponentMesh.EnableAttributes();
 		ComponentMesh.Attributes()->EnableMaterialID();
 		FDynamicMeshMaterialAttribute* MaterialIDs = ComponentMesh.Attributes()->GetMaterialID();
-		for (int TID : ComponentMesh.TriangleIndicesItr())
+		int32 NumComponentMaterials = MaterialRemap[ComponentIdx].Num();
+		// If component had valid materials, remap them. (If it had no materials, no IDs are valid, so just leave them as they were.)
+		if (NumComponentMaterials > 0)
 		{
-			MaterialIDs->SetValue(TID, MaterialRemap[ComponentIdx][MaterialIDs->GetValue(TID)]);
+			for (int TID : ComponentMesh.TriangleIndicesItr())
+			{
+				int32 MID = MaterialIDs->GetValue(TID);
+				MaterialIDs->SetValue(TID, MaterialRemap[ComponentIdx][MID >= 0 && MID < NumComponentMaterials ? MID : 0]);
+			}
 		}
 		FTransformSRT3d WorldTransform = TransformProxies[ComponentIdx]->GetTransform();
 		if (WorldTransform.GetDeterminant() < 0)
diff --git a/Engine/Plugins/Runtime/MeshModelingToolset/Source/MeshModelingTools/Private/CSGMeshesTool.cpp b/Engine/Plugins/Runtime/MeshModelingToolset/Source/MeshModelingTools/Private/CSGMeshesTool.cpp
--- a/Engine/Plugins/Runtime/MeshModelingToolset/Source/MeshModelingTools/Private/CSGMeshesTool.cpp
+++ b/Engine/Plugins/Runtime/MeshModelingToolset/Source/MeshModelingTools/Private/CSGMeshesTool.cpp
@@ -225,9 +225,15 @@ void UCSGMeshesTool::ConvertInputsAndSetPreviewMaterials(bool bSetPreviewMesh)
 		OriginalDynamicMeshes[ComponentIdx]->EnableAttributes();
 		OriginalDynamicMeshes[ComponentIdx]->Attributes()->EnableMaterialID();
 		FDynamicMeshMaterialAttribute* MaterialIDs = OriginalDynamicMeshes[ComponentIdx]->Attributes()->GetMaterialID();
-		for (int TID : OriginalDynamicMeshes[ComponentIdx]->TriangleIndicesItr())
+		int32 NumComponentMaterials = MaterialRemap[ComponentIdx].Num();
+		// If component had valid materials, remap them. (If it had no materials, no IDs are valid, so just leave them as they were.)
+		if (NumComponentMaterials > 0)
 		{
-			MaterialIDs->SetValue(TID, MaterialRemap[ComponentIdx][MaterialIDs->GetValue(TID)]);
+			for (int TID : OriginalDynamicMeshes[ComponentIdx]->TriangleIndicesItr())
+			{
+				int32 MID = MaterialIDs->GetValue(TID);
+				MaterialIDs->SetValue(TID, MaterialRemap[ComponentIdx][MID >= 0 && MID < NumComponentMaterials ? MID : 0]);
+			}
 		}
 
 		if (bSetPreviewMesh)
diff --git a/Engine/Plugins/Runtime/MeshModelingToolset/Source/MeshModelingTools/Private/CutMeshWithMeshTool.cpp b/Engine/Plugins/Runtime/MeshModelingToolset/Source/MeshModelingTools/Private/CutMeshWithMeshTool.cpp
--- a/Engine/Plugins/Runtime/MeshModelingToolset/Source/MeshModelingTools/Private/CutMeshWithMeshTool.cpp
+++ b/Engine/Plugins/Runtime/MeshModelingToolset/Source/MeshModelingTools/Private/CutMeshWithMeshTool.cpp
@@ -123,9 +123,15 @@ void UCutMeshWithMeshTool::ConvertInputsAndSetPreviewMaterials(bool bSetPreviewM
 		Mesh->EnableAttributes();
 		Mesh->Attributes()->EnableMaterialID();
 		FDynamicMeshMaterialAttribute* MaterialIDs = Mesh->Attributes()->GetMaterialID();
-		for (int TID : Mesh->TriangleIndicesItr())
+		int32 NumComponentMaterials = MaterialRemap[ComponentIdx].Num();
+		// If component had valid materials, remap them. (If it had no materials, no IDs are valid, so just leave them as they were.)
+		if (NumComponentMaterials > 0)
 		{
-			MaterialIDs->SetValue(TID, MaterialRemap[ComponentIdx][MaterialIDs->GetValue(TID)]);
+			for (int TID : Mesh->TriangleIndicesItr())
+			{
+				int32 MID = MaterialIDs->GetValue(TID);
+				MaterialIDs->SetValue(TID, MaterialRemap[ComponentIdx][MID >= 0 && MID < NumComponentMaterials ? MID : 0]);
+			}
 		}
 
 		if (ComponentIdx == 0)
