DrawSplineTool: add option to fit spline onto an existing surface. Update the ProjectSplineToSurface function to use absolute distance theshold rather then scaling by spline segment length

diff --git a/Engine/Plugins/Experimental/MeshModelingToolsetExp/Source/MeshModelingToolsEditorOnlyExp/Private/DrawSplineTool.cpp b/Engine/Plugins/Experimental/MeshModelingToolsetExp/Source/MeshModelingToolsEditorOnlyExp/Private/DrawSplineTool.cpp
--- a/Engine/Plugins/Experimental/MeshModelingToolsetExp/Source/MeshModelingToolsEditorOnlyExp/Private/DrawSplineTool.cpp
+++ b/Engine/Plugins/Experimental/MeshModelingToolsetExp/Source/MeshModelingToolsEditorOnlyExp/Private/DrawSplineTool.cpp
@@ -27,6 +27,12 @@
 #include "SplineUtil.h"
 #include "ToolBuilderUtil.h"
 #include "UnrealEdGlobals.h" // GUnrealEd
+#include "TargetInterfaces/DynamicMeshProvider.h"
+#include "TargetInterfaces/DynamicMeshSource.h"
+#include "TargetInterfaces/PrimitiveComponentBackedTarget.h"
+#include "UDynamicMesh.h"
+#include "ToolTargetManager.h"
+#include "ToolTargets/DynamicMeshComponentToolTarget.h"
 
 #include UE_INLINE_GENERATED_CPP_BY_NAME(DrawSplineTool)
 
@@ -396,7 +402,52 @@ void UDrawSplineTool::Setup()
 		TransitionOutputMode();
 	});
 
+	Settings->WatchProperty(Settings->bFitToSurface, [this](bool) 
+	{
+		RefitToSurface();
+	});
+
+	Settings->WatchProperty(Settings->MaxNewPoints, [this](int32)
+	{
+		RefitToSurface();
+	});
+
+	Settings->WatchProperty(Settings->DistanceThreshold, [this](double) 
+	{
+		RefitToSurface();
+	});
+
 	Settings->SilentUpdateWatched();
+
+	if (Target)
+	{
+		if (IPersistentDynamicMeshSource* const DynamicMeshSource = Cast<IPersistentDynamicMeshSource>(Target))
+		{
+			UDynamicMesh* const MeshContainer = DynamicMeshSource->GetDynamicMeshContainer();
+			ProjectionTarget = MakeUnique<UE::Geometry::FDynamicMeshAABBTree3>(MeshContainer->GetMeshPtr(), true);
+		}
+		else if (IDynamicMeshProvider* const DynamicMeshProvider = Cast<IDynamicMeshProvider>(Target))
+		{
+			ProjectionTargetMesh = MakeUnique<UE::Geometry::FDynamicMesh3>();
+			ProjectionTargetMesh->Copy(DynamicMeshProvider->GetDynamicMesh());
+			ProjectionTarget = MakeUnique<UE::Geometry::FDynamicMeshAABBTree3>(ProjectionTargetMesh.Get(), true);
+		}
+
+		if (ISceneComponentBackedTarget* const PrimitiveToolTarget = Cast<ISceneComponentBackedTarget>(Target))
+		{
+			MeshTargetTransform = PrimitiveToolTarget->GetWorldTransform();
+		}
+		else
+		{
+			MeshTargetTransform = FTransform();
+		}		
+
+		Settings->bMeshSelected = true;
+	}
+	else
+	{
+		Settings->bMeshSelected = false;
+	}
 }
 
 // Set things up for a new output mode or destination
@@ -593,9 +644,17 @@ void UDrawSplineTool::Shutdown(EToolShutdownType ShutdownType)
 		HiddenSpline = nullptr;
 	}
 
-	if (ShutdownType == EToolShutdownType::Accept && WorkingSpline.IsValid() && WorkingSpline->GetNumberOfSplinePoints() > 0)
+	if (ShutdownType == EToolShutdownType::Accept)
 	{
-		GenerateAsset();
+		if (Settings->bFitToSurface && ProjectionSpline.IsValid())
+		{
+			WorkingSpline->SetSpline(ProjectionSpline->GetSpline());
+		}
+
+		if (WorkingSpline.IsValid() && WorkingSpline->GetNumberOfSplinePoints() > 0)
+		{
+			GenerateAsset();
+		}
 	}
 
 	PlaneMechanic->Shutdown();
@@ -606,6 +665,11 @@ void UDrawSplineTool::Shutdown(EToolShutdownType ShutdownType)
 		WorkingSpline->DestroyComponent();
 	}
 
+	if (ProjectionSpline.IsValid())
+	{
+		ProjectionSpline->DestroyComponent();
+	}
+
 	if (PreviewActor)
 	{
 		PreviewActor->Destroy();
@@ -708,6 +772,35 @@ void UDrawSplineTool::GenerateAsset()
 	GetToolManager()->EndUndoTransaction();
 }
 
+void UDrawSplineTool::RefitToSurface()
+{
+	if (Settings->bFitToSurface && ProjectionTarget)
+	{
+		const FInterpCurveVector& Curve = WorkingSpline->GetSplinePointsPosition();
+		FInterpCurveVector NewCurve;
+		SplineUtil::ProjectSplineToSurface(NewCurve, Curve, *ProjectionTarget, WorkingSpline->GetComponentTransform(), MeshTargetTransform, Settings->DistanceThreshold, Settings->MaxNewPoints);
+
+		if (!ProjectionSpline.IsValid())
+		{
+			ProjectionSpline = DrawSplineToolLocals::CreateNewSplineInActor(PreviewActor);
+		}
+
+		TArray<FVector> ProjectedPoints;
+		for (FInterpCurvePointVector NewPoint : NewCurve.Points)
+		{
+			ProjectedPoints.Add(NewPoint.OutVal);
+		}
+
+		ProjectionSpline->SetSplinePoints(ProjectedPoints, ESplineCoordinateSpace::World, /*bUpdateSpline = */ true);
+	}
+	else if (ProjectionSpline.IsValid())
+	{
+		ProjectionSpline->DestroyComponent();
+		ProjectionSpline.Reset();
+	}
+}
+
+
 // Helper to add a point given a hit location and hit normal
 void UDrawSplineTool::AddSplinePoint(const FVector3d& HitLocation, const FVector3d& HitNormal)
 {
@@ -724,6 +817,9 @@ void UDrawSplineTool::AddSplinePoint(const FVector3d& HitLocation, const FVector
 		/*bUpdate =*/ false);
 	WorkingSpline->SetUpVectorAtSplinePoint(NumSplinePoints, UpVectorToUse, ESplineCoordinateSpace::World, 
 		/*bUpdate =*/ true);
+
+	RefitToSurface();
+
 }
 
 FVector3d UDrawSplineTool::GetUpVectorToUse(const FVector3d& HitLocation, const FVector3d& HitNormal, int32 NumSplinePointsBeforehand)
@@ -1106,7 +1202,7 @@ void UDrawSplineTool::Render(IToolsContextRenderAPI* RenderAPI)
 		PlaneMechanic->Render(RenderAPI);
 	}
 
-	if (WorkingSpline.IsValid())
+	if (!Settings->bFitToSurface && WorkingSpline.IsValid())
 	{
 		if (bDrawTangentForLastPoint)
 		{
@@ -1117,6 +1213,15 @@ void UDrawSplineTool::Render(IToolsContextRenderAPI* RenderAPI)
 		DrawSettings.ScaleVisualizationWidth = Settings->FrameVisualizationWidth;
 		SplineUtil::DrawSpline(*WorkingSpline, *RenderAPI, DrawSettings);
 	}
+
+	if (Settings->bFitToSurface && ProjectionSpline.IsValid())
+	{
+		SplineUtil::FDrawSplineSettings DrawSettings;
+		DrawSettings.ScaleVisualizationWidth = Settings->FrameVisualizationWidth;
+		DrawSettings.RegularColor = FColor::Cyan;
+		SplineUtil::DrawSpline(*ProjectionSpline, *RenderAPI, DrawSettings);
+	}
+
 }
 
 void UDrawSplineTool::OnPropertyModified(UObject* PropertySet, FProperty* Property)
@@ -1145,18 +1250,22 @@ bool UDrawSplineToolBuilder::CanBuildTool(const FToolBuilderState& SceneState) c
 	return true;
 }
 
-UInteractiveTool* UDrawSplineToolBuilder::BuildTool(const FToolBuilderState& SceneState) const
+USingleSelectionMeshEditingTool* UDrawSplineToolBuilder::CreateNewTool(const FToolBuilderState& SceneState) const
 {
-	UDrawSplineTool* NewTool = NewObject<UDrawSplineTool>(SceneState.ToolManager);
-	NewTool->SetWorld(SceneState.World);
+	return NewObject<UDrawSplineTool>(SceneState.ToolManager);
+}
 
-	// May be null
-	NewTool->SetSelectedActor(ToolBuilderUtil::FindFirstActor(SceneState, [](AActor*) { return true; }));
+void UDrawSplineToolBuilder::InitializeNewTool(USingleSelectionMeshEditingTool* Tool, const FToolBuilderState& SceneState) const
+{
+	USingleSelectionMeshEditingToolBuilder::InitializeNewTool(Tool, SceneState);
 
-	return NewTool;
+	if (UDrawSplineTool* SplineTool = Cast<UDrawSplineTool>(Tool))
+	{
+		// May be null
+		SplineTool->SetSelectedActor(ToolBuilderUtil::FindFirstActor(SceneState, [](AActor*) { return true; }));
+	}
 }
 
-
 void UDrawSplineTool::FSplineChange::Apply(UObject* Object)
 {
 	UDrawSplineTool* Tool = Cast<UDrawSplineTool>(Object);
diff --git a/Engine/Plugins/Experimental/MeshModelingToolsetExp/Source/MeshModelingToolsEditorOnlyExp/Public/DrawSplineTool.h b/Engine/Plugins/Experimental/MeshModelingToolsetExp/Source/MeshModelingToolsEditorOnlyExp/Public/DrawSplineTool.h
--- a/Engine/Plugins/Experimental/MeshModelingToolsetExp/Source/MeshModelingToolsEditorOnlyExp/Public/DrawSplineTool.h
+++ b/Engine/Plugins/Experimental/MeshModelingToolsetExp/Source/MeshModelingToolsEditorOnlyExp/Public/DrawSplineTool.h
@@ -7,9 +7,11 @@
 #include "BaseBehaviors/BehaviorTargetInterfaces.h"
 #include "InteractiveTool.h"
 #include "InteractiveToolBuilder.h"
+#include "BaseTools/SingleSelectionMeshEditingTool.h"
 #include "InteractiveToolChange.h"
 #include "ToolContextInterfaces.h" // FViewCameraState
 #include "TransactionUtil.h"
+#include "DynamicMesh/DynamicMeshAABBTree3.h"
 
 #include "DrawSplineTool.generated.h"
 
@@ -173,6 +175,18 @@ public:
 	UPROPERTY(EditAnywhere, Category = RaycastTargets, meta = (DisplayName = "Ground Planes"))
 	bool bHitGroundPlanes = true;
 
+	/** Whether to fit the spline to the selected surface. The spline will be subdivided and projected to reduce the distance between the spline and the surface */
+	UPROPERTY(EditAnywhere, Category = RaycastTargets, meta = (EditCondition = bMeshSelected, HideEditConditionToggle))
+	bool bFitToSurface = false;
+
+	/** The maximum number of times the spline will be subdivided in order to fit the surface */
+	UPROPERTY(EditAnywhere, Category = RaycastTargets, meta = (EditCondition = "bFitToSurface && bMeshSelected", ClampMin = 0, UIMax = 100))
+	int32 MaxNewPoints = 100;
+
+	/** The threshold for spline subdivision. New points are iteratively added where the distance from the spline to the surface is greater than this value */
+	UPROPERTY(EditAnywhere, Category = RaycastTargets, meta = (EditCondition = "bFitToSurface && bMeshSelected", ClampMin = 0, UIMax = 100))
+	double DistanceThreshold = 1.0;
+
 	/**
 	 * If modifying a blueprint actor, whether to run the construction script while dragging
 	 * or only at the end of a drag. Can be toggled off for expensive construction scripts.
@@ -180,10 +194,14 @@ public:
 	UPROPERTY(EditAnywhere, Category = Spline, AdvancedDisplay, meta = (
 		EditCondition = "OutputMode != EDrawSplineOutputMode::EmptyActor"))
 	bool bRerunConstructionScriptOnDrag = true;
+
+	// Hidden property set by the tool
+	UPROPERTY()
+	bool bMeshSelected = false;
 };
 
 UCLASS(MinimalAPI)
-class UDrawSplineTool : public UInteractiveTool
+class UDrawSplineTool : public USingleSelectionMeshEditingTool
 	, public IClickBehaviorTarget
 	, public IClickDragBehaviorTarget
 {
@@ -248,6 +266,8 @@ private:
 	bool bFreeDrawPlacedPreviewPoint = false;
 	int32 FreeDrawStrokeStartIndex = 0;
 
+	void RefitToSurface();
+
 	UE_API bool Raycast(const FRay& WorldRay, FVector3d& HitLocationOut, FVector3d& HitNormalOut, double& HitTOut);
 	UE_API void AddSplinePoint(const FVector3d& HitLocation, const FVector3d& UpVector);
 	UE_API FVector3d GetUpVectorToUse(const FVector3d& HitLocation, const FVector3d& HitNormal, int32 NumSplinePointsBeforehand);
@@ -270,6 +290,12 @@ private:
 
 	FViewCameraState CameraState;
 
+	TUniquePtr<UE::Geometry::FDynamicMeshAABBTree3> ProjectionTarget;
+	TUniquePtr<UE::Geometry::FDynamicMesh3> ProjectionTargetMesh;
+	FTransform MeshTargetTransform;
+	TWeakObjectPtr<USplineComponent> ProjectionSpline = nullptr;
+
+
 private:
 	UE::TransactionUtil::FLongTransactionTracker LongTransactions;
 
@@ -291,14 +317,16 @@ public:
 
 
 UCLASS(MinimalAPI, Transient)
-class UDrawSplineToolBuilder : public UInteractiveToolBuilder
+class UDrawSplineToolBuilder : public USingleSelectionMeshEditingToolBuilder
 {
 	GENERATED_BODY()
 
 public:
 	UE_API virtual bool CanBuildTool(const FToolBuilderState& SceneState) const override;
 
-	UE_API virtual UInteractiveTool* BuildTool(const FToolBuilderState& SceneState) const override;
+	UE_API virtual USingleSelectionMeshEditingTool* CreateNewTool(const FToolBuilderState& SceneState) const override;
+
+	UE_API virtual void InitializeNewTool(USingleSelectionMeshEditingTool* Tool, const FToolBuilderState& SceneState) const override;
 };
 
 #undef UE_API
diff --git a/Engine/Plugins/Runtime/MeshModelingToolset/Source/ModelingComponents/Private/SplineUtil.cpp b/Engine/Plugins/Runtime/MeshModelingToolset/Source/ModelingComponents/Private/SplineUtil.cpp
--- a/Engine/Plugins/Runtime/MeshModelingToolset/Source/ModelingComponents/Private/SplineUtil.cpp
+++ b/Engine/Plugins/Runtime/MeshModelingToolset/Source/ModelingComponents/Private/SplineUtil.cpp
@@ -182,7 +182,7 @@ void UE::Geometry::SplineUtil::ProjectSplineToSurface(FInterpCurveVector& Output
 	const UE::Geometry::FDynamicMeshAABBTree3& SurfaceAABBTree, 
 	const FTransform& SplineTransform, 
 	const FTransform& MeshTransform, 
-	double RelativeErrorThreshold,
+	double DistanceThreshold,
 	int32 MaxNewPoints)
 {
 	if (InputSpline.Points.Num() == 0)
@@ -232,7 +232,7 @@ void UE::Geometry::SplineUtil::ProjectSplineToSurface(FInterpCurveVector& Output
 	};
 
 
-	// Sample the full spline and find the furthest point from the surface
+	// Sample the full spline and finds the point with greatest error (distance to surface)
 	auto ApproximationError = [&SplineTransform, &MeshTransform, &SurfaceAABBTree, GetNearestPointOnSurface](const FInterpCurveVector& Spline, int32 NumSamplesPerSegment = 10) -> TPair<double, double>
 	{
 		NumSamplesPerSegment = FMath::Max(NumSamplesPerSegment, 1);
@@ -253,12 +253,12 @@ void UE::Geometry::SplineUtil::ProjectSplineToSurface(FInterpCurveVector& Output
 				const double T = MinT + DeltaT * (SampleIndex + 1);
 				const FVector SplinePoint = Spline.Eval(T, FVector(0));
 				const FVector3d SplineNearestPoint = GetNearestPointOnSurface(Spline, T);
-				const double Dist = FVector3d::Distance(SplineNearestPoint, SplinePoint);
+				const double Error = FVector3d::Distance(SplineNearestPoint, SplinePoint);
 
-				if (MaxError < Dist)
+				if (MaxError < Error)
 				{
 					MaxErrorT = T;
-					MaxError = Dist;
+					MaxError = Error;
 				}
 			}
 		}
@@ -267,7 +267,7 @@ void UE::Geometry::SplineUtil::ProjectSplineToSurface(FInterpCurveVector& Output
 	};
 
 
-	// Initialize OutputSpline using intput points
+	// Initialize OutputSpline using input points
 
 	const double MaxT = InputSpline.Points.Last().InVal;
 	double AvgSegmentLength = 0.0;
@@ -276,16 +276,6 @@ void UE::Geometry::SplineUtil::ProjectSplineToSurface(FInterpCurveVector& Output
 		const FInterpCurvePointVector& Point = InputSpline.Points[PointIndex];
 		const FInterpCurvePointVector ProjectedPoint = ProjectToSurface(InputSpline, Point.InVal, Point.InterpMode);
 		OutputSpline.Points.Add(ProjectedPoint);
-
-		if (PointIndex > 0)
-		{
-			AvgSegmentLength += FVector3d::Distance(InputSpline.Points[PointIndex].OutVal, InputSpline.Points[PointIndex - 1].OutVal);
-		}
-	}
-
-	if (InputSpline.Points.Num() > 1)
-	{
-		AvgSegmentLength /= (double)(InputSpline.Points.Num() - 1);
 	}
 
 	// Now iteratively insert points at the max-error sample along the spline
@@ -294,11 +284,9 @@ void UE::Geometry::SplineUtil::ProjectSplineToSurface(FInterpCurveVector& Output
 	double MaxErrorT = TAndError.Key;
 	double MaxError = TAndError.Value;
 
-	const double ErrorThreshold = 0.1 * AvgSegmentLength;
-
 	for (int32 InsertIteration = 0; InsertIteration < MaxNewPoints; ++InsertIteration)
 	{
-		if (MaxError < ErrorThreshold)
+		if (MaxError < DistanceThreshold)
 		{
 			break;
 		}
diff --git a/Engine/Plugins/Runtime/MeshModelingToolset/Source/ModelingComponents/Public/SplineUtil.h b/Engine/Plugins/Runtime/MeshModelingToolset/Source/ModelingComponents/Public/SplineUtil.h
--- a/Engine/Plugins/Runtime/MeshModelingToolset/Source/ModelingComponents/Public/SplineUtil.h
+++ b/Engine/Plugins/Runtime/MeshModelingToolset/Source/ModelingComponents/Public/SplineUtil.h
@@ -47,7 +47,7 @@ namespace SplineUtil {
 		const UE::Geometry::FDynamicMeshAABBTree3& SurfaceAABBTree, 
 		const FTransform& SplineTransform, 
 		const FTransform& MeshTransform, 
-		double RelativeErrorThreshold = 0.1,
+		double DistanceThreshold = 10.0,
 		int32 MaxNewPoints = 100);
 
 }
