Draw Spline Tool: don't check() if the selected target is non-null (we allow Target = nullptr for this tool)

diff --git a/Engine/Plugins/Experimental/MeshModelingToolsetExp/Source/MeshModelingToolsEditorOnlyExp/Private/DrawSplineTool.cpp b/Engine/Plugins/Experimental/MeshModelingToolsetExp/Source/MeshModelingToolsEditorOnlyExp/Private/DrawSplineTool.cpp
--- a/Engine/Plugins/Experimental/MeshModelingToolsetExp/Source/MeshModelingToolsEditorOnlyExp/Private/DrawSplineTool.cpp
+++ b/Engine/Plugins/Experimental/MeshModelingToolsetExp/Source/MeshModelingToolsEditorOnlyExp/Private/DrawSplineTool.cpp
@@ -1257,9 +1257,13 @@ USingleSelectionMeshEditingTool* UDrawSplineToolBuilder::CreateNewTool(const FTo
 
 void UDrawSplineToolBuilder::InitializeNewTool(USingleSelectionMeshEditingTool* Tool, const FToolBuilderState& SceneState) const
 {
-	USingleSelectionMeshEditingToolBuilder::InitializeNewTool(Tool, SceneState);
+	// Same as USingleSelectionMeshEditingToolBuilder::InitializeNewTool but allow for null Target
+	UToolTarget* const Target = SceneState.TargetManager->BuildFirstSelectedTargetable(SceneState, GetTargetRequirements());
+	Tool->SetTarget(Target);
+	Tool->SetWorld(SceneState.World);
 
-	if (UDrawSplineTool* SplineTool = Cast<UDrawSplineTool>(Tool))
+
+	if (UDrawSplineTool* const SplineTool = Cast<UDrawSplineTool>(Tool))
 	{
 		// May be null
 		SplineTool->SetSelectedActor(ToolBuilderUtil::FindFirstActor(SceneState, [](AActor*) { return true; }));
