ModelingTools: Make various ISM Editor tool improvements:
- Fix ctrl not removing from selection in single clicks
- Add extra selection manipulation methods: expand/contract last, select all, and expand all.
- Add ability to extract/duplicate items to a new actor.
- Fix selectable boxes not being hidden while dragging.

#Synced-CL 48611972

diff --git a/Engine/Plugins/Experimental/MeshModelingToolsetExp/Source/MeshModelingToolsEditorOnlyExp/Private/ISMEditorTool.cpp b/Engine/Plugins/Experimental/MeshModelingToolsetExp/Source/MeshModelingToolsEditorOnlyExp/Private/ISMEditorTool.cpp
--- a/Engine/Plugins/Experimental/MeshModelingToolsetExp/Source/MeshModelingToolsEditorOnlyExp/Private/ISMEditorTool.cpp
+++ b/Engine/Plugins/Experimental/MeshModelingToolsetExp/Source/MeshModelingToolsEditorOnlyExp/Private/ISMEditorTool.cpp
@@ -17,8 +17,9 @@
 #include "BaseGizmos/GizmoComponents.h"
 #include "BaseGizmos/TransformGizmoUtil.h"
 
-#include "Components/PrimitiveComponent.h"
+#include "Components/HierarchicalInstancedStaticMeshComponent.h"
 #include "Components/InstancedStaticMeshComponent.h"
+#include "Components/PrimitiveComponent.h"
 #include "Engine/World.h"
 
 #include "Editor/UnrealEdEngine.h"
@@ -26,10 +27,12 @@
 
 #include UE_INLINE_GENERATED_CPP_BY_NAME(ISMEditorTool)
 
-using namespace UE::Geometry;
-
 #define LOCTEXT_NAMESPACE "UISMEditorTool"
 
+namespace ISMEditorToolsLocals
+{
+	FText NoSelectionError = LOCTEXT("NoSelectionError", "Must select at least one Instance");
+}
 
 /*
  * ToolBuilder
@@ -150,6 +153,9 @@ void UISMEditorTool::Setup()
 	ToolActions->Initialize(this);
 	AddToolPropertySource(ToolActions);
 
+	ExtractActions = NewObject<UISMEditorToolExtractPropertySet>();
+	ExtractActions->Initialize(this);
+	AddToolPropertySource(ExtractActions);
 
 	TArray<AActor*> UniqueActors;
 	for (UInstancedStaticMeshComponent* Component : TargetComponents)
@@ -166,11 +172,13 @@ void UISMEditorTool::Setup()
 	PreviewGeometry = NewObject<UPreviewGeometry>(this);
 	PreviewGeometry->CreateInWorld(TargetComponents[0]->GetWorld(), FTransform::Identity);
 
-
 	SetToolDisplayName(LOCTEXT("ToolName", "Edit ISMs"));
 	GetToolManager()->DisplayMessage(
 		LOCTEXT("OnStartISMEditorTool", "Edit the selected ISMs"),
 		EToolMessageLevel::UserNotification);
+
+	// Add space in the warning area that appears when we clear warning messages.
+	GetToolManager()->DisplayMessage(FText(), EToolMessageLevel::UserError);
 }
 
 
@@ -208,6 +216,8 @@ void UISMEditorTool::OnUpdateModifierState(int ModifierID, bool bIsOn)
 
 void UISMEditorTool::OnMarqueeRectangleFinished(const FCameraRectangle& Rectangle, bool bCancelled)
 {
+	using namespace UE::Geometry;
+
 	if (!bCancelled)
 	{
 		TArray<const FActorChildMesh*> HitMeshes;
@@ -364,6 +374,9 @@ void UISMEditorTool::OnTick(float DeltaTime)
 {
 	if (PendingAction != EISMEditorToolActions::NoAction)
 	{
+		// Clear any messages from last action
+		GetToolManager()->DisplayMessage(FText(), EToolMessageLevel::UserError);
+
 		switch (PendingAction)
 		{
 		case EISMEditorToolActions::ClearSelection:
@@ -375,9 +388,24 @@ void UISMEditorTool::OnTick(float DeltaTime)
 		case EISMEditorToolActions::Duplicate:
 			OnDuplicateSelection();
 			break;
+		case EISMEditorToolActions::DuplicateToNew:
+			OnDuplicateToNew();
+			break;
+		case EISMEditorToolActions::ExtractToNew:
+			OnExtractToNew();
+			break;
 		case EISMEditorToolActions::Replace:
 			OnReplaceSelection();
 			break;
+		case EISMEditorToolActions::ExpandAll:
+			OnExpandAllSelection();
+			break;
+		case EISMEditorToolActions::ExpandLast:
+			OnExpandLastSelection();
+			break;
+		case EISMEditorToolActions::ContractLast:
+			OnContractLastSelection();
+			break;
 		}
 		PendingAction = EISMEditorToolActions::NoAction;
 	}
@@ -393,7 +421,8 @@ void UISMEditorTool::OnTick(float DeltaTime)
 
 	if (PreviewGeometry != nullptr)
 	{
-		PreviewGeometry->SetLineSetVisibility(TEXT("SelectableBoxes"), TransformProps->bShowSelectable);
+		PreviewGeometry->SetLineSetVisibility(TEXT("SelectableBoxes"), 
+			TransformProps->bShowSelectable && !(bInActiveDrag && TransformProps->bHideWhenDragging));
 	}
 }
 
@@ -432,6 +461,8 @@ void UISMEditorTool::RegisterActions(FInteractiveToolActionSet& ActionSet)
 
 FInputRayHit UISMEditorTool::IsHitByClick(const FInputDeviceRay& ClickPos)
 {
+	using namespace UE::Geometry;
+
 	if (CurrentSelection.Num() > 0)
 	{
 		return FInputRayHit(0);
@@ -450,6 +481,8 @@ FInputRayHit UISMEditorTool::IsHitByClick(const FInputDeviceRay& ClickPos)
 
 void UISMEditorTool::OnClicked(const FInputDeviceRay& ClickPos)
 {
+	using namespace UE::Geometry;
+
 	TArray<FSelectedInstance> NewSelection;
 
 	FMeshSceneRayHit HitResult;
@@ -472,6 +505,7 @@ void UISMEditorTool::OnClicked(const FInputDeviceRay& ClickPos)
 		}
 		else if (bCtrlModifier)
 		{
+			NewSelection = CurrentSelection;
 			int32 ExistingIndex = NewSelection.Find(NewItem);
 			if (ExistingIndex != INDEX_NONE)
 			{
@@ -724,7 +758,167 @@ void UISMEditorTool::OnDuplicateSelection()
 	GetToolManager()->EndUndoTransaction();
 }
 
+void UISMEditorTool::OnDuplicateToNew()
+{
+	DuplicateToNew(/*bDeleteFromCurrent =*/ false);
+}
+
+void UISMEditorTool::DuplicateToNew(bool bDeleteFromCurrent)
+{
+	if (CurrentSelection.Num() == 0)
+	{
+		GetToolManager()->DisplayMessage(ISMEditorToolsLocals::NoSelectionError, EToolMessageLevel::UserError);
+		return;
+	}
+
+	// Group our selected items by same static mesh and materials
+	TArray<TArray<int32>> SelectionIndexBuckets;
+	if (TargetComponents.Num() == 1)
+	{
+		SelectionIndexBuckets.Emplace();
+		for (int32 Index = 0; Index < CurrentSelection.Num(); ++Index)
+		{
+			SelectionIndexBuckets[0].Add(Index);
+		}
+	}
+	else
+	{
+		TMap<UInstancedStaticMeshComponent*, int32> ComponentToBucket;
+		for (int32 SelectionIndex = 0; SelectionIndex < CurrentSelection.Num(); ++SelectionIndex)
+		{
+			FSelectedInstance& Instance = CurrentSelection[SelectionIndex];
+			if (!ensure(Instance.Component))
+			{
+				continue;
+			}
+
+			int32 InstanceBucketIndex = INDEX_NONE;
+			if (int32* KnownBucket = ComponentToBucket.Find(Instance.Component))
+			{
+				InstanceBucketIndex = *KnownBucket;
+			}
+			else
+			{
+				TArray<UMaterialInterface*> Materials = Instance.Component->GetMaterials();
+				for (int32 BucketIndex = 0; BucketIndex < SelectionIndexBuckets.Num(); ++BucketIndex)
+				{
+					const FSelectedInstance& RepresentativeItem = CurrentSelection[SelectionIndexBuckets[BucketIndex][0]];
+					if (Instance.Component->GetStaticMesh() == RepresentativeItem.Component->GetStaticMesh()
+						&& RepresentativeItem.Component->GetMaterials() == Materials)
+					{
+						InstanceBucketIndex = BucketIndex;
+						ComponentToBucket.Add(Instance.Component, InstanceBucketIndex);
+						break;
+					}
+				}
+			}
+			if (InstanceBucketIndex == INDEX_NONE)
+			{
+				InstanceBucketIndex = SelectionIndexBuckets.Emplace();
+				ComponentToBucket.Add(Instance.Component, InstanceBucketIndex);
+			}
+			SelectionIndexBuckets[InstanceBucketIndex].Add(SelectionIndex);
+		}
+	}
+
+	if (!ensure(SelectionIndexBuckets.Num() > 0))
+	{
+		return;
+	}
+
+	GetToolManager()->BeginUndoTransaction(bDeleteFromCurrent ? LOCTEXT("ExtractToNew_Transaction", "Extract To New")
+		: LOCTEXT("DuplicateToNew_Transaction", "Duplicate To New"));
+	ON_SCOPE_EXIT
+	{
+		GetToolManager()->EndUndoTransaction();
+	};
+
+	// Create new actor
+	UWorld* TargetWorld = CurrentSelection[SelectionIndexBuckets[0][0]].Component->GetWorld();
+	FString UseBaseName = (ExtractActions->ActorName.Len() > 0) ? ExtractActions->ActorName : TEXT("Instances");
+	FActorSpawnParameters SpawnInfo;
+	SpawnInfo.ObjectFlags = RF_Transactional;
+	AActor* NewActor = TargetWorld ? TargetWorld->SpawnActor<AActor>(SpawnInfo) : nullptr;
+	if (!NewActor)
+	{
+		return;
+	}
 
+	FActorLabelUtilities::SetActorLabelUnique(NewActor, UseBaseName);
+
+	// Create a component for each bucket. First one as root and others attached.
+	bool bFirst = true;
+	for (int32 BucketIndex = 0; BucketIndex < SelectionIndexBuckets.Num(); ++BucketIndex)
+	{
+		TArray<int32>& SelectionIndices = SelectionIndexBuckets[BucketIndex];
+		const FSelectedInstance& RepresentativeItem = CurrentSelection[SelectionIndices[0]];
+		if (!RepresentativeItem.Component || !RepresentativeItem.Component->GetStaticMesh())
+		{
+			continue;
+		}
+
+		UE::Geometry::FAxisAlignedBox3d Bounds;
+		for (int32 Index : SelectionIndices)
+		{
+			Bounds.Contain(CurrentSelection[Index].WorldBounds);
+		}
+		FVector ComponentPosition = Bounds.Center();
+		
+		TSubclassOf<UInstancedStaticMeshComponent> ComponentClass = (ExtractActions->bCreateHISMs) ?
+			UHierarchicalInstancedStaticMeshComponent::StaticClass() : UInstancedStaticMeshComponent::StaticClass();
+		FString Suffix = (ExtractActions->bCreateHISMs) ? TEXT("_HISM") : TEXT("_ISM");
+
+		FName UseName = (UseBaseName.Len() > 0) ? FName(UseBaseName + Suffix) : FName();
+
+		UInstancedStaticMeshComponent* ISMComponent = NewObject<UInstancedStaticMeshComponent>(NewActor, ComponentClass,
+			MakeUniqueObjectName(NewActor, ComponentClass, UseName), RF_Transactional);
+		ISMComponent->bHasPerInstanceHitProxies = true;
+		ISMComponent->SetStaticMesh(RepresentativeItem.Component->GetStaticMesh());
+		TArray<UMaterialInterface*> Materials = RepresentativeItem.Component->GetMaterials();
+		for (int32 j = 0; j < Materials.Num(); ++j)
+		{
+			ISMComponent->SetMaterial(j, Materials[j]);
+		}
+		ISMComponent->OnComponentCreated();
+
+		if (bFirst)
+		{
+			NewActor->SetRootComponent(ISMComponent);
+			NewActor->SetActorTransform(FTransform(ComponentPosition));
+
+			bFirst = false;
+		}
+		else
+		{
+			ISMComponent->AttachToComponent(NewActor->GetRootComponent(), FAttachmentTransformRules::KeepRelativeTransform);
+			ISMComponent->SetWorldTransform(FTransform(ComponentPosition));
+		}
+		NewActor->AddInstanceComponent(ISMComponent);
+
+		for (int32 SelectionIndex : SelectionIndices)
+		{
+			FTransform InstanceTransform;
+			if (ensure(CurrentSelection.IsValidIndex(SelectionIndex))
+				&& ensure(CurrentSelection[SelectionIndex].Component->GetInstanceTransform(
+					CurrentSelection[SelectionIndex].Index, InstanceTransform, /*bWorldSpace=*/ true)))
+			{
+				ISMComponent->AddInstance(InstanceTransform, /*bTransformInWorldSpace=*/true);
+			}
+		}
+
+		ISMComponent->RegisterComponent();
+	}
+ 
+	if (bDeleteFromCurrent)
+	{
+		OnDeleteSelection();
+	}
+}
+
+void UISMEditorTool::OnExtractToNew()
+{
+	DuplicateToNew(/*bDeleteFromCurrent=*/ true);
+}
 
 void UISMEditorTool::OnReplaceSelection()
 {
@@ -807,9 +1001,144 @@ void UISMEditorTool::OnReplaceSelection()
 	ParentActor->PostEditChange();
 }
 
+void UISMEditorTool::OnContractLastSelection()
+{
+	if (CurrentSelection.IsEmpty())
+	{
+		return;
+	}
+	if (!ensure(CurrentSelection.Last().Component))
+	{
+		return;
+	}
+
+	// Find all components that have the same static mesh and materials as the last selected item
+	UInstancedStaticMeshComponent* SourceComponent = CurrentSelection.Last().Component;
+	UStaticMesh* StaticMesh = CurrentSelection.Last().Component->GetStaticMesh();
+	TArray<UMaterialInterface*> Materials = CurrentSelection.Last().Component->GetMaterials();
+
+	TArray<UInstancedStaticMeshComponent*> ComponentsToUnselect = TargetComponents.FilterByPredicate(
+		[SourceComponent, StaticMesh, &Materials](UInstancedStaticMeshComponent* Component)
+		{
+			return SourceComponent == Component
+				|| (Component && Component->GetStaticMesh() == StaticMesh && Component->GetMaterials() == Materials);
+		});
+
+	// Remove these components from the selection
+	TArray<FSelectedInstance> NewSelection = CurrentSelection;
+	NewSelection.RemoveAll([&ComponentsToUnselect](const FSelectedInstance& Instance)
+	{
+		return ComponentsToUnselect.Contains(Instance.Component);
+	});
+
+	if (NewSelection.Num() != CurrentSelection.Num())
+	{
+		UpdateSelectionInternal(NewSelection, true);
+	}
+}
+
+void UISMEditorTool::OnExpandLastSelection()
+{
+	if (CurrentSelection.IsEmpty())
+	{
+		return;
+	}
+	if (!ensure(CurrentSelection.Last().Component))
+	{
+		return;
+	}
+
+	// Find all components that have the same static mesh and materials as the last selected item
+	UInstancedStaticMeshComponent* SourceComponent = CurrentSelection.Last().Component;
+	UStaticMesh* StaticMesh = CurrentSelection.Last().Component->GetStaticMesh();
+	TArray<UMaterialInterface*> Materials = CurrentSelection.Last().Component->GetMaterials();
+
+	TArray<UInstancedStaticMeshComponent*> ComponentsToFullySelect = TargetComponents.FilterByPredicate(
+		[SourceComponent, StaticMesh, &Materials](UInstancedStaticMeshComponent* Component)
+		{
+			return SourceComponent == Component 
+				|| (Component && Component->GetStaticMesh() == StaticMesh && Component->GetMaterials() == Materials);
+		});
+
+	// Select all instances in these components
+	TArray<FSelectedInstance> NewSelection = CurrentSelection;
+	for (UInstancedStaticMeshComponent* Component : ComponentsToFullySelect)
+	{
+		int32 NumIndices = Component->GetInstanceCount();
+		for (int32 Index = 0; Index < NumIndices; ++Index)
+		{
+			if (Component->IsValidInstance(Index))
+			{
+				NewSelection.AddUnique(FSelectedInstance(Component, Index,
+					MeshScene->GetMeshBoundingBox(Component, Index)));
+			}
+		}
+	}
+	
+	if (NewSelection.Num() != CurrentSelection.Num())
+	{
+		UpdateSelectionInternal(NewSelection, true);
+	}
+}
+
+void UISMEditorTool::OnExpandAllSelection()
+{
+	if (CurrentSelection.IsEmpty())
+	{
+		return;
+	}
+
+	TArray<UInstancedStaticMeshComponent*> SourceComponents;
+	for (const FSelectedInstance& Instance : CurrentSelection)
+	{
+		if (ensure(Instance.Component))
+		{
+			SourceComponents.AddUnique(Instance.Component);
+		}
+	}
+	// Also add any components that have the same static mesh and material
+	TArray<UInstancedStaticMeshComponent*> ComponentsToFullySelect = SourceComponents;
+	for (UInstancedStaticMeshComponent* TargetComponent : TargetComponents)
+	{
+		if (!SourceComponents.Contains(TargetComponent))
+		{
+			for (UInstancedStaticMeshComponent* SourceComponent : SourceComponents)
+			{
+				if (TargetComponent->GetStaticMesh() == SourceComponent->GetStaticMesh()
+					&& TargetComponent->GetMaterials() == SourceComponent->GetMaterials())
+				{
+					ComponentsToFullySelect.Add(TargetComponent);
+					break;
+				}
+			}
+		}
+	}
+
+	// Select all instances in these components
+	TArray<FSelectedInstance> NewSelection = CurrentSelection;
+	for (UInstancedStaticMeshComponent* Component : ComponentsToFullySelect)
+	{
+		int32 NumIndices = Component->GetInstanceCount();
+		for (int32 Index = 0; Index < NumIndices; ++Index)
+		{
+			if (Component->IsValidInstance(Index))
+			{
+				NewSelection.AddUnique(FSelectedInstance(Component, Index,
+					MeshScene->GetMeshBoundingBox(Component, Index)));
+			}
+		}
+	}
+
+	if (NewSelection.Num() != CurrentSelection.Num())
+	{
+		UpdateSelectionInternal(NewSelection, true);
+	}
+}
 
 void UISMEditorTool::RebuildMeshScene()
 {
+	using namespace UE::Geometry;
+
 	MeshScene = MakeShared<FMeshSceneAdapter>();
 	TArray<UActorComponent*> TempComponents;
 	for (UInstancedStaticMeshComponent* ISMC : TargetComponents)
diff --git a/Engine/Plugins/Experimental/MeshModelingToolsetExp/Source/MeshModelingToolsEditorOnlyExp/Public/ISMEditorTool.h b/Engine/Plugins/Experimental/MeshModelingToolsetExp/Source/MeshModelingToolsEditorOnlyExp/Public/ISMEditorTool.h
--- a/Engine/Plugins/Experimental/MeshModelingToolsetExp/Source/MeshModelingToolsEditorOnlyExp/Public/ISMEditorTool.h
+++ b/Engine/Plugins/Experimental/MeshModelingToolsetExp/Source/MeshModelingToolsEditorOnlyExp/Public/ISMEditorTool.h
@@ -112,8 +112,13 @@ enum class EISMEditorToolActions
 {
 	NoAction,
 	ClearSelection,
+	ContractLast,
+	ExpandLast,
+	ExpandAll,
 	Delete,
 	Duplicate,
+	DuplicateToNew,
+	ExtractToNew,
 	Replace
 };
 
@@ -141,12 +146,30 @@ public:
 	UFUNCTION(CallInEditor, Category = SelectionEdits, meta = (DisplayPriority = 0))
 	void ClearSelection() { PostAction(EISMEditorToolActions::ClearSelection); }
 
-	/** Delete the selected instances */
+	/**
+	 * Take the last selected element and contract the selection to exclude any instances of that static mesh and materials.
+	 */
 	UFUNCTION(CallInEditor, Category = SelectionEdits, meta = (DisplayPriority = 1))
+	void ContractLast() { PostAction(EISMEditorToolActions::ContractLast); }
+
+	/**
+	 * Take the last selected element and expand the selection to include any other instances of the same static mesh and materials.
+	 */
+	UFUNCTION(CallInEditor, Category = SelectionEdits, meta = (DisplayPriority = 2))
+	void ExpandLast() { PostAction(EISMEditorToolActions::ExpandLast); }
+
+	/**
+	 * Expand the selection to include any other instances of the same static mesh(es) and materials as selected.
+	 */
+	UFUNCTION(CallInEditor, Category = SelectionEdits, meta = (DisplayPriority = 3))
+	void ExpandAll() { PostAction(EISMEditorToolActions::ExpandAll); }
+
+	/** Delete the selected instances */
+	UFUNCTION(CallInEditor, Category = SelectionEdits, meta = (DisplayPriority = 4))
 	void Delete() { PostAction(EISMEditorToolActions::Delete); }
 
 	/** Duplicate the selected instances */
-	UFUNCTION(CallInEditor, Category = SelectionEdits, meta = (DisplayPriority = 2))
+	UFUNCTION(CallInEditor, Category = SelectionEdits, meta = (DisplayPriority = 5))
 	void Duplicate() { PostAction(EISMEditorToolActions::Duplicate); }
 };
 
@@ -165,6 +188,28 @@ public:
 	void Replace() { PostAction(EISMEditorToolActions::Replace); }
 };
 
+UCLASS(MinimalAPI)
+class UISMEditorToolExtractPropertySet : public UISMEditorToolActionPropertySetBase
+{
+	GENERATED_BODY()
+public:
+	/** Name to use when duplicating or extracting to a new actor. */
+	UPROPERTY(EditAnywhere, Category = Extract)
+	FString ActorName;
+
+	/** When true, the components created when extracting to a new actor are HISM's instead of ISM's. */
+	UPROPERTY(EditAnywhere, Category = Extract)
+	bool bCreateHISMs = false;
+
+	/** Duplicate the selected instances into a new actor. */
+	UFUNCTION(CallInEditor, Category = Extract)
+	void DuplicateToNew() { PostAction(EISMEditorToolActions::DuplicateToNew); }
+
+	/** Duplicate the selected instances into a new actor and delete them from the current actor(s). */
+	UFUNCTION(CallInEditor, Category = Extract)
+	void ExtractToNew() { PostAction(EISMEditorToolActions::ExtractToNew); }
+};
+
 
 
 /**
@@ -232,6 +277,10 @@ protected:
 	UPROPERTY()
 	TObjectPtr<UPreviewGeometry> PreviewGeometry;
 
+private:
+	UPROPERTY()
+	TObjectPtr<UISMEditorToolExtractPropertySet> ExtractActions;
+protected:
 
 	TSharedPtr<UE::Geometry::FMeshSceneAdapter> MeshScene;
 	TArray<UE::Geometry::FAxisAlignedBox3d> AllMeshBoundingBoxes;
@@ -269,7 +318,16 @@ protected:
 	UE_API virtual void OnDeleteSelection();
 	UE_API virtual void OnDuplicateSelection();
 	UE_API virtual void OnReplaceSelection();
+private:
+	void OnContractLastSelection();
+	void OnExpandLastSelection();
+	void OnExpandAllSelection();
+	void OnDuplicateToNew();
+	void OnExtractToNew();
 
+	void DuplicateToNew(bool bDeleteFromCurrent);
+
+protected:
 	UE_API void UpdatePreviewGeometry();
 
 
