ModelingTools: Bind expand/contract selection actions to double click in ISM Editor tool, and make shift and ctrl behave the same way in clicking and marquee (shift alone adds, ctrl alone removes, and shift+ctrl toggles).

#Synced-CL 48746059

diff --git a/Engine/Plugins/Experimental/MeshModelingToolsetExp/Source/MeshModelingToolsEditorOnlyExp/Private/ISMEditorTool.cpp b/Engine/Plugins/Experimental/MeshModelingToolsetExp/Source/MeshModelingToolsEditorOnlyExp/Private/ISMEditorTool.cpp
--- a/Engine/Plugins/Experimental/MeshModelingToolsetExp/Source/MeshModelingToolsEditorOnlyExp/Private/ISMEditorTool.cpp
+++ b/Engine/Plugins/Experimental/MeshModelingToolsetExp/Source/MeshModelingToolsEditorOnlyExp/Private/ISMEditorTool.cpp
@@ -8,6 +8,7 @@
 #include "ToolSetupUtil.h"
 #include "DynamicMesh/DynamicMesh3.h"
 #include "BaseBehaviors/ClickDragBehavior.h"
+#include "BaseBehaviors/DoubleClickBehavior.h"
 #include "ToolSceneQueriesUtil.h"
 #include "ModelingToolTargetUtil.h"
 #include "Scene/MeshSceneAdapter.h"
@@ -141,6 +142,14 @@ void UISMEditorTool::Setup()
 	ClickOrDragBehavior->Modifiers.RegisterModifier(2, FInputDeviceState::IsCtrlKeyDown);
 	AddInputBehavior(ClickOrDragBehavior);
 
+	DoubleClickBehavior = NewObject<ULocalDoubleClickInputBehavior>();
+	DoubleClickBehavior->SetDefaultPriority(ClickOrDragBehavior->GetPriority().MakeHigher());
+	DoubleClickBehavior->Initialize();
+	DoubleClickBehavior->IsHitByClickFunc = [this](const FInputDeviceRay& ClickPos) { return IsHitByClick(ClickPos); };
+	DoubleClickBehavior->OnUpdateModifierStateFunc = [this](int ModifierID, bool bIsOn) { return OnUpdateModifierState(ModifierID, bIsOn); };
+	DoubleClickBehavior->OnClickedFunc = [this](const FInputDeviceRay& ClickPos) { return OnDoubleClicked(ClickPos); };
+	AddInputBehavior(DoubleClickBehavior);
+
 	TransformProps = NewObject<UISMEditorToolProperties>();
 	AddToolPropertySource(TransformProps);
 	TransformProps->RestoreProperties(this);
@@ -238,13 +247,11 @@ void UISMEditorTool::OnMarqueeRectangleFinished(const FCameraRectangle& Rectangl
 				return true; 
 			});
 
-
 		TArray<FSelectedInstance> NewSelection;
-		if (bShiftModifier || bCtrlModifier)
+		if (ModifyingSelection())
 		{
 			NewSelection = CurrentSelection;
 		}
-		bool bSubtract = bCtrlModifier;
 
 		for (const FActorChildMesh* HitMesh : HitMeshes)
 		{
@@ -254,16 +261,19 @@ void UISMEditorTool::OnMarqueeRectangleFinished(const FCameraRectangle& Rectangl
 				HitMesh->MeshSpatial->GetWorldBounds(
 					[&](const FVector3d& P) { return HitMesh->WorldTransform.TransformPosition(P); })
 			};
-
-			if (bSubtract)
+			
+			if (ShouldSubtractSelection())
 			{
-				int32 ExistingIndex = NewSelection.Find(NewItem);
-				if (ExistingIndex != INDEX_NONE)
+				NewSelection.RemoveSwap(NewItem);
+			}
+			else if (ShouldToggleSelection())
+			{
+				if (NewSelection.RemoveSwap(NewItem) == 0)
 				{
-					NewSelection.RemoveAt(ExistingIndex);
+					NewSelection.Add(NewItem);
 				}
 			}
-			else
+			else // if adding or setting
 			{
 				NewSelection.AddUnique(NewItem);
 			}
@@ -498,27 +508,24 @@ void UISMEditorTool::OnClicked(const FInputDeviceRay& ClickPos)
 				[&](const FVector3d& P) { return HitResult.LocalToWorld.TransformPosition(P); })
 		};
 
-		if (bShiftModifier)
+		if (ModifyingSelection())
 		{
 			NewSelection = CurrentSelection;
-			NewSelection.AddUnique(NewItem);
 		}
-		else if (bCtrlModifier)
+		if (ShouldSubtractSelection())
 		{
-			NewSelection = CurrentSelection;
-			int32 ExistingIndex = NewSelection.Find(NewItem);
-			if (ExistingIndex != INDEX_NONE)
-			{
-				NewSelection.RemoveAt(ExistingIndex);
-			}
-			else 
+			NewSelection.RemoveSwap(NewItem);
+		}
+		else if (ShouldToggleSelection())
+		{
+			if (NewSelection.RemoveSwap(NewItem) == 0)
 			{
 				NewSelection.Add(NewItem);
 			}
 		}
-		else
+		else // if adding or setting
 		{
-			NewSelection.Add(NewItem);
+			NewSelection.AddUnique(NewItem);
 		}
 	}
 
@@ -528,6 +535,60 @@ void UISMEditorTool::OnClicked(const FInputDeviceRay& ClickPos)
 	}
 }
 
+void UISMEditorTool::OnDoubleClicked(const FInputDeviceRay& ClickPos)
+{
+	// This is largely the same as OnClicked, except it does selection expansion/contraction.
+
+	using namespace UE::Geometry;
+
+	FText TransactionLabel = LOCTEXT("ExpandByDoubleClick", "Expand Selection By Double Click");
+	TArray<FSelectedInstance> NewSelection;
+
+	FMeshSceneRayHit HitResult;
+	if (MeshScene->FindNearestRayIntersection(ClickPos.WorldRay, HitResult))
+	{
+		FSelectedInstance NewItem{
+			(UInstancedStaticMeshComponent*)HitResult.HitComponent,
+			HitResult.HitComponentElementIndex,
+			HitResult.HitMeshSpatialWrapper->GetWorldBounds(
+				[&](const FVector3d& P) { return HitResult.LocalToWorld.TransformPosition(P); })
+		};
+
+		if (ModifyingSelection())
+		{
+			NewSelection = CurrentSelection;
+		}
+
+		if (ShouldSubtractSelection())
+		{
+			TransactionLabel = LOCTEXT("ContractByDoubleClick", "Contract Selection By Double Click");
+			ContractSelectionBy(NewItem, NewSelection);
+		}
+		else if (ShouldToggleSelection())
+		{
+			TransactionLabel = LOCTEXT("ToggleByDoubleClick", "Toggle Selection By Double Click");
+			ToggleSelectionBy(NewItem, NewSelection);
+
+			// A double click actually triggers the single click callback first, so we need to toggle
+			//  the original item again to undo that toggle. Of course, this would be slightly broken if
+			//  the mouse is moving such that the first click clicked something else, but that seems unlikely,
+			//  especially since this toggle action is probably uncommon.
+			if (NewSelection.Remove(NewItem) == 0)
+			{
+				NewSelection.AddUnique(NewItem);
+			}
+		}
+		else // if adding or setting
+		{
+			ExpandSelectionBy(NewItem, NewSelection);
+		}
+	}
+
+	if (NewSelection != CurrentSelection)
+	{
+		UpdateSelectionInternal(NewSelection, true, TransactionLabel);
+	}
+}
 
 void UISMEditorTool::OnSelectionUpdated()
 {
@@ -1012,10 +1073,20 @@ void UISMEditorTool::OnContractLastSelection()
 		return;
 	}
 
+	TArray<FSelectedInstance> NewSelection = CurrentSelection;
+	ContractSelectionBy(CurrentSelection.Last(), NewSelection);
+	if (NewSelection.Num() != CurrentSelection.Num())
+	{
+		UpdateSelectionInternal(NewSelection, true);
+	}
+}
+
+void UISMEditorTool::ContractSelectionBy(const FSelectedInstance& SelectedInstance, TArray<FSelectedInstance>& SelectionToModify)
+{
 	// Find all components that have the same static mesh and materials as the last selected item
-	UInstancedStaticMeshComponent* SourceComponent = CurrentSelection.Last().Component;
-	UStaticMesh* StaticMesh = CurrentSelection.Last().Component->GetStaticMesh();
-	TArray<UMaterialInterface*> Materials = CurrentSelection.Last().Component->GetMaterials();
+	UInstancedStaticMeshComponent* SourceComponent = SelectedInstance.Component;
+	UStaticMesh* StaticMesh = SelectedInstance.Component->GetStaticMesh();
+	TArray<UMaterialInterface*> Materials = SelectedInstance.Component->GetMaterials();
 
 	TArray<UInstancedStaticMeshComponent*> ComponentsToUnselect = TargetComponents.FilterByPredicate(
 		[SourceComponent, StaticMesh, &Materials](UInstancedStaticMeshComponent* Component)
@@ -1025,15 +1096,40 @@ void UISMEditorTool::OnContractLastSelection()
 		});
 
 	// Remove these components from the selection
-	TArray<FSelectedInstance> NewSelection = CurrentSelection;
-	NewSelection.RemoveAll([&ComponentsToUnselect](const FSelectedInstance& Instance)
+	SelectionToModify.RemoveAll([&ComponentsToUnselect](const FSelectedInstance& Instance)
 	{
 		return ComponentsToUnselect.Contains(Instance.Component);
 	});
+}
 
-	if (NewSelection.Num() != CurrentSelection.Num())
+void UISMEditorTool::ToggleSelectionBy(const FSelectedInstance& SelectedInstance, TArray<FSelectedInstance>& SelectionToModify)
+{
+	// Find all components that have the same static mesh and materials as the last selected item
+	UInstancedStaticMeshComponent* SourceComponent = SelectedInstance.Component;
+	UStaticMesh* StaticMesh = SelectedInstance.Component->GetStaticMesh();
+	TArray<UMaterialInterface*> Materials = SelectedInstance.Component->GetMaterials();
+
+	TArray<UInstancedStaticMeshComponent*> ComponentsToProcess = TargetComponents.FilterByPredicate(
+		[SourceComponent, StaticMesh, &Materials](UInstancedStaticMeshComponent* Component)
+		{
+			return SourceComponent == Component
+				|| (Component && Component->GetStaticMesh() == StaticMesh && Component->GetMaterials() == Materials);
+		});
+
+	for (UInstancedStaticMeshComponent* Component : ComponentsToProcess)
 	{
-		UpdateSelectionInternal(NewSelection, true);
+		int32 NumIndices = Component->GetInstanceCount();
+		for (int32 Index = 0; Index < NumIndices; ++Index)
+		{
+			if (Component->IsValidInstance(Index))
+			{
+				FSelectedInstance Instance(Component, Index, MeshScene->GetMeshBoundingBox(Component, Index));
+				if (SelectionToModify.RemoveSwap(Instance) == 0)
+				{
+					SelectionToModify.Add(Instance);
+				}
+			}
+		}
 	}
 }
 
@@ -1047,11 +1143,20 @@ void UISMEditorTool::OnExpandLastSelection()
 	{
 		return;
 	}
+	TArray<FSelectedInstance> NewSelection = CurrentSelection;
+	ExpandSelectionBy(CurrentSelection.Last(), NewSelection);
+	if (NewSelection.Num() != CurrentSelection.Num())
+	{
+		UpdateSelectionInternal(NewSelection, true);
+	}
+}
 
+void UISMEditorTool::ExpandSelectionBy(const FSelectedInstance& SelectedInstance, TArray<FSelectedInstance>& SelectionToModify)
+{
 	// Find all components that have the same static mesh and materials as the last selected item
-	UInstancedStaticMeshComponent* SourceComponent = CurrentSelection.Last().Component;
-	UStaticMesh* StaticMesh = CurrentSelection.Last().Component->GetStaticMesh();
-	TArray<UMaterialInterface*> Materials = CurrentSelection.Last().Component->GetMaterials();
+	UInstancedStaticMeshComponent* SourceComponent = SelectedInstance.Component;
+	UStaticMesh* StaticMesh = SelectedInstance.Component->GetStaticMesh();
+	TArray<UMaterialInterface*> Materials = SelectedInstance.Component->GetMaterials();
 
 	TArray<UInstancedStaticMeshComponent*> ComponentsToFullySelect = TargetComponents.FilterByPredicate(
 		[SourceComponent, StaticMesh, &Materials](UInstancedStaticMeshComponent* Component)
@@ -1061,7 +1166,6 @@ void UISMEditorTool::OnExpandLastSelection()
 		});
 
 	// Select all instances in these components
-	TArray<FSelectedInstance> NewSelection = CurrentSelection;
 	for (UInstancedStaticMeshComponent* Component : ComponentsToFullySelect)
 	{
 		int32 NumIndices = Component->GetInstanceCount();
@@ -1069,16 +1173,11 @@ void UISMEditorTool::OnExpandLastSelection()
 		{
 			if (Component->IsValidInstance(Index))
 			{
-				NewSelection.AddUnique(FSelectedInstance(Component, Index,
+				SelectionToModify.AddUnique(FSelectedInstance(Component, Index,
 					MeshScene->GetMeshBoundingBox(Component, Index)));
 			}
 		}
 	}
-	
-	if (NewSelection.Num() != CurrentSelection.Num())
-	{
-		UpdateSelectionInternal(NewSelection, true);
-	}
 }
 
 void UISMEditorTool::OnExpandAllSelection()
@@ -1207,13 +1306,18 @@ void UISMEditorTool::InternalNotifySceneModified(const TArray<UInstancedStaticMe
 
 
 void UISMEditorTool::UpdateSelectionInternal(const TArray<FSelectedInstance>& NewSelection, bool bEmitChange)
+{
+	UpdateSelectionInternal(NewSelection, bEmitChange, LOCTEXT("SelectionChange", "Selection Change"));
+}
+
+void UISMEditorTool::UpdateSelectionInternal(const TArray<FSelectedInstance>& NewSelection, bool bEmitChange, const FText& TransactionLabel)
 {
 	if (bEmitChange)
 	{
 		TUniquePtr<FISMEditorSelectionChange> SelectionChange = MakeUnique<FISMEditorSelectionChange>();
 		SelectionChange->OldSelection = CurrentSelection;
 		SelectionChange->NewSelection = NewSelection;
-		GetToolManager()->EmitObjectChange(this, MoveTemp(SelectionChange), LOCTEXT("SelectionChange", "Selection Change"));
+		GetToolManager()->EmitObjectChange(this, MoveTemp(SelectionChange), TransactionLabel);
 	}
 
 	CurrentSelection = NewSelection;
diff --git a/Engine/Plugins/Experimental/MeshModelingToolsetExp/Source/MeshModelingToolsEditorOnlyExp/Public/ISMEditorTool.h b/Engine/Plugins/Experimental/MeshModelingToolsetExp/Source/MeshModelingToolsEditorOnlyExp/Public/ISMEditorTool.h
--- a/Engine/Plugins/Experimental/MeshModelingToolsetExp/Source/MeshModelingToolsEditorOnlyExp/Public/ISMEditorTool.h
+++ b/Engine/Plugins/Experimental/MeshModelingToolsetExp/Source/MeshModelingToolsEditorOnlyExp/Public/ISMEditorTool.h
@@ -21,6 +21,7 @@
 class UBaseAxisTranslationGizmo;
 class UAxisAngleGizmo;
 class UDragAlignmentMechanic;
+class ULocalDoubleClickInputBehavior;
 class UTransformGizmo;
 class UTransformProxy;
 class UPreviewGeometry;
@@ -280,6 +281,9 @@ protected:
 private:
 	UPROPERTY()
 	TObjectPtr<UISMEditorToolExtractPropertySet> ExtractActions;
+
+	UPROPERTY()
+	TObjectPtr<ULocalDoubleClickInputBehavior> DoubleClickBehavior;
 protected:
 
 	TSharedPtr<UE::Geometry::FMeshSceneAdapter> MeshScene;
@@ -325,8 +329,20 @@ private:
 	void OnDuplicateToNew();
 	void OnExtractToNew();
 
+	void ExpandSelectionBy(const FSelectedInstance& SelectedInstance, TArray<FSelectedInstance>& SelectionToModify);
+	void ContractSelectionBy(const FSelectedInstance& SelectedInstance, TArray<FSelectedInstance>& SelectionToModify);
+	void ToggleSelectionBy(const FSelectedInstance& SelectedInstance, TArray<FSelectedInstance>& SelectionToModify);
+	void UpdateSelectionInternal(const TArray<FSelectedInstance>& NewSelection, bool bEmitChange, const FText& TransactionLabel);
+
 	void DuplicateToNew(bool bDeleteFromCurrent);
 
+	void OnDoubleClicked(const FInputDeviceRay& ClickPos);
+
+	// Helpers that use the modifiers to determine what to do with selection
+	bool ShouldAddToSelection() { return bShiftModifier && !bCtrlModifier; }
+	bool ShouldToggleSelection() { return bShiftModifier && bCtrlModifier; }
+	bool ShouldSubtractSelection() { return !bShiftModifier && bCtrlModifier; }
+	bool ModifyingSelection() { return bShiftModifier || bCtrlModifier; }
 protected:
 	UE_API void UpdatePreviewGeometry();
 
