disable the mesh sculpt layer toggle eye icon code for now since the underlying feature still doesn't exist

#Synced-CL 51148411

diff --git a/Engine/Plugins/Experimental/MeshModelingToolsetExp/Source/ModelingEditorUI/Private/ModelingWidgets/SMeshLayers.cpp b/Engine/Plugins/Experimental/MeshModelingToolsetExp/Source/ModelingEditorUI/Private/ModelingWidgets/SMeshLayers.cpp
--- a/Engine/Plugins/Experimental/MeshModelingToolsetExp/Source/ModelingEditorUI/Private/ModelingWidgets/SMeshLayers.cpp
+++ b/Engine/Plugins/Experimental/MeshModelingToolsetExp/Source/ModelingEditorUI/Private/ModelingWidgets/SMeshLayers.cpp
@@ -302,6 +302,9 @@ void SMeshLayerItem::Construct(const FArguments& InArgs, const TSharedRef<STable
 	}
 
 	// enable Checkbox - TODO
+	// Note the code for this checkbox is intentionally left here but disabled, as the underlying feature doesn't exist yet
+	// but this has the desired UI look for when it does
+#if 0
 	BoxArgs
 	+ SHorizontalBox::Slot()
 	.AutoWidth()
@@ -347,6 +350,7 @@ void SMeshLayerItem::Construct(const FArguments& InArgs, const TSharedRef<STable
 			]
 		]
 	];
+#endif
 
 	// display Name
 	BoxArgs
