[ModelingToolExp] Fixed value at brush displayed at the wrong position

#Synced-CL 51489786

diff --git a/Engine/Plugins/Experimental/MeshModelingToolsetExp/Source/MeshModelingToolsEditorOnlyExp/Private/MeshVertexAttributePaintToolBase.cpp b/Engine/Plugins/Experimental/MeshModelingToolsetExp/Source/MeshModelingToolsEditorOnlyExp/Private/MeshVertexAttributePaintToolBase.cpp
--- a/Engine/Plugins/Experimental/MeshModelingToolsetExp/Source/MeshModelingToolsEditorOnlyExp/Private/MeshVertexAttributePaintToolBase.cpp
+++ b/Engine/Plugins/Experimental/MeshModelingToolsetExp/Source/MeshModelingToolsEditorOnlyExp/Private/MeshVertexAttributePaintToolBase.cpp
@@ -885,7 +885,7 @@ void UMeshVertexAttributePaintToolBase::Shutdown(EToolShutdownType ShutdownType)
 		MeshSelector = nullptr;
 	}
 
-	// we need to call this directly now because we use the preview mesh instead of the built-in component
+	// Using our custom commit method
 	if (ShutdownType == EToolShutdownType::Accept)
 	{
 		ToolDynamicMesh->EditMesh([this](FDynamicMesh3& Mesh)
@@ -894,6 +894,7 @@ void UMeshVertexAttributePaintToolBase::Shutdown(EToolShutdownType ShutdownType)
 		});
 	}
 
+	// Clearing the preview mesh so base class commit is No-Op
 	if (PreviewMesh)
 	{
 		PreviewMesh->SetVisible(false);
@@ -2228,15 +2229,13 @@ void UMeshVertexAttributePaintToolBase::DrawHUD(FCanvas* Canvas, IToolsContextRe
 		{
 			// Display the value under the brush
 			const float BrushRadius = BrushIndicator->BrushRadius;
-			FVector Offset{ 0 };
 			const FVector HoverPosition = HoverStamp.WorldFrame.ToFTransform().GetTranslation();
 
-			const FVector BrushPos = HoverPosition + Offset;
-
-			FVector4 ScreenBrushPos = RenderAPI->GetSceneView()->WorldToScreen(BrushPos);
-			FVector2D ScreenBrushPos2D = FVector2D(ScreenBrushPos.X, ScreenBrushPos.Y);
-
-			RenderAPI->GetSceneView()->ScreenToPixel(ScreenBrushPos, ScreenBrushPos2D);
+			FVector2D ScreenBrushPos2D;
+			static FVector2D Offset(0,75.0);
+			RenderAPI->GetSceneView()->WorldToPixel(HoverPosition, ScreenBrushPos2D);
+			ScreenBrushPos2D += Offset;
+			ScreenBrushPos2D /= Canvas->GetDPIScale();
 
 			// Render value under at brush even in mesh mode
 			const FText BaseText = bPendingPickWeight
@@ -2325,9 +2324,9 @@ void UMeshVertexAttributePaintToolBase::DrawHUD(FCanvas* Canvas, IToolsContextRe
 							const FText ValueText = FText::AsNumber(RemappedValue, &NumberFormattingOptions);
 
 							const FVector3d VtxPosition = Mesh->GetVertex(VertexID);
-							const FVector4 ScreenBrushPos = RenderAPI->GetSceneView()->WorldToScreen(VtxPosition);
-							FVector2D ScreenBrushPos2D = FVector2D(ScreenBrushPos.X, ScreenBrushPos.Y);
-							RenderAPI->GetSceneView()->ScreenToPixel(ScreenBrushPos, ScreenBrushPos2D);
+							FVector2D ScreenBrushPos2D;
+							RenderAPI->GetSceneView()->WorldToPixel(VtxPosition, ScreenBrushPos2D);
+							ScreenBrushPos2D /= Canvas->GetDPIScale();
 
 							FCanvasTextItem TextItem(ScreenBrushPos2D, ValueText, GEngine->GetLargeFont(), FLinearColor::White);
 							TextItem.EnableShadow(FLinearColor::Black);
