[ModelingToolExp] Turn off poly select mechanic before shutting it down so that its input behavior (still registered with input router during shutdown) cannot trigger callbacks on it after it has been shutdown.

#Synced-CL 51489786

diff --git a/Engine/Plugins/Experimental/MeshModelingToolsetExp/Source/MeshModelingToolsEditorOnlyExp/Private/MeshVertexAttributePaintToolBase.cpp b/Engine/Plugins/Experimental/MeshModelingToolsetExp/Source/MeshModelingToolsEditorOnlyExp/Private/MeshVertexAttributePaintToolBase.cpp
--- a/Engine/Plugins/Experimental/MeshModelingToolsetExp/Source/MeshModelingToolsEditorOnlyExp/Private/MeshVertexAttributePaintToolBase.cpp
+++ b/Engine/Plugins/Experimental/MeshModelingToolsetExp/Source/MeshModelingToolsEditorOnlyExp/Private/MeshVertexAttributePaintToolBase.cpp
@@ -881,6 +881,10 @@ void UMeshVertexAttributePaintToolBase::Shutdown(EToolShutdownType ShutdownType)
 
 	if (MeshSelector)
 	{
+		// Avoid behaviors owned by poly select mechanics from 
+		// triggering a crash after modal windows closes, if drived class has one (eg. MorphTargetMaskTool)
+		// we might want to unregister tool behaviors before calling tool shutdown ?
+		MeshSelector->SetIsEnabled(false);
 		MeshSelector->Shutdown();
 		MeshSelector = nullptr;
 	}
