[ModelingTools Exp] updated comment about a workaround

#Synced-CL 51489786

diff --git a/Engine/Plugins/Experimental/MeshModelingToolsetExp/Source/MeshModelingToolsEditorOnlyExp/Private/MeshVertexAttributePaintToolBase.cpp b/Engine/Plugins/Experimental/MeshModelingToolsetExp/Source/MeshModelingToolsEditorOnlyExp/Private/MeshVertexAttributePaintToolBase.cpp
--- a/Engine/Plugins/Experimental/MeshModelingToolsetExp/Source/MeshModelingToolsEditorOnlyExp/Private/MeshVertexAttributePaintToolBase.cpp
+++ b/Engine/Plugins/Experimental/MeshModelingToolsetExp/Source/MeshModelingToolsEditorOnlyExp/Private/MeshVertexAttributePaintToolBase.cpp
@@ -881,9 +881,10 @@ void UMeshVertexAttributePaintToolBase::Shutdown(EToolShutdownType ShutdownType)
 
 	if (MeshSelector)
 	{
-		// Avoid behaviors owned by poly select mechanics from 
-		// triggering a crash after modal windows closes, if drived class has one (eg. MorphTargetMaskTool)
-		// we might want to unregister tool behaviors before calling tool shutdown ?
+		// Disable the PolygonSelectionMechanic to make sure any of its behaviors that are still registered with input router
+		// cannot trigger meaningful logic on the Mechanic after it has been shutdown. The input router typically doesn't
+		// have an opportunity to tick the behaviors beyond this point, but it can tick one more time if a modal window is
+		// spawned by the child class during CommitToolMesh below (Morph Target Mask Tool, for example). 
 		MeshSelector->SetIsEnabled(false);
 		MeshSelector->Shutdown();
 		MeshSelector = nullptr;
