[ITF] Unregister input behavior before shutdown is called to avoid potential crashes if a modal window is created during tool shutdown. When modal window closes, it can tick input router during tool shutdown, which causes input behaviors that haven't been unregistered to trigger callbacks on their owner that could have been shutdown as part of tool shutdown.

#Synced-CL 51489786

diff --git a/Engine/Plugins/Experimental/MeshModelingToolsetExp/Source/MeshModelingToolsEditorOnlyExp/Private/MeshVertexAttributePaintToolBase.cpp b/Engine/Plugins/Experimental/MeshModelingToolsetExp/Source/MeshModelingToolsEditorOnlyExp/Private/MeshVertexAttributePaintToolBase.cpp
--- a/Engine/Plugins/Experimental/MeshModelingToolsetExp/Source/MeshModelingToolsEditorOnlyExp/Private/MeshVertexAttributePaintToolBase.cpp
+++ b/Engine/Plugins/Experimental/MeshModelingToolsetExp/Source/MeshModelingToolsEditorOnlyExp/Private/MeshVertexAttributePaintToolBase.cpp
@@ -881,11 +881,6 @@ void UMeshVertexAttributePaintToolBase::Shutdown(EToolShutdownType ShutdownType)
 
 	if (MeshSelector)
 	{
-		// Disable the PolygonSelectionMechanic to make sure any of its behaviors that are still registered with input router
-		// cannot trigger meaningful logic on the Mechanic after it has been shutdown. The input router typically doesn't
-		// have an opportunity to tick the behaviors beyond this point, but it can tick one more time if a modal window is
-		// spawned by the child class during CommitToolMesh below (Morph Target Mask Tool, for example). 
-		MeshSelector->SetIsEnabled(false);
 		MeshSelector->Shutdown();
 		MeshSelector = nullptr;
 	}
diff --git a/Engine/Source/Runtime/InteractiveToolsFramework/Private/InteractiveToolManager.cpp b/Engine/Source/Runtime/InteractiveToolsFramework/Private/InteractiveToolManager.cpp
--- a/Engine/Source/Runtime/InteractiveToolsFramework/Private/InteractiveToolManager.cpp
+++ b/Engine/Source/Runtime/InteractiveToolsFramework/Private/InteractiveToolManager.cpp
@@ -304,11 +304,14 @@ void UInteractiveToolManager::DeactivateToolInternal(EToolSide Side, EToolShutdo
 
 		bInToolShutdown = true;
 
+		// Stop any input behaviors first so that we don't have any pending termination callbacks post-Shutdown.
+		// Note that the deregister is done because a tool that generates a modal window during Shutdown
+		// could otherwise re-trigger input recapture, since modals continue pumping input messages.
 		InputRouter->ForceTerminateSource(ActiveLeftTool);
-
+		InputRouter->DeregisterSource(ActiveLeftTool);
+		
 		ActiveLeftTool->Shutdown(ShutdownType);
 
-		InputRouter->DeregisterSource(ActiveLeftTool);
 
 		UInteractiveTool* DoneTool = ActiveLeftTool;
 		ActiveLeftTool = nullptr;
