Fixed crash on deleting last mesh layer

#Synced-CL 52209628

diff --git a/Engine/Plugins/Experimental/MeshModelingToolsetExp/Source/ModelingEditorUI/Private/ModelingWidgets/SMeshLayers.cpp b/Engine/Plugins/Experimental/MeshModelingToolsetExp/Source/ModelingEditorUI/Private/ModelingWidgets/SMeshLayers.cpp
--- a/Engine/Plugins/Experimental/MeshModelingToolsetExp/Source/ModelingEditorUI/Private/ModelingWidgets/SMeshLayers.cpp
+++ b/Engine/Plugins/Experimental/MeshModelingToolsetExp/Source/ModelingEditorUI/Private/ModelingWidgets/SMeshLayers.cpp
@@ -83,22 +83,8 @@ void SMeshLayersList::DeleteMeshLayer(TSharedPtr<FMeshLayerElement> LayerToDelet
 		return;
 	}
 
-	// remove layer from this List of layers - not technically necessary if base layer is removed above and List is being refreshed below,
-	// since the refresh will generate a new Layers list based on the updated layers (without the deleted one) from the Controller
-	const int32 StackIndexToDelete = LayerToDelete->GetIndexInStack();
-	Layers.RemoveAt(StackIndexToDelete);
-
-	// correct indices of other layers
-	for (TSharedPtr<FMeshLayerElement> Layer : Layers)
-	{
-		const int32 StackIndex = Layer->GetIndexInStack();
-		// for all layers after the removed one, adjust their IndexInStack
-		if (StackIndex != INDEX_NONE && StackIndex > StackIndexToDelete)
-		{
-			Layer->SetIndexInStack(StackIndex - 1);
-		}
-	}
-	RequestListRefresh();
+	// Note: The controller will indirectly trigger a RefreshAndRestore() call via RemoveMeshLayer
+	// so we don't need to manually update the list post-removal.
 }
 
 void SMeshLayersList::RefreshAndRestore()
