Vertex attribute paint tool : improve mirror settings behavior
- Mirror activation and visibility are now saved across sessions
- In mesh mode mirror settings are modifiable regardless of the status of the selection
- Selection visibility no longer changes when switching in mesh mode

#Synced-CL 52663992

diff --git a/Engine/Plugins/Experimental/MeshModelingToolsetExp/Source/MeshModelingToolsEditorOnlyExp/Private/MeshVertexAttributePaintToolBase.cpp b/Engine/Plugins/Experimental/MeshModelingToolsetExp/Source/MeshModelingToolsEditorOnlyExp/Private/MeshVertexAttributePaintToolBase.cpp
--- a/Engine/Plugins/Experimental/MeshModelingToolsetExp/Source/MeshModelingToolsEditorOnlyExp/Private/MeshVertexAttributePaintToolBase.cpp
+++ b/Engine/Plugins/Experimental/MeshModelingToolsetExp/Source/MeshModelingToolsEditorOnlyExp/Private/MeshVertexAttributePaintToolBase.cpp
@@ -673,8 +673,6 @@ void UMeshVertexAttributePaintToolBase::InitializeMirrorIndicator()
 	MirrorPlaneMesh->SetShadowsEnabled(false);
 	MirrorPlaneMesh->SetOverrideWireframeRenderMaterial(BrushIndicatorMaterial);
 
-	SetMirrorPlaneWidgetVisible(ToolProperties->EditingMode == EMeshVertexAttributePaintToolEditMode::Brush);
-	
 	UpdateMirrorPlaneMesh();
 }
 
@@ -936,26 +934,31 @@ bool UMeshVertexAttributePaintToolBase::IsValuePickerEnabled() const
 
 void UMeshVertexAttributePaintToolBase::EnableBrushMirroring(bool bEnable)
 {
-	bBrushMirroringEnabled = bEnable;
-	SetFocusInViewport();
+	if (ToolProperties)
+	{
+		ToolProperties->MirrorProperties.bEnableBrushMirroring = bEnable;
+		SetFocusInViewport();
+	}
 }
 
 bool UMeshVertexAttributePaintToolBase::IsBrushMirroringEnabled() const
 {
-	return bBrushMirroringEnabled;
+	return ToolProperties ? ToolProperties->MirrorProperties.bEnableBrushMirroring: false;
 }
 
 void UMeshVertexAttributePaintToolBase::SetMirrorPlaneWidgetVisible(bool bVisible)
 {
-	bMirrorPlaneWidgetVisible = bVisible;
-	UpdateMirrorPlaneMesh();
-
-	SetFocusInViewport();
+	if (ToolProperties)
+	{
+		ToolProperties->MirrorProperties.bMirrorPlaneWidgetVisible = bVisible;
+		UpdateMirrorPlaneMesh();
+		SetFocusInViewport();
+	}
 }
 
 bool UMeshVertexAttributePaintToolBase::IsMirrorPlaneWidgetVisible()
 {
-	return bMirrorPlaneWidgetVisible;
+	return ToolProperties ? ToolProperties->MirrorProperties.bMirrorPlaneWidgetVisible : true;
 }
 
 bool UMeshVertexAttributePaintToolBase::IsBrushMirroringInObjectSpace() const
@@ -1318,10 +1321,6 @@ void UMeshVertexAttributePaintToolBase::SetEditingMode(EMeshVertexAttributePaint
 	if (ToolProperties)
 	{
 		ToolProperties->EditingMode = Editmode;
-		if (Editmode == EMeshVertexAttributePaintToolEditMode::Mesh)
-		{
-			SetMirrorPlaneWidgetVisible(false);
-		}
 		SetFocusInViewport();
 	}
 }
@@ -2442,7 +2441,7 @@ void UMeshVertexAttributePaintToolBase::OnTick(float DeltaTime)
 			LastStamp.TimeStamp = FDateTime::Now();
 
 			// now mirror stamp 
-			if (bBrushMirroringEnabled)
+			if (IsBrushMirroringEnabled())
 			{
 				bWeightsModified = false;
 				// update sculpt ROI
diff --git a/Engine/Plugins/Experimental/MeshModelingToolsetExp/Source/MeshModelingToolsEditorOnlyExp/Public/MeshVertexAttributePaintToolBase.h b/Engine/Plugins/Experimental/MeshModelingToolsetExp/Source/MeshModelingToolsEditorOnlyExp/Public/MeshVertexAttributePaintToolBase.h
--- a/Engine/Plugins/Experimental/MeshModelingToolsetExp/Source/MeshModelingToolsEditorOnlyExp/Public/MeshVertexAttributePaintToolBase.h
+++ b/Engine/Plugins/Experimental/MeshModelingToolsetExp/Source/MeshModelingToolsEditorOnlyExp/Public/MeshVertexAttributePaintToolBase.h
@@ -296,6 +296,14 @@ struct FMeshVertexAttributePaintToolMirrorProperties
 {
 	GENERATED_BODY()
 
+	/* whether the mirror is enabled in brush mode */
+	UPROPERTY()
+	bool bEnableBrushMirroring = false;
+
+	/* whether the mirror plane is visible or not */
+	UPROPERTY()
+	bool bMirrorPlaneWidgetVisible = true;
+
 	/* if true, the mirror planes will be centered on the bounds of the object */
 	UPROPERTY()
 	bool bObjectSpace = true;
@@ -631,8 +639,6 @@ protected:
 
 	bool bPendingPickWeight = false;
 	bool bBrushIsHoverMesh = false;
-	bool bBrushMirroringEnabled = false;
-	bool bMirrorPlaneWidgetVisible = true;
 	bool bIsMirrorHiddenBecauseOfBrushStroke = false;
 
 	TArray<int32> ROITriangleBuffer;
